UAF: Moved/renamed AnimNext and AnimNextAnimGraph plugins

Moved to UAF subdirectory
Renamed AnimNext -> UAF
Renamed AnimNextAnimGraph -> UAFAnimGraph

Updated all references & provided redirectors

diff --git a/Engine/Plugins/Experimental/AnimNextChooser/AnimNextChooser.uplugin b/Engine/Plugins/Experimental/AnimNextChooser/AnimNextChooser.uplugin
--- a/Engine/Plugins/Experimental/AnimNextChooser/AnimNextChooser.uplugin
+++ b/Engine/Plugins/Experimental/AnimNextChooser/AnimNextChooser.uplugin
@@ -18,11 +18,11 @@
 	"Plugins":
 	[
 		{
-			"Name": "AnimNext",
+			"Name": "UAF",
 			"Enabled": true
 		},
 		{
-			"Name": "AnimNextAnimGraph",
+			"Name": "UAFAnimGraph",
 			"Enabled": true
 		},
 		{
diff --git a/Engine/Plugins/Experimental/AnimNextChooser/Source/AnimNextChooser/AnimNextChooser.Build.cs b/Engine/Plugins/Experimental/AnimNextChooser/Source/AnimNextChooser/AnimNextChooser.Build.cs
--- a/Engine/Plugins/Experimental/AnimNextChooser/Source/AnimNextChooser/AnimNextChooser.Build.cs
+++ b/Engine/Plugins/Experimental/AnimNextChooser/Source/AnimNextChooser/AnimNextChooser.Build.cs
@@ -14,9 +14,9 @@ namespace UnrealBuildTool.Rules
 					"RigVM",
 					"ControlRig",
 					"Engine",
-					"AnimNext",
+					"UAF",
 					"Chooser",
-					"AnimNextAnimGraph",
+					"UAFAnimGraph",
 				}
 			);
 		}
diff --git a/Engine/Plugins/Experimental/AnimNextControlRig/AnimNextControlRig.uplugin b/Engine/Plugins/Experimental/AnimNextControlRig/AnimNextControlRig.uplugin
--- a/Engine/Plugins/Experimental/AnimNextControlRig/AnimNextControlRig.uplugin
+++ b/Engine/Plugins/Experimental/AnimNextControlRig/AnimNextControlRig.uplugin
@@ -22,11 +22,11 @@
 			"Enabled": true
 		},
 		{
-			"Name": "AnimNext",
+			"Name": "UAF",
 			"Enabled": true
 		},
 		{
-			"Name": "AnimNextAnimGraph",
+			"Name": "UAFAnimGraph",
 			"Enabled": true
 		},
 	],
diff --git a/Engine/Plugins/Experimental/AnimNextControlRig/Source/AnimNextControlRig/AnimNextControlRig.Build.cs b/Engine/Plugins/Experimental/AnimNextControlRig/Source/AnimNextControlRig/AnimNextControlRig.Build.cs
--- a/Engine/Plugins/Experimental/AnimNextControlRig/Source/AnimNextControlRig/AnimNextControlRig.Build.cs
+++ b/Engine/Plugins/Experimental/AnimNextControlRig/Source/AnimNextControlRig/AnimNextControlRig.Build.cs
@@ -14,8 +14,8 @@ namespace UnrealBuildTool.Rules
 					"Core",
 					"ControlRig",
 					"RigVM",
-					"AnimNext",
-					"AnimNextAnimGraph",
+					"UAF",
+					"UAFAnimGraph",
 				}
 			);
 
@@ -24,11 +24,11 @@ namespace UnrealBuildTool.Rules
 				PrivateDependencyModuleNames.AddRange(
 				   new string[]
 				   {
-					   "ControlRigDeveloper",
+						"ControlRigDeveloper",
 						"WorkspaceEditor",
 						"SlateCore",
 						"Slate",
-						"AnimNextUncookedOnly",
+						"UAFUncookedOnly",
 						"RigVMDeveloper",
 						"AnimationCore",
 				   });
diff --git a/Engine/Plugins/Experimental/AnimNextControlRig/Source/AnimNextControlRigEditor/AnimNextControlRigEditor.Build.cs b/Engine/Plugins/Experimental/AnimNextControlRig/Source/AnimNextControlRigEditor/AnimNextControlRigEditor.Build.cs
--- a/Engine/Plugins/Experimental/AnimNextControlRig/Source/AnimNextControlRigEditor/AnimNextControlRigEditor.Build.cs
+++ b/Engine/Plugins/Experimental/AnimNextControlRig/Source/AnimNextControlRigEditor/AnimNextControlRigEditor.Build.cs
@@ -15,8 +15,8 @@ namespace UnrealBuildTool.Rules
 					"ControlRig",
 					"ControlRigDeveloper",
 					"RigVM",
-					"AnimNext",
-					"AnimNextAnimGraph",
+					"UAF",
+					"UAFAnimGraph",
 				}
 			);
 
@@ -29,10 +29,10 @@ namespace UnrealBuildTool.Rules
 					"SlateCore",
 					"Slate",
 					"RigVMEditor",
-					"AnimNextEditor",
+					"UAFEditor",
 					"AnimNextControlRig",
-					"AnimNextUncookedOnly",
-					"AnimNextAnimGraphUncookedOnly",
+					"UAFUncookedOnly",
+					"UAFAnimGraphUncookedOnly",
 				}
 			);
 		}
diff --git a/Engine/Plugins/Experimental/AnimNextControlRig/Source/AnimNextControlRigEditor/Private/AnimNextControlRigEditorModule.cpp b/Engine/Plugins/Experimental/AnimNextControlRig/Source/AnimNextControlRigEditor/Private/AnimNextControlRigEditorModule.cpp
--- a/Engine/Plugins/Experimental/AnimNextControlRig/Source/AnimNextControlRigEditor/Private/AnimNextControlRigEditorModule.cpp
+++ b/Engine/Plugins/Experimental/AnimNextControlRig/Source/AnimNextControlRigEditor/Private/AnimNextControlRigEditorModule.cpp
@@ -47,7 +47,7 @@ void FAnimNextControlRigEditorModule::StartupModule()
 		"ControlRigTraitSharedData",
 		FOnGetPropertyTypeCustomizationInstance::CreateLambda([] { return MakeShared<UE::UAF::Editor::FControlRigTraitSharedDataCustomization>(); }));
 
-	UE::UAF::Editor::IAnimNextEditorModule& AnimNextEditorModule = FModuleManager::LoadModuleChecked<UE::UAF::Editor::IAnimNextEditorModule>("AnimNextEditor");
+	UE::UAF::Editor::IAnimNextEditorModule& AnimNextEditorModule = FModuleManager::LoadModuleChecked<UE::UAF::Editor::IAnimNextEditorModule>("UAFEditor");
 
 	OnNodeDblClickDelegateHandle = AnimNextEditorModule.RegisterNodeDblClickHandler(UE::UAF::Editor::IAnimNextEditorModule::FNodeDblClickNotificationDelegate::CreateLambda([](const Workspace::FWorkspaceEditorContext& InContext, const UEdGraphNode* InNode)
 		{
@@ -89,9 +89,9 @@ void FAnimNextControlRigEditorModule::StartupModule()
 
 void FAnimNextControlRigEditorModule::ShutdownModule()
 {
-	if (FModuleManager::Get().IsModuleLoaded("AnimNextEditor"))
+	if (FModuleManager::Get().IsModuleLoaded("UAFEditor"))
 	{
-		UE::UAF::Editor::IAnimNextEditorModule& AnimNextEditorModule = FModuleManager::LoadModuleChecked<UE::UAF::Editor::IAnimNextEditorModule>("AnimNextEditor");
+		UE::UAF::Editor::IAnimNextEditorModule& AnimNextEditorModule = FModuleManager::LoadModuleChecked<UE::UAF::Editor::IAnimNextEditorModule>("UAFEditor");
 
 		AnimNextEditorModule.UnregisterNodeDblClickHandler(OnNodeDblClickDelegateHandle);
 	}
diff --git a/Engine/Plugins/Experimental/AnimNextPoseSearch/Source/AnimNextPoseSearch/AnimNextPoseSearch.Build.cs b/Engine/Plugins/Experimental/AnimNextPoseSearch/Source/AnimNextPoseSearch/AnimNextPoseSearch.Build.cs
--- a/Engine/Plugins/Experimental/AnimNextPoseSearch/Source/AnimNextPoseSearch/AnimNextPoseSearch.Build.cs
+++ b/Engine/Plugins/Experimental/AnimNextPoseSearch/Source/AnimNextPoseSearch/AnimNextPoseSearch.Build.cs
@@ -9,8 +9,8 @@ namespace UnrealBuildTool.Rules
 			PublicDependencyModuleNames.AddRange(
 				new string[]
 				{
-					"AnimNext",
-					"AnimNextAnimGraph",
+					"UAF",
+					"UAFAnimGraph",
 					"Core",
 					"CoreUObject",
 					"Engine",
diff --git a/Engine/Plugins/Experimental/AnimNextStateTree/Source/AnimNextStateTree/AnimNextStateTree.Build.cs b/Engine/Plugins/Experimental/AnimNextStateTree/Source/AnimNextStateTree/AnimNextStateTree.Build.cs
--- a/Engine/Plugins/Experimental/AnimNextStateTree/Source/AnimNextStateTree/AnimNextStateTree.Build.cs
+++ b/Engine/Plugins/Experimental/AnimNextStateTree/Source/AnimNextStateTree/AnimNextStateTree.Build.cs
@@ -9,7 +9,7 @@ namespace UnrealBuildTool.Rules
 			PrivateDependencyModuleNames.AddRange(
 				new string[]
 				{
-					"AnimNextAnimGraph",
+					"UAFAnimGraph",
 					"RigVM"
 				});
 
@@ -20,7 +20,7 @@ namespace UnrealBuildTool.Rules
 					"CoreUObject",
 					"StateTreeModule",
 					"Engine",
-					"AnimNext",
+					"UAF",
 				}
 			);
 
@@ -29,10 +29,10 @@ namespace UnrealBuildTool.Rules
 				PrivateDependencyModuleNames.AddRange(
 					new string[]
 					{
-						"AnimNextAnimGraph",
-						"AnimNextAnimGraphUncookedOnly",
-						"AnimNextEditor",
-						"AnimNextUncookedOnly",
+						"UAFAnimGraph",
+						"UAFAnimGraphUncookedOnly",
+						"UAFEditor",
+						"UAFUncookedOnly",
 						"RigVMDeveloper",
 						"StateTreeEditorModule",
 						"UnrealEd"
diff --git a/Engine/Plugins/Experimental/AnimNextStateTree/Source/AnimNextStateTree/Internal/Tasks/AnimNextStateTreeGraphInstanceTask.h b/Engine/Plugins/Experimental/AnimNextStateTree/Source/AnimNextStateTree/Internal/Tasks/AnimNextStateTreeGraphInstanceTask.h
--- a/Engine/Plugins/Experimental/AnimNextStateTree/Source/AnimNextStateTree/Internal/Tasks/AnimNextStateTreeGraphInstanceTask.h
+++ b/Engine/Plugins/Experimental/AnimNextStateTree/Source/AnimNextStateTree/Internal/Tasks/AnimNextStateTreeGraphInstanceTask.h
@@ -25,7 +25,7 @@ struct ANIMNEXTSTATETREE_API FAnimNextGraphInstanceTaskInstanceData
 	// Upgrade on load all structs/property bags (bags need source asset to reference)
 
 	// The asset to instantiate
-	UPROPERTY(EditAnywhere, Category = Animation, meta=(GetAllowedClasses="/Script/AnimNextAnimGraph.AnimNextAnimGraphSettings:GetAllowedAssetClasses"))
+	UPROPERTY(EditAnywhere, Category = Animation, meta=(GetAllowedClasses="/Script/UAFAnimGraph.AnimNextAnimGraphSettings:GetAllowedAssetClasses"))
 	TObjectPtr<UObject> Asset = nullptr;
 
 #if WITH_EDITORONLY_DATA
diff --git a/Engine/Plugins/Experimental/AnimNextStateTree/Source/AnimNextStateTreeEditor/AnimNextStateTreeEditor.Build.cs b/Engine/Plugins/Experimental/AnimNextStateTree/Source/AnimNextStateTreeEditor/AnimNextStateTreeEditor.Build.cs
--- a/Engine/Plugins/Experimental/AnimNextStateTree/Source/AnimNextStateTreeEditor/AnimNextStateTreeEditor.Build.cs
+++ b/Engine/Plugins/Experimental/AnimNextStateTree/Source/AnimNextStateTreeEditor/AnimNextStateTreeEditor.Build.cs
@@ -14,10 +14,10 @@ namespace UnrealBuildTool.Rules
 				    "EditorInteractiveToolsFramework",
 				    "InteractiveToolsFramework",
 				    "EditorFramework",
-				    "AnimNext",
+				    "UAF",
 				    "UnrealEd",
 				    "AssetDefinition",
-				    "AnimNextAnimGraph",
+					"UAFAnimGraph",
 				    "RigVM",
 				    "StructUtilsEditor",
 				    "PropertyEditor",
@@ -34,13 +34,13 @@ namespace UnrealBuildTool.Rules
 					"StateTreeModule",
 					"PropertyBindingUtils",
 					"AnimNextStateTreeUncookedOnly",
-					"AnimNextEditor",
+					"UAFEditor",
 					"SlateCore",
 					"EditorSubsystem",
 					"Engine",
 					"Slate",
 					"RigVMDeveloper",
-					"AnimNextUncookedOnly"
+					"UAFUncookedOnly"
 				}
 			);
 		}
diff --git a/Engine/Plugins/Experimental/AnimNextStateTree/Source/AnimNextStateTreeEditor/Private/AnimNextStateTreeEditorModule.cpp b/Engine/Plugins/Experimental/AnimNextStateTree/Source/AnimNextStateTreeEditor/Private/AnimNextStateTreeEditorModule.cpp
--- a/Engine/Plugins/Experimental/AnimNextStateTree/Source/AnimNextStateTreeEditor/Private/AnimNextStateTreeEditorModule.cpp
+++ b/Engine/Plugins/Experimental/AnimNextStateTree/Source/AnimNextStateTreeEditor/Private/AnimNextStateTreeEditorModule.cpp
@@ -35,7 +35,7 @@ namespace UE::UAF::StateTree
 void FAnimNextStateTreeEditorModule::StartupModule()
 {
 	// Register StateTree as supported asset in AnimNext workspaces
-	Editor::IAnimNextEditorModule& AnimNextEditorModule = FModuleManager::Get().LoadModuleChecked<Editor::IAnimNextEditorModule>("AnimNextEditor");	
+	Editor::IAnimNextEditorModule& AnimNextEditorModule = FModuleManager::Get().LoadModuleChecked<Editor::IAnimNextEditorModule>("UAFEditor");	
 	AnimNextEditorModule.AddWorkspaceSupportedAssetClass(UAnimNextStateTree::StaticClass()->GetClassPathName());
 	AnimNextEditorModule.RegisterAssetCompilationHandler(UAnimNextStateTree::StaticClass()->GetClassPathName(), Editor::FAssetCompilationHandlerFactoryDelegate::CreateLambda([](UObject* InAsset)
 	{
@@ -349,7 +349,7 @@ void FAnimNextStateTreeEditorModule::ShutdownModule()
 				WorkspaceEditorModule->UnregisterWorkspaceItemDetails(Workspace::FOutlinerItemDetailsId(FAnimNextStateTreeStateOutlinerData::StaticStruct()->GetFName()));
 		}
 	
-		if( Editor::IAnimNextEditorModule* AnimNextEditorModule = FModuleManager::Get().GetModulePtr<Editor::IAnimNextEditorModule>("AnimNextEditor"))
+		if( Editor::IAnimNextEditorModule* AnimNextEditorModule = FModuleManager::Get().GetModulePtr<Editor::IAnimNextEditorModule>("UAFEditor"))
 		{
 			AnimNextEditorModule->RemoveWorkspaceSupportedAssetClass(UStateTree::StaticClass()->GetClassPathName());
 			AnimNextEditorModule->UnregisterAssetCompilationHandler(UAnimNextStateTree::StaticClass()->GetClassPathName());
diff --git a/Engine/Plugins/Experimental/AnimNextStateTree/Source/AnimNextStateTreeUncookedOnly/AnimNextStateTreeUncookedOnly.Build.cs b/Engine/Plugins/Experimental/AnimNextStateTree/Source/AnimNextStateTreeUncookedOnly/AnimNextStateTreeUncookedOnly.Build.cs
--- a/Engine/Plugins/Experimental/AnimNextStateTree/Source/AnimNextStateTreeUncookedOnly/AnimNextStateTreeUncookedOnly.Build.cs
+++ b/Engine/Plugins/Experimental/AnimNextStateTree/Source/AnimNextStateTreeUncookedOnly/AnimNextStateTreeUncookedOnly.Build.cs
@@ -23,11 +23,11 @@ namespace UnrealBuildTool.Rules
 					"Engine",
 					"WorkspaceEditor",
 					"SlateCore",
-					"AnimNext",
-					"AnimNextUncookedOnly",
+					"UAF",
+					"UAFUncookedOnly",
 					"AnimNextStateTree",
-					"AnimNextAnimGraph",
-					"AnimNextAnimGraphUncookedOnly"
+					"UAFAnimGraph",
+					"UAFAnimGraphUncookedOnly"
 				}
 			);
 		}
diff --git a/Engine/Plugins/Experimental/AnimNextWarping/Source/AnimNextWarping/AnimNextWarping.Build.cs b/Engine/Plugins/Experimental/AnimNextWarping/Source/AnimNextWarping/AnimNextWarping.Build.cs
--- a/Engine/Plugins/Experimental/AnimNextWarping/Source/AnimNextWarping/AnimNextWarping.Build.cs
+++ b/Engine/Plugins/Experimental/AnimNextWarping/Source/AnimNextWarping/AnimNextWarping.Build.cs
@@ -14,8 +14,8 @@ namespace UnrealBuildTool.Rules
 					"AnimationCore",
 					"RigVM",
 					"Engine",
-					"AnimNext",
-					"AnimNextAnimGraph",
+					"UAF",
+					"UAFAnimGraph",
 				}
 			);
 		}
diff --git a/Engine/Plugins/Experimental/MoverAnimNext/MoverAnimNext.uplugin b/Engine/Plugins/Experimental/MoverAnimNext/MoverAnimNext.uplugin
--- a/Engine/Plugins/Experimental/MoverAnimNext/MoverAnimNext.uplugin
+++ b/Engine/Plugins/Experimental/MoverAnimNext/MoverAnimNext.uplugin
@@ -2,8 +2,8 @@
 	"FileVersion": 1,
 	"Version": 1,
 	"VersionName": "1.0",
-	"FriendlyName": "Mover AnimNext",
-	"Description": "This plugin adds AnimNext support for the Mover plugin.",
+	"FriendlyName": "Mover UAF",
+	"Description": "This plugin adds UAF support for the Mover plugin.",
 	"Category": "Gameplay",
 	"CreatedBy": "Epic Games, Inc.",
 	"CreatedByURL": "https://epicgames.com",
@@ -27,7 +27,7 @@
 			"Enabled": true
 		},
 		{
-			"Name": "AnimNext",
+			"Name": "UAF",
 			"Enabled": true
 		}
 	]
diff --git a/Engine/Plugins/Experimental/MoverAnimNext/Source/MoverAnimNext/MoverAnimNext.Build.cs b/Engine/Plugins/Experimental/MoverAnimNext/Source/MoverAnimNext/MoverAnimNext.Build.cs
--- a/Engine/Plugins/Experimental/MoverAnimNext/Source/MoverAnimNext/MoverAnimNext.Build.cs
+++ b/Engine/Plugins/Experimental/MoverAnimNext/Source/MoverAnimNext/MoverAnimNext.Build.cs
@@ -12,7 +12,7 @@ public class MoverAnimNext : ModuleRules
 			{
 				"Core",
 				"Mover",
-				"AnimNext",
+				"UAF",
 				"PoseSearch"
 			}
 			);
diff --git a/Engine/Plugins/Experimental/MovieSceneAnimMixer/Source/MovieSceneAnimMixer/MovieSceneAnimMixer.Build.cs b/Engine/Plugins/Experimental/MovieSceneAnimMixer/Source/MovieSceneAnimMixer/MovieSceneAnimMixer.Build.cs
--- a/Engine/Plugins/Experimental/MovieSceneAnimMixer/Source/MovieSceneAnimMixer/MovieSceneAnimMixer.Build.cs
+++ b/Engine/Plugins/Experimental/MovieSceneAnimMixer/Source/MovieSceneAnimMixer/MovieSceneAnimMixer.Build.cs
@@ -15,8 +15,8 @@ namespace UnrealBuildTool.Rules
 					"MovieScene",
 					"MovieSceneTracks",
 					"AnimGraphRuntime",
-					"AnimNext",
-					"AnimNextAnimGraph"
+					"UAF",
+					"UAFAnimGraph"
 				}
 			);
 
diff --git a/Engine/Plugins/Experimental/MovieScenePoseSearchTracks/Source/MovieScenePoseSearchTracks/MovieScenePoseSearchTracks.Build.cs b/Engine/Plugins/Experimental/MovieScenePoseSearchTracks/Source/MovieScenePoseSearchTracks/MovieScenePoseSearchTracks.Build.cs
--- a/Engine/Plugins/Experimental/MovieScenePoseSearchTracks/Source/MovieScenePoseSearchTracks/MovieScenePoseSearchTracks.Build.cs
+++ b/Engine/Plugins/Experimental/MovieScenePoseSearchTracks/Source/MovieScenePoseSearchTracks/MovieScenePoseSearchTracks.Build.cs
@@ -15,9 +15,9 @@ namespace UnrealBuildTool.Rules
 					"MovieScene",
 					"MovieSceneTracks",
 					"MovieSceneAnimMixer",
-					"AnimNext",
+					"UAF",
 					"AnimNextPoseSearch",
-					"AnimNextAnimGraph",
+					"UAFAnimGraph",
 					"AnimGraphRuntime",
 					"PoseSearch"
 				}
diff --git a/Engine/Plugins/Experimental/RigLogicAnimNext/RigLogicAnimNext.uplugin b/Engine/Plugins/Experimental/RigLogicAnimNext/RigLogicAnimNext.uplugin
--- a/Engine/Plugins/Experimental/RigLogicAnimNext/RigLogicAnimNext.uplugin
+++ b/Engine/Plugins/Experimental/RigLogicAnimNext/RigLogicAnimNext.uplugin
@@ -2,8 +2,8 @@
 	"FileVersion": 3,
 	"Version": 1,
 	"VersionName": "1.0",
-	"FriendlyName": "RigLogic for AnimNext",
-	"Description": "RigLogic for AnimNext",
+	"FriendlyName": "RigLogic for UAF",
+	"Description": "RigLogic for UAF",
 	"Category": "Animation",
 	"CreatedBy": "Epic Games Inc.",
 	"CreatedByURL": "",
@@ -29,11 +29,11 @@
 			"Enabled": true
 		},
 		{
-			"Name": "AnimNext",
+			"Name": "UAF",
 			"Enabled": true
 		},
 		{
-			"Name": "AnimNextAnimGraph",
+			"Name": "UAFAnimGraph",
 			"Enabled": true
 		}
 	]
diff --git a/Engine/Plugins/Experimental/UAF/UAF/Config/BaseUAF.ini b/Engine/Plugins/Experimental/UAF/UAF/Config/BaseUAF.ini
new file mode 100644
--- /dev/null
+++ b/Engine/Plugins/Experimental/UAF/UAF/Config/BaseUAF.ini
@@ -0,0 +1,287 @@
+[CoreRedirects]
+; Traits rename
++StructRedirects=(OldName="/Script/AnimNext.AnimNextDecoratorHandle", NewName="/Script/UAFAnimGraph.AnimNextTraitHandle")
++StructRedirects=(OldName="/Script/AnimNext.RigUnit_AnimNextDecoratorStack", NewName="/Script/UAFAnimGraph.RigUnit_AnimNextTraitStack")
+
+; Parameter block/graph merge
++ClassRedirects=(OldName="/Script/AnimNext.AnimNextParameterBlock", NewName="/Script/UAF.AnimNextModule")
++ClassRedirects=(OldName="/Script/AnimNextUncookedOnly.AnimNextParameterBlock_EditorData", NewName="/Script/UAFUncookedOnly.AnimNextModule_EditorData")
++ClassRedirects=(OldName="/Script/AnimNextUncookedOnly.AnimNextParameterBlock_EdGraph", NewName="/Script/UAFUncookedOnly.AnimNextEdGraph")
++ClassRedirects=(OldName="/Script/AnimNextUncookedOnly.AnimNextParameterBlock_EdGraphNode", NewName="/Script/UAFUncookedOnly.AnimNextEdGraphNode")
++ClassRedirects=(OldName="/Script/AnimNextUncookedOnly.AnimNextParameterBlock_EdGraphSchema", NewName="/Script/UAFUncookedOnly.AnimNextEdGraphSchema")
++ClassRedirects=(OldName="/Script/AnimNextUncookedOnly.AnimNextParameterBlockParameter", NewName="/Script/UAFUncookedOnly.AnimNextVariableEntry")
++ClassRedirects=(OldName="/Script/AnimNextUncookedOnly.AnimNextParameterBlockGraph", NewName="/Script/UAFUncookedOnly.AnimNextEventGraphEntry")
++ClassRedirects=(OldName="/Script/AnimNextUncookedOnly.AnimNextGraphEntry", NewName="/Script/UAFAnimGraphUncookedOnly.AnimNextAnimationGraphEntry")
++ClassRedirects=(OldName="/Script/AnimNextUncookedOnly.AnimNextGraph_TraitSchema", NewName="/Script/UAFAnimGraphUncookedOnly.AnimNextAnimationGraphSchema")
++PropertyRedirects=(OldName="/Script/AnimNext.AnimNextScheduleEntry_ParamScope.ParameterBlocks", NewName="Parameters")
++PropertyRedirects=(OldName="/Script/AnimNext.AnimNextScheduleParamScopeEntryTask.ParameterBlocks", NewName="Parameters")
+
+; Workspace plugin split
++ClassRedirects=(OldName="/Script/AnimNextEditor.AnimNextWorkspace", NewName="/Script/WorkspaceEditor.Workspace")
+
+; Scene graph prototype compatibility
++ClassRedirects=(OldName="/Script/AnimNext.AnimNextMeshComponent", NewName="/Script/Engine.SkeletalMeshComponent")
+
+; AnimNextGraph -> AnimNextModule
++ClassRedirects=(OldName="/Script/AnimNext.AnimNextGraph", NewName="/Script/UAF.AnimNextModule")
++ClassRedirects=(OldName="/Script/AnimNextUncookedOnly.AnimNextGraph_Controller", NewName="/Script/UAFUncookedOnly.AnimNextControllerBase")
++ClassRedirects=(OldName="/Script/AnimNextUncookedOnly.AnimNextGraph_EditorData", NewName="/Script/UAFUncookedOnly.AnimNextModule_EditorData")
++ClassRedirects=(OldName="/Script/AnimNextUncookedOnly.AnimNextGraph_AnimationGraphSchema", NewName="/Script/UAFAnimGraphUncookedOnly.AnimNextAnimationGraphSchema")
++ClassRedirects=(OldName="/Script/AnimNextUncookedOnly.AnimNextGraph_EdGraph", NewName="/Script/UAFUncookedOnly.AnimNextEdGraph")
++ClassRedirects=(OldName="/Script/AnimNextUncookedOnly.AnimNextGraph_EdGraphNode", NewName="/Script/UAFUncookedOnly.AnimNextEdGraphNode")
++ClassRedirects=(OldName="/Script/AnimNextUncookedOnly.AnimNextGraph_EdGraphSchema", NewName="/Script/UAFUncookedOnly.AnimNextEdGraphSchema")
++ClassRedirects=(OldName="/Script/AnimNextUncookedOnly.AnimNextParameterBlock_ParamSchema", NewName="/Script/UAFUncookedOnly.AnimNextEventGraphSchema")
++ClassRedirects=(OldName="/Script/AnimNextUncookedOnly.AnimNextGraph_EventGraphSchema", NewName="/Script/UAFUncookedOnly.AnimNextEventGraphSchema")
++ClassRedirects=(OldName="/Script/AnimNextUncookedOnly.AnimNextGraph_UnitNode", NewName="/Script/UAFUncookedOnly.AnimNextUnitNode")
++ClassRedirects=(OldName="/Script/AnimNextUncookedOnly.AnimNextRigVMWorkspaceAssetUserData", NewName="/Script/UAFUncookedOnly.AnimNextAssetWorkspaceAssetUserData")
+
++PropertyRedirects=(OldName="/Script/AnimNext.AnimNextGraphInstance.Graph", NewName="Module")
++PropertyRedirects=(OldName="/Script/AnimNext.AnimNextScheduleEntry_AnimNextGraph.Graph", NewName="Module")
++PropertyRedirects=(OldName="/Script/AnimNext.AnimNextSubGraphHostTraitSharedData.SubGraph", NewName="Module")
++PropertyRedirects=(OldName="/Script/AnimNext.AnimNextScheduleGraphTask.Graph", NewName="Module")
++PropertyRedirects=(OldName="/Script/AnimNext.AnimNextScheduleGraphTask.DynamicGraph", NewName="DynamicModule")
++PropertyRedirects=(OldName="/Script/AnimNextUncookedOnly.AnimNextParameterProviderAssetRegistryExports.Parameters", NewName="Variables")
+
+; Splitting up modules & adding variables to graphs
++ClassRedirects=(OldName="/Script/AnimNextUncookedOnly.AnimNextModuleWorkspaceAssetUserData", NewName="/Script/UAFUncookedOnly.AnimNextAssetWorkspaceAssetUserData")
++ClassRedirects=(OldName="/Script/AnimNextUncookedOnly.AnimNextModule_Controller", NewName="/Script/UAFUncookedOnly.AnimNextControllerBase")
++ClassRedirects=(OldName="/Script/AnimNextUncookedOnly.AnimNextModule_AnimationGraph", NewName="/Script/UAFAnimGraphUncookedOnly.AnimNextAnimationGraphEntry")
++ClassRedirects=(OldName="/Script/AnimNextUncookedOnly.AnimNextModule_Parameter", NewName="/Script/UAFUncookedOnly.AnimNextVariableEntry")
++ClassRedirects=(OldName="/Script/AnimNextUncookedOnly.AnimNextModule_EventGraph", NewName="/Script/UAFUncookedOnly.AnimNextEventGraphEntry")
++ClassRedirects=(OldName="/Script/AnimNextUncookedOnly.AnimNextGraph_AnimationGraph", NewName="/Script/UAFAnimGraphUncookedOnly.AnimNextAnimationGraphEntry")
++ClassRedirects=(OldName="/Script/AnimNextUncookedOnly.AnimNextGraph_Parameter", NewName="/Script/UAFUncookedOnly.AnimNextVariableEntry")
++ClassRedirects=(OldName="/Script/AnimNextUncookedOnly.AnimNextGraph_EventGraph", NewName="/Script/UAFUncookedOnly.AnimNextEventGraphEntry")
+
++StructRedirects=(OldName="/Script/AnimNextUncookedOnly.AnimNextGraphAssetOutlinerData", NewName="/Script/UAFUncookedOnly.AnimNextModuleOutlinerData")
++StructRedirects=(OldName="/Script/AnimNextUncookedOnly.AnimNextParameterOutlinerData", NewName="/Script/UAFUncookedOnly.AnimNextVariableOutlinerData")
++StructRedirects=(OldName="/Script/AnimNextUncookedOnly.AnimNextParameterProviderAssetRegistryExports", NewName="/Script/UAFUncookedOnly.AnimNextAssetRegistryExports")
+
++EnumRedirects=(OldName="EAnimNextModuleInitMethod",NewName="/Script/UAF.EAnimNextModuleInitMethod")
++EnumRedirects=(OldName="EAnimNextParameterFlags",NewName="/Script/UAFUncookedOnly.EAnimNextExportedVariableFlags")
+
++PropertyRedirects=(OldName="/Script/AnimNext.RigVMTrait_AnimNextPublicVariables.Asset", NewName="InternalAsset")
++PropertyRedirects=(OldName="/Script/AnimNext.RigVMTrait_AnimNextPublicVariables.VariableNames", NewName="InternalVariableNames")
+
+; Move all animation graph functionality to AnimNextAnimGraph plugin
++ClassRedirects=(OldName="/Script/AnimNext.AnimNextAnimationGraph", NewName="/Script/UAFAnimGraph.AnimNextAnimationGraph")
++StructRedirects=(OldName="/Script/AnimNext.AnimNextGraphContextData", NewName="/Script/UAFAnimGraph.AnimNextGraphContextData")
++StructRedirects=(OldName="/Script/AnimNext.AnimNextGraphEntryPoint", NewName="/Script/UAFAnimGraph.AnimNextGraphEntryPoint")
++StructRedirects=(OldName="/Script/AnimNext.RigDecorator_AnimNextCppDecorator", NewName="/Script/UAFAnimGraph.RigDecorator_AnimNextCppDecorator")
++StructRedirects=(OldName="/Script/AnimNext.AnimNextGraphEvaluatorExecuteArgument", NewName="/Script/UAFAnimGraph.AnimNextGraphEvaluatorExecuteArgument")
++StructRedirects=(OldName="/Script/AnimNext.AnimNextGraphEvaluatorExecuteDefinition", NewName="/Script/UAFAnimGraph.AnimNextGraphEvaluatorExecuteDefinition")
++StructRedirects=(OldName="/Script/AnimNext.RigUnit_AnimNextGraphEvaluator", NewName="/Script/UAFAnimGraph.RigUnit_AnimNextGraphEvaluator")
++StructRedirects=(OldName="/Script/AnimNext.RigUnit_AnimNextGraphRoot", NewName="/Script/UAFAnimGraph.RigUnit_AnimNextGraphRoot")
++StructRedirects=(OldName="/Script/AnimNext.RigUnit_AnimNextRunAnimationGraph", NewName="/Script/UAFAnimGraph.RigUnit_AnimNextRunAnimationGraph_v1")
++StructRedirects=(OldName="/Script/AnimNext.RigUnit_AnimNextShimRoot", NewName="/Script/UAFAnimGraph.RigUnit_AnimNextShimRoot")
++StructRedirects=(OldName="/Script/AnimNext.RigUnit_AnimNextTraitStack", NewName="/Script/UAFAnimGraph.RigUnit_AnimNextTraitStack")
++StructRedirects=(OldName="/Script/AnimNext.AnimNextEntryPointHandle", NewName="/Script/UAFAnimGraph.AnimNextEntryPointHandle")
++StructRedirects=(OldName="/Script/AnimNext.AnimNextTraitHandle", NewName="/Script/UAFAnimGraph.AnimNextTraitHandle")
++StructRedirects=(OldName="/Script/AnimNext.AnimNextTraitSharedData", NewName="/Script/UAFAnimGraph.AnimNextTraitSharedData")
++StructRedirects=(OldName="/Script/AnimNext.AnimNextGraphInstance", NewName="/Script/UAFAnimGraph.AnimNextGraphInstance")
+
++ClassRedirects=(OldName="/Script/AnimNextEditor.AnimNextAnimationGraphFactory", NewName="/Script/UAFAnimGraphEditor.AnimNextAnimationGraphFactory")
++ClassRedirects=(OldName="/Script/AnimNextEditor.AssetDefinition_AnimNextAnimationGraph", NewName="/Script/UAFAnimGraphEditor.AssetDefinition_AnimNextAnimationGraph")
+
++ClassRedirects=(OldName="/Script/AnimNextUncookedOnly.AnimNextAnimationGraph_EditorData", NewName="/Script/UAFAnimGraphUncookedOnly.AnimNextAnimationGraph_EditorData")
++ClassRedirects=(OldName="/Script/AnimNextUncookedOnly.AnimNextAnimationGraphSchema", NewName="/Script/UAFAnimGraphUncookedOnly.AnimNextAnimationGraphSchema")
++StructRedirects=(OldName="/Script/AnimNextUncookedOnly.AnimNextAnimationGraphOutlinerData", NewName="/Script/UAFAnimGraphUncookedOnly.AnimNextAnimationGraphOutlinerData")
++ClassRedirects=(OldName="/Script/AnimNextUncookedOnly.AnimNextAnimationGraphEntry", NewName="/Script/UAFAnimGraphUncookedOnly.AnimNextAnimationGraphEntry")
+
+; PlayAnim -> Injection
++StructRedirects=(OldName="/Script/AnimNext.ModuleInstanceComponent", NewName="/Script/UAF.UAFModuleInstanceComponent")
++StructRedirects=(OldName="/Script/AnimNext.AnimNextGraphInstancePtr", NewName="/Script/UAFAnimGraph.AnimNextAnimGraph")
+
+; Pass-by-ref
++StructRedirects=(OldName="/Script/AnimNext.AnimNextDataInterfaceInstance", NewName="/Script/UAF.UAFAssetInstance")
++ClassRedirects=(OldName="/Script/AnimNext.AnimNextDataInterface", NewName="/Script/UAF.AnimNextSharedVariables")
++ClassRedirects=(OldName="/Script/AnimNextEditor.AnimNextDataInterfaceFactory", NewName="/Script/UAFEditor.AnimNextSharedVariablesFactory")
++ClassRedirects=(OldName="/Script/AnimNextUncookedOnly.AnimNextDataInterface_EditorData", NewName="/Script/UAFUncookedOnly.AnimNextSharedVariables_EditorData")
++ClassRedirects=(OldName="/Script/AnimNextUncookedOnly.AnimNextDataInterfaceEntry", NewName="/Script/UAFUncookedOnly.AnimNextSharedVariablesEntry")
++PropertyRedirects=(OldName="/Script/AnimNextUncookedOnly.AnimNextSharedVariablesEntry.DataInterface", NewName="Asset")
++PropertyRedirects=(OldName="/Script/AnimNextUncookedOnly.AnimNextSharedVariablesEntry.DataInterfacePath", NewName="ObjectPath")
++FunctionRedirects=(OldName="AnimNextRigVMAssetLibrary.AddDataInterface",NewName="AnimNextRigVMAssetLibrary.AddSharedVariables")
++FunctionRedirects=(OldName="AnimNextComponent.SetVariable",NewName="AnimNextComponent.BlueprintSetVariable")
+
+; AnimNext plugin -> UAF plugin
++PackageRedirects=(OldName="/Script/AnimNext", NewName="/Script/UAF")
++PackageRedirects=(OldName="/Script/AnimNextEditor", NewName="/Script/UAFEditor")
++PackageRedirects=(OldName="/Script/AnimNextUncookedOnly", NewName="/Script/UAFUncookedOnly")
++PackageRedirects=(OldName="/Script/AnimNextTestSuite", NewName="/Script/UAFTestSuite")
+
++StructRedirects=(OldName="/Script/AnimNext.AnimNextExecuteContext", NewName="/Script/UAF.AnimNextExecuteContext")
++StructRedirects=(OldName="/Script/AnimNext.AnimNextNativeDataInterface", NewName="/Script/UAF.AnimNextNativeDataInterface")
++StructRedirects=(OldName="/Script/AnimNext.AnimNextGraphReferencePose", NewName="/Script/UAF.AnimNextGraphReferencePose")
++StructRedirects=(OldName="/Script/AnimNext.AnimNextGraphLODPose", NewName="/Script/UAF.AnimNextGraphLODPose")
++StructRedirects=(OldName="/Script/AnimNext.RigUnit_AnimNextBase", NewName="/Script/UAF.RigUnit_AnimNextBase")
++StructRedirects=(OldName="/Script/AnimNext.RigUnit_AnimNextBeginExecution", NewName="/Script/UAF.RigUnit_AnimNextBeginExecution")
++StructRedirects=(OldName="/Script/AnimNext.RigVMTrait_AnimNextPublicVariables", NewName="/Script/UAF.RigVMTrait_AnimNextPublicVariables")
++StructRedirects=(OldName="/Script/AnimNext.AnimNextActorComponentReferenceComponent", NewName="/Script/UAF.AnimNextActorComponentReferenceComponent")
++StructRedirects=(OldName="/Script/AnimNext.AnimNextModuleContextData", NewName="/Script/UAF.AnimNextModuleContextData")
++StructRedirects=(OldName="/Script/AnimNext.AnimNextModuleInstance", NewName="/Script/UAF.AnimNextModuleInstance")
++StructRedirects=(OldName="/Script/AnimNext.AnimNextSkeletalMeshComponentReferenceComponent", NewName="/Script/UAF.AnimNextSkeletalMeshComponentReferenceComponent")
++StructRedirects=(OldName="/Script/AnimNext.RigUnit_AnimNextModuleAddDependencies", NewName="/Script/UAF.RigUnit_AnimNextModuleAddDependencies")
++StructRedirects=(OldName="/Script/AnimNext.RigUnit_AnimNextModuleRemoveDependencies", NewName="/Script/UAF.RigUnit_AnimNextModuleRemoveDependencies")
++StructRedirects=(OldName="/Script/AnimNext.RigUnit_AnimNextModuleEventBase", NewName="/Script/UAF.RigUnit_AnimNextModuleEventBase")
++StructRedirects=(OldName="/Script/AnimNext.RigUnit_AnimNextExecuteBindings_GT", NewName="/Script/UAF.RigUnit_AnimNextExecuteBindings_GT")
++StructRedirects=(OldName="/Script/AnimNext.RigUnit_AnimNextExecuteBindings_WT", NewName="/Script/UAF.RigUnit_AnimNextExecuteBindings_WT")
++StructRedirects=(OldName="/Script/AnimNext.RigUnit_AnimNextInitializeEvent", NewName="/Script/UAF.RigUnit_AnimNextInitializeEvent")
++StructRedirects=(OldName="/Script/AnimNext.RigUnit_AnimNextUserEvent", NewName="/Script/UAF.RigUnit_AnimNextUserEvent")
++StructRedirects=(OldName="/Script/AnimNext.RigUnit_AnimNextPrePhysicsEvent", NewName="/Script/UAF.RigUnit_AnimNextPrePhysicsEvent")
++StructRedirects=(OldName="/Script/AnimNext.RigUnit_AnimNextPostPhysicsEvent", NewName="/Script/UAF.RigUnit_AnimNextPostPhysicsEvent")
++StructRedirects=(OldName="/Script/AnimNext.AnimNextActorLocatorFragment", NewName="/Script/UAF.AnimNextActorLocatorFragment")
++StructRedirects=(OldName="/Script/AnimNext.AnimNextComponentLocatorFragment", NewName="/Script/UAF.AnimNextComponentLocatorFragment")
++StructRedirects=(OldName="/Script/AnimNext.AnimNextObjectCastLocatorFragment", NewName="/Script/UAF.AnimNextObjectCastLocatorFragment")
++StructRedirects=(OldName="/Script/AnimNext.AnimNextObjectFunctionLocatorFragment", NewName="/Script/UAF.AnimNextObjectFunctionLocatorFragment")
++StructRedirects=(OldName="/Script/AnimNext.AnimNextObjectPropertyLocatorFragment", NewName="/Script/UAF.AnimNextObjectPropertyLocatorFragment")
++StructRedirects=(OldName="/Script/AnimNext.RigUnit_AnimNextParameterBase", NewName="/Script/UAF.RigUnit_AnimNextParameterBase")
++StructRedirects=(OldName="/Script/AnimNext.RigUnit_AnimNextParameterBeginExecution", NewName="/Script/UAF.RigUnit_AnimNextParameterBeginExecution")
++StructRedirects=(OldName="/Script/AnimNext.RigVMDispatch_GetLayerParameter", NewName="/Script/UAF.RigVMDispatch_GetLayerParameter")
++StructRedirects=(OldName="/Script/AnimNext.RigVMDispatch_GetParameter", NewName="/Script/UAF.RigVMDispatch_GetParameter")
++StructRedirects=(OldName="/Script/AnimNext.RigVMDispatch_SetLayerParameter", NewName="/Script/UAF.RigVMDispatch_SetLayerParameter")
++StructRedirects=(OldName="/Script/AnimNext.AnimNextFieldPath", NewName="/Script/UAF.AnimNextFieldPath")
++StructRedirects=(OldName="/Script/AnimNext.AnimNextSoftFunctionPtr", NewName="/Script/UAF.AnimNextSoftFunctionPtr")
++StructRedirects=(OldName="/Script/AnimNext.AnimNextVariableOverridesCollection", NewName="/Script/UAF.AnimNextVariableOverridesCollection")
++StructRedirects=(OldName="/Script/AnimNext.RigUnit_CopyModuleProxyVariables", NewName="/Script/UAF.RigUnit_CopyModuleProxyVariables")
++StructRedirects=(OldName="/Script/AnimNext.RigUnit_ResolveUniversalObjectLocator", NewName="/Script/UAF.RigUnit_ResolveUniversalObjectLocator")
++StructRedirects=(OldName="/Script/AnimNext.RigVMDispatch_CallHoistedAccessorFunctionBase", NewName="/Script/UAF.RigVMDispatch_CallHoistedAccessorFunctionBase")
++StructRedirects=(OldName="/Script/AnimNext.RigVMDispatch_CallHoistedAccessorFunctionNative", NewName="/Script/UAF.RigVMDispatch_CallHoistedAccessorFunctionNative")
++StructRedirects=(OldName="/Script/AnimNext.RigVMDispatch_CallHoistedAccessorFunctionScript", NewName="/Script/UAF.RigVMDispatch_CallHoistedAccessorFunctionScript")
++StructRedirects=(OldName="/Script/AnimNext.RigVMDispatch_CallObjectAccessorFunctionBase", NewName="/Script/UAF.RigVMDispatch_CallObjectAccessorFunctionBase")
++StructRedirects=(OldName="/Script/AnimNext.RigVMDispatch_CallObjectAccessorFunctionNative", NewName="/Script/UAF.RigVMDispatch_CallObjectAccessorFunctionNative")
++StructRedirects=(OldName="/Script/AnimNext.RigVMDispatch_CallObjectAccessorFunctionScript", NewName="/Script/UAF.RigVMDispatch_CallObjectAccessorFunctionScript")
++StructRedirects=(OldName="/Script/AnimNext.RigVMDispatch_GetObjectProperty", NewName="/Script/UAF.RigVMDispatch_GetObjectProperty")
++StructRedirects=(OldName="/Script/AnimNext.UAFInstanceVariableContainer", NewName="/Script/UAF.UAFInstanceVariableContainer")
++StructRedirects=(OldName="/Script/AnimNext.UAFInstanceVariableContainerProxy", NewName="/Script/UAF.UAFInstanceVariableContainerProxy")
++StructRedirects=(OldName="/Script/AnimNext.UAFInstanceVariableDataProxy", NewName="/Script/UAF.UAFInstanceVariableDataProxy")
++StructRedirects=(OldName="/Script/AnimNext.RemappedBone", NewName="/Script/UAF.RemappedBone")
++StructRedirects=(OldName="/Script/AnimNext.BoneRemapping", NewName="/Script/UAF.BoneRemapping")
++StructRedirects=(OldName="/Script/AnimNext.RemapPoseData", NewName="/Script/UAF.RemapPoseData")
++StructRedirects=(OldName="/Script/AnimNext.RigUnit_AnimNextRemapPose", NewName="/Script/UAF.RigUnit_AnimNextRemapPose")
++StructRedirects=(OldName="/Script/AnimNext.RigUnit_AnimNextWriteSkeletalMeshComponentPose", NewName="/Script/UAF.RigUnit_AnimNextWriteSkeletalMeshComponentPose")
++StructRedirects=(OldName="/Script/AnimNext.AnimNextActorTransformComponent", NewName="/Script/UAF.AnimNextActorTransformComponent")
++StructRedirects=(OldName="/Script/AnimNext.RigUnit_GetActorTransform", NewName="/Script/UAF.RigUnit_GetActorTransform")
++StructRedirects=(OldName="/Script/AnimNext.AnimNextActorTransformComponent", NewName="/Script/UAF.AnimNextActorTransformComponent")
++StructRedirects=(OldName="/Script/AnimNext.RigUnit_GetActorTransform", NewName="/Script/UAF.RigUnit_GetActorTransform")
++StructRedirects=(OldName="/Script/AnimNext.RigUnit_IsAnimBPDriven", NewName="/Script/UAF.RigUnit_IsAnimBPDriven")
++StructRedirects=(OldName="/Script/AnimNext.RigUnit_IsRecentlyRendered", NewName="/Script/UAF.RigUnit_IsRecentlyRendered")
++StructRedirects=(OldName="/Script/AnimNext.RigUnit_MakeReferencePoseFromSkeletalMeshComponent", NewName="/Script/UAF.RigUnit_MakeReferencePoseFromSkeletalMeshComponent")
++StructRedirects=(OldName="/Script/AnimNext.RigVMTrait_ModuleEventDependency_ActorComponentClassPrimaryTickFunction", NewName="/Script/UAF.RigVMTrait_ModuleEventDependency_ActorComponentClassPrimaryTickFunction")
++StructRedirects=(OldName="/Script/AnimNext.RigVMTrait_ModuleEventDependency_ActorComponentPrimaryTickFunction", NewName="/Script/UAF.RigVMTrait_ModuleEventDependency_ActorComponentPrimaryTickFunction")
++StructRedirects=(OldName="/Script/AnimNext.AnimNextOverrideVariablesWorkData", NewName="/Script/UAF.AnimNextOverrideVariablesWorkData")
++StructRedirects=(OldName="/Script/AnimNext.RigUnit_OverrideVariables", NewName="/Script/UAF.RigUnit_OverrideVariables")
++StructRedirects=(OldName="/Script/AnimNext.AnimNextVariableOverrideInfo", NewName="/Script/UAF.AnimNextVariableOverrideInfo")
++StructRedirects=(OldName="/Script/AnimNext.RigVMTrait_OverrideVariables", NewName="/Script/UAF.RigVMTrait_OverrideVariables")
++StructRedirects=(OldName="/Script/AnimNext.RigVMTrait_OverrideVariablesForAsset", NewName="/Script/UAF.RigVMTrait_OverrideVariablesForAsset")
++StructRedirects=(OldName="/Script/AnimNext.RigVMTrait_OverrideVariablesForStruct", NewName="/Script/UAF.RigVMTrait_OverrideVariablesForStruct")
++StructRedirects=(OldName="/Script/AnimNext.AnimNextGraphState", NewName="/Script/UAF.AnimNextGraphState")
++StructRedirects=(OldName="/Script/AnimNext.AnimNextLODPose", NewName="/Script/UAF.AnimNextLODPose")
++StructRedirects=(OldName="/Script/AnimNext.AnimNextModuleHandle", NewName="/Script/UAF.AnimNextModuleHandle")
++StructRedirects=(OldName="/Script/AnimNext.RigVMTrait_ModuleEventDependency", NewName="/Script/UAF.RigVMTrait_ModuleEventDependency")
++StructRedirects=(OldName="/Script/AnimNext.UAFModuleInstanceComponent", NewName="/Script/UAF.UAFModuleInstanceComponent")
++StructRedirects=(OldName="/Script/AnimNext.AnimNextTag", NewName="/Script/UAF.AnimNextTag")
++StructRedirects=(OldName="/Script/AnimNext.AnimNextScope", NewName="/Script/UAF.AnimNextScope")
++StructRedirects=(OldName="/Script/AnimNext.AnimNextEntryPoint", NewName="/Script/UAF.AnimNextEntryPoint")
++StructRedirects=(OldName="/Script/AnimNext.AnimNextParamType", NewName="/Script/UAF.AnimNextParamType")
++StructRedirects=(OldName="/Script/AnimNext.AnimNextReferencePose", NewName="/Script/UAF.AnimNextReferencePose")
++StructRedirects=(OldName="/Script/AnimNext.AnimNextTraitEvent", NewName="/Script/UAF.AnimNextTraitEvent")
++StructRedirects=(OldName="/Script/AnimNext.UAFAssetInstance", NewName="/Script/UAF.UAFAssetInstance")
++StructRedirects=(OldName="/Script/AnimNext.UAFAssetInstanceComponent", NewName="/Script/UAF.UAFAssetInstanceComponent")
++StructRedirects=(OldName="/Script/AnimNext.UAFRigVMComponent", NewName="/Script/UAF.UAFRigVMComponent")
++StructRedirects=(OldName="/Script/AnimNext.AnimNextSoftVariableReference", NewName="/Script/UAF.AnimNextSoftVariableReference")
++StructRedirects=(OldName="/Script/AnimNext.AnimNextVariableReference", NewName="/Script/UAF.AnimNextVariableReference")
++StructRedirects=(OldName="/Script/AnimNext.UAFInstanceVariableData", NewName="/Script/UAF.UAFInstanceVariableData")
++StructRedirects=(OldName="/Script/AnimNext.AnimNextRigVMFunctionData", NewName="/Script/UAF.AnimNextRigVMFunctionData")
++StructRedirects=(OldName="/Script/AnimNext.AnimNextRigVMFunctionData", NewName="/Script/UAF.AnimNextRigVMFunctionData")
+
++ClassRedirects=(OldName="/Script/AnimNext.AnimNextConfig", NewName="/Script/UAF.AnimNextConfig")
++ClassRedirects=(OldName="/Script/AnimNext.AnimNextWorldSubsystem", NewName="/Script/UAF.AnimNextWorldSubsystem")
++ClassRedirects=(OldName="/Script/AnimNext.AnimNextModule", NewName="/Script/UAF.AnimNextModule")
++ClassRedirects=(OldName="/Script/AnimNext.AnimNextComponentWorldSubsystem", NewName="/Script/UAF.AnimNextComponentWorldSubsystem")
++ClassRedirects=(OldName="/Script/AnimNext.AnimNextAddDependencyMenuContext", NewName="/Script/UAF.AnimNextAddDependencyMenuContext")
++ClassRedirects=(OldName="/Script/AnimNext.AnimNextWorldLibrary", NewName="/Script/UAF.AnimNextWorldLibrary")
++ClassRedirects=(OldName="/Script/AnimNext.AnimNextOverrideVariablesMenuContext", NewName="/Script/UAF.AnimNextOverrideVariablesMenuContext")
++ClassRedirects=(OldName="/Script/AnimNext.AnimNextRigVMAsset", NewName="/Script/UAF.AnimNextRigVMAsset")
++ClassRedirects=(OldName="/Script/AnimNext.AnimNextComponent", NewName="/Script/UAF.AnimNextComponent")
++ClassRedirects=(OldName="/Script/AnimNext.AnimNextSharedVariables", NewName="/Script/UAF.AnimNextSharedVariables")
+
++EnumRedirects=(OldName="EAnimNextModuleInitMethod", NewName="/Script/UAF.EAnimNextModuleInitMethod")
++EnumRedirects=(OldName="EAnimNextModuleEventDependencyOrdering", NewName="/Script/UAF.EAnimNextModuleEventDependencyOrdering")
+
++StructRedirects=(OldName="/Script/AnimNextEditor.AnimNextSchemaAction", NewName="/Script/UAFEditor.AnimNextSchemaAction")
++StructRedirects=(OldName="/Script/AnimNextEditor.AnimNextSchemaAction_RigUnit", NewName="/Script/UAFEditor.AnimNextSchemaAction_RigUnit")
++StructRedirects=(OldName="/Script/AnimNextEditor.AnimNextSchemaAction_DispatchFactory", NewName="/Script/UAFEditor.AnimNextSchemaAction_DispatchFactory")
++StructRedirects=(OldName="/Script/AnimNextEditor.AnimNextSchemaAction_Variable", NewName="/Script/UAFEditor.AnimNextSchemaAction_Variable")
++StructRedirects=(OldName="/Script/AnimNextEditor.AnimNextSchemaAction_AddComment", NewName="/Script/UAFEditor.AnimNextSchemaAction_AddComment")
++StructRedirects=(OldName="/Script/AnimNextEditor.AnimNextSchemaAction_Function", NewName="/Script/UAFEditor.AnimNextSchemaAction_Function")
++StructRedirects=(OldName="/Script/AnimNextEditor.AnimNextSchemaAction_PromoteToVariable", NewName="/Script/UAFEditor.AnimNextSchemaAction_PromoteToVariable")
++StructRedirects=(OldName="/Script/AnimNextEditor.AnimNextWorkspaceState", NewName="/Script/UAFEditor.AnimNextWorkspaceState")
+
++ClassRedirects=(OldName="/Script/AnimNextEditor.AnimNextModuleFactory", NewName="/Script/UAFEditor.AnimNextModuleFactory")
++ClassRedirects=(OldName="/Script/AnimNextEditor.AnimNextSharedVariablesFactory", NewName="/Script/UAFEditor.AnimNextSharedVariablesFactory")
++ClassRedirects=(OldName="/Script/AnimNextEditor.AnimNextWorkspaceFactory", NewName="/Script/UAFEditor.AnimNextWorkspaceFactory")
++ClassRedirects=(OldName="/Script/AnimNextEditor.AnimNextEditorContext", NewName="/Script/UAFEditor.AnimNextEditorContext")
++ClassRedirects=(OldName="/Script/AnimNextEditor.AnimNextAssetFindReplaceVariables", NewName="/Script/UAFEditor.AnimNextAssetFindReplaceVariables")
++ClassRedirects=(OldName="/Script/AnimNextEditor.AssetDefinition_AnimNextVariableEntry", NewName="/Script/UAFEditor.AssetDefinition_AnimNextVariableEntry")
++ClassRedirects=(OldName="/Script/AnimNextEditor.AssetDefinition_AnimNextEventGraphEntry", NewName="/Script/UAFEditor.AssetDefinition_AnimNextEventGraphEntry")
++ClassRedirects=(OldName="/Script/AnimNextEditor.AssetDefinition_AnimNextSharedVariablesEntry", NewName="/Script/UAFEditor.AssetDefinition_AnimNextSharedVariablesEntry")
++ClassRedirects=(OldName="/Script/AnimNextEditor.RigVMAssetViewMenuContext", NewName="/Script/UAFEditor.RigVMAssetViewMenuContext")
++ClassRedirects=(OldName="/Script/AnimNextEditor.AssetDefinition_AnimNextModule", NewName="/Script/UAFEditor.AssetDefinition_AnimNextModule")
++ClassRedirects=(OldName="/Script/AnimNextEditor.PropertyBagDetailsObject", NewName="/Script/UAFEditor.PropertyBagDetailsObject")
++ClassRedirects=(OldName="/Script/AnimNextEditor.AssetDefinition_AnimNextSharedVariables", NewName="/Script/UAFEditor.AssetDefinition_AnimNextSharedVariables")
++ClassRedirects=(OldName="/Script/AnimNextEditor.AnimNextVariableItemMenuContext", NewName="/Script/UAFEditor.AnimNextVariableItemMenuContext")
++ClassRedirects=(OldName="/Script/AnimNextEditor.AnimNextVariableSettings", NewName="/Script/UAFEditor.AnimNextVariableSettings")
++ClassRedirects=(OldName="/Script/AnimNextEditor.AnimNextWorkspaceEditorMode", NewName="/Script/UAFEditor.AnimNextWorkspaceEditorMode")
++ClassRedirects=(OldName="/Script/AnimNextEditor.AnimNextWorkspaceSchema", NewName="/Script/UAFEditor.AnimNextWorkspaceSchema")
+
++EnumRedirects=(OldName="ESearchScope", NewName="/Script/UAFEditor.ESearchScope")
+
++StructRedirects=(OldName="/Script/AnimNextTestSuite.AnimNextParamTypeTestStruct", NewName="/Script/UAFTestSuite.AnimNextParamTypeTestStruct")
++StructRedirects=(OldName="/Script/AnimNextTestSuite.AnimNextTests_TestOperation", NewName="/Script/UAFTestSuite.AnimNextTests_TestOperation")
++StructRedirects=(OldName="/Script/AnimNextTestSuite.AnimNextTests_PrintResult", NewName="/Script/UAFTestSuite.AnimNextTests_PrintResult")
++StructRedirects=(OldName="/Script/AnimNextTestSuite.AnimNextTestData", NewName="/Script/UAFTestSuite.AnimNextTestData")
+
++ClassRedirects=(OldName="/Script/AnimNextTestSuite.AnimNextTestFuncLib", NewName="/Script/UAFTestSuite.AnimNextTestFuncLib")
+
++StructRedirects=(OldName="/Script/AnimNextUncookedOnly.AnimNextRigVMAssetOutlinerData", NewName="/Script/UAFUncookedOnly.AnimNextRigVMAssetOutlinerData")
++StructRedirects=(OldName="/Script/AnimNextUncookedOnly.AnimNextModuleOutlinerData", NewName="/Script/UAFUncookedOnly.AnimNextModuleOutlinerData")
++StructRedirects=(OldName="/Script/AnimNextUncookedOnly.AnimNextSharedVariablesOutlinerData", NewName="/Script/UAFUncookedOnly.AnimNextSharedVariablesOutlinerData")
++StructRedirects=(OldName="/Script/AnimNextUncookedOnly.AnimNextAssetEntryOutlinerData", NewName="/Script/UAFUncookedOnly.AnimNextAssetEntryOutlinerData")
++StructRedirects=(OldName="/Script/AnimNextUncookedOnly.AnimNextVariableOutlinerData", NewName="/Script/UAFUncookedOnly.AnimNextVariableOutlinerData")
++StructRedirects=(OldName="/Script/AnimNextUncookedOnly.AnimNextCollapseGraphsOutlinerDataBase", NewName="/Script/UAFUncookedOnly.AnimNextCollapseGraphsOutlinerDataBase")
++StructRedirects=(OldName="/Script/AnimNextUncookedOnly.AnimNextCollapseGraphOutlinerData", NewName="/Script/UAFUncookedOnly.AnimNextCollapseGraphOutlinerData")
++StructRedirects=(OldName="/Script/AnimNextUncookedOnly.AnimNextGraphFunctionOutlinerData", NewName="/Script/UAFUncookedOnly.AnimNextGraphFunctionOutlinerData")
++StructRedirects=(OldName="/Script/AnimNextUncookedOnly.AnimNextGraphOutlinerData", NewName="/Script/UAFUncookedOnly.AnimNextGraphOutlinerData")
++StructRedirects=(OldName="/Script/AnimNextUncookedOnly.AnimNextUniversalObjectLocatorBindingData", NewName="/Script/UAFUncookedOnly.AnimNextUniversalObjectLocatorBindingData")
++StructRedirects=(OldName="/Script/AnimNextUncookedOnly.AnimNextExportData", NewName="/Script/UAFUncookedOnly.AnimNextExportData")
++StructRedirects=(OldName="/Script/AnimNextUncookedOnly.AnimNextExport", NewName="/Script/UAFUncookedOnly.AnimNextExport")
++StructRedirects=(OldName="/Script/AnimNextUncookedOnly.AnimNextVariableReferenceData", NewName="/Script/UAFUncookedOnly.AnimNextVariableReferenceData")
++StructRedirects=(OldName="/Script/AnimNextUncookedOnly.AnimNextVariableDeclarationData", NewName="/Script/UAFUncookedOnly.AnimNextVariableDeclarationData")
++StructRedirects=(OldName="/Script/AnimNextUncookedOnly.AnimNextManifestNodeData", NewName="/Script/UAFUncookedOnly.AnimNextManifestNodeData")
++StructRedirects=(OldName="/Script/AnimNextUncookedOnly.AnimNextAssetRegistryExports", NewName="/Script/UAFUncookedOnly.AnimNextAssetRegistryExports")
++StructRedirects=(OldName="/Script/AnimNextUncookedOnly.AnimNextVariableBinding", NewName="/Script/UAFUncookedOnly.AnimNextVariableBinding")
++StructRedirects=(OldName="/Script/AnimNextUncookedOnly.AnimNextVariableBindingData", NewName="/Script/UAFUncookedOnly.AnimNextVariableBindingData")
+
++ClassRedirects=(OldName="/Script/AnimNextUncookedOnly.AnimNextAssetWorkspaceAssetUserData", NewName="/Script/UAFUncookedOnly.AnimNextAssetWorkspaceAssetUserData")
++ClassRedirects=(OldName="/Script/AnimNextUncookedOnly.AnimNextControllerBase", NewName="/Script/UAFUncookedOnly.AnimNextControllerBase")
++ClassRedirects=(OldName="/Script/AnimNextUncookedOnly.AnimNextEdGraph", NewName="/Script/UAFUncookedOnly.AnimNextEdGraph")
++ClassRedirects=(OldName="/Script/AnimNextUncookedOnly.AnimNextEdGraphNode", NewName="/Script/UAFUncookedOnly.AnimNextEdGraphNode")
++ClassRedirects=(OldName="/Script/AnimNextUncookedOnly.AnimNextEdGraphSchema", NewName="/Script/UAFUncookedOnly.AnimNextEdGraphSchema")
++ClassRedirects=(OldName="/Script/AnimNextUncookedOnly.AnimNextRigVMAssetLibrary", NewName="/Script/UAFUncookedOnly.AnimNextRigVMAssetLibrary")
++ClassRedirects=(OldName="/Script/AnimNextUncookedOnly.AnimNextRigVMAssetEditorData", NewName="/Script/UAFUncookedOnly.AnimNextRigVMAssetEditorData")
++ClassRedirects=(OldName="/Script/AnimNextUncookedOnly.AnimNextRigVMAssetSchema", NewName="/Script/UAFUncookedOnly.AnimNextRigVMAssetSchema")
++ClassRedirects=(OldName="/Script/AnimNextUncookedOnly.AnimNextSharedVariableNode", NewName="/Script/UAFUncookedOnly.AnimNextSharedVariableNode")
++ClassRedirects=(OldName="/Script/AnimNextUncookedOnly.AnimNextUnitNode", NewName="/Script/UAFUncookedOnly.AnimNextUnitNode")
++ClassRedirects=(OldName="/Script/AnimNextUncookedOnly.AnimNextCategoryEntry", NewName="/Script/UAFUncookedOnly.AnimNextCategoryEntry")
++ClassRedirects=(OldName="/Script/AnimNextUncookedOnly.AnimNextEventGraphEntry", NewName="/Script/UAFUncookedOnly.AnimNextEventGraphEntry")
++ClassRedirects=(OldName="/Script/AnimNextUncookedOnly.AnimNextRigVMAssetEntry", NewName="/Script/UAFUncookedOnly.AnimNextRigVMAssetEntry")
++ClassRedirects=(OldName="/Script/AnimNextUncookedOnly.AnimNextSharedVariablesEntry", NewName="/Script/UAFUncookedOnly.AnimNextSharedVariablesEntry")
++ClassRedirects=(OldName="/Script/AnimNextUncookedOnly.AnimNextVariableEntry", NewName="/Script/UAFUncookedOnly.AnimNextVariableEntry")
++ClassRedirects=(OldName="/Script/AnimNextUncookedOnly.AnimNextEventGraphSchema", NewName="/Script/UAFUncookedOnly.AnimNextEventGraphSchema")
++ClassRedirects=(OldName="/Script/AnimNextUncookedOnly.AnimNextModule_EditorData", NewName="/Script/UAFUncookedOnly.AnimNextModule_EditorData")
++ClassRedirects=(OldName="/Script/AnimNextUncookedOnly.AnimNextSharedVariables_EditorData", NewName="/Script/UAFUncookedOnly.AnimNextSharedVariables_EditorData")
++ClassRedirects=(OldName="/Script/AnimNextUncookedOnly.K2Node_AnimNextComponentSetVariable", NewName="/Script/UAFUncookedOnly.K2Node_AnimNextComponentSetVariable")
+
++EnumRedirects=(OldName="EAnimNextSharedVariablesType", NewName="/Script/UAFUncookedOnly.EAnimNextSharedVariablesType")
++EnumRedirects=(OldName="EAnimNextExportAccessSpecifier", NewName="/Script/UAFUncookedOnly.EAnimNextExportAccessSpecifier")
++EnumRedirects=(OldName="FAnimNextUniversalObjectLocatorBindingType", NewName="/Script/UAFUncookedOnly.FAnimNextUniversalObjectLocatorBindingType")
++EnumRedirects=(OldName="EAnimNextExportedVariableFlags", NewName="/Script/UAFUncookedOnly.EAnimNextExportedVariableFlags")
\ No newline at end of file
diff --git a/Engine/Plugins/Experimental/UAF/UAF/Source/UAF/Internal/Graph/RigUnit_AnimNextBase.h b/Engine/Plugins/Experimental/UAF/UAF/Source/UAF/Internal/Graph/RigUnit_AnimNextBase.h
new file mode 100644
--- /dev/null
+++ b/Engine/Plugins/Experimental/UAF/UAF/Source/UAF/Internal/Graph/RigUnit_AnimNextBase.h
@@ -0,0 +1,14 @@
+﻿// Copyright Epic Games, Inc. All Rights Reserved.
+
+#pragma once
+
+#include "CoreMinimal.h"
+#include "RigVMCore/RigVMStruct.h"
+#include "AnimNextExecuteContext.h"
+#include "RigUnit_AnimNextBase.generated.h"
+
+USTRUCT(meta=(ExecuteContext="FAnimNextExecuteContext"))
+struct FRigUnit_AnimNextBase : public FRigVMStruct
+{
+	GENERATED_BODY()
+};
diff --git a/Engine/Plugins/Experimental/UAF/UAF/Source/UAF/Internal/Module/AnimNextModuleContextData.h b/Engine/Plugins/Experimental/UAF/UAF/Source/UAF/Internal/Module/AnimNextModuleContextData.h
new file mode 100644
--- /dev/null
+++ b/Engine/Plugins/Experimental/UAF/UAF/Source/UAF/Internal/Module/AnimNextModuleContextData.h
@@ -0,0 +1,46 @@
+// Copyright Epic Games, Inc. All Rights Reserved.
+
+#pragma once
+
+#include "CoreMinimal.h"
+#include "UAFAssetInstance.h"
+#include "AnimNextModuleContextData.generated.h"
+
+struct FAnimNextModuleInstance;
+
+USTRUCT()
+struct FAnimNextModuleContextData
+{
+	GENERATED_BODY()
+
+	FAnimNextModuleContextData() = default;
+
+	UAF_API explicit FAnimNextModuleContextData(FAnimNextModuleInstance* InModuleInstance);
+
+	UAF_API FAnimNextModuleContextData(FAnimNextModuleInstance* InModuleInstance, const FUAFAssetInstance* InInstance);
+
+	// Get the object that the module instance is bound to, if any
+	UAF_API UObject* GetObject() const;
+
+	FAnimNextModuleInstance& GetModuleInstance() const
+	{
+		check(ModuleInstance != nullptr);
+		return *ModuleInstance;
+	}
+
+	const FUAFAssetInstance& GetInstance() const
+	{
+		check(Instance != nullptr);
+		return *Instance;
+	}
+
+private:
+	// Module instance that is currently executing.
+	FAnimNextModuleInstance* ModuleInstance = nullptr;
+
+	// Instance that is currently executing. Can be the same as ModuleInstance
+	const FUAFAssetInstance* Instance = nullptr;
+
+	friend class UAnimNextModule;
+	friend struct FAnimNextExecuteContext;
+};
diff --git a/Engine/Plugins/Experimental/UAF/UAF/Source/UAF/Internal/Module/AnimNextModuleInstance.h b/Engine/Plugins/Experimental/UAF/UAF/Source/UAF/Internal/Module/AnimNextModuleInstance.h
new file mode 100644
--- /dev/null
+++ b/Engine/Plugins/Experimental/UAF/UAF/Source/UAF/Internal/Module/AnimNextModuleInstance.h
@@ -0,0 +1,280 @@
+// Copyright Epic Games, Inc. All Rights Reserved.
+
+#pragma once
+
+#include "CoreMinimal.h"
+#include "AnimNextModule.h"
+#include "Module/ModuleHandle.h"
+#include "Variables/UAFInstanceVariableDataProxy.h"
+#include "UAFAssetInstance.h"
+#include "Module/ModuleTickFunction.h"
+#include "TraitCore/TraitEvent.h"
+#include "TraitCore/TraitEventList.h"
+#include "Module/ModuleHandle.h"
+#include "Module/UAFModuleInstanceComponent.h"
+#include "StructUtils/InstancedStruct.h"
+#include "Module/AnimNextModuleInitMethod.h"
+#include "Module/TaskRunLocation.h"
+#include "AnimNextModuleInstance.generated.h"
+
+struct FAnimNextModuleInstance;
+struct FRigUnit_AnimNextRunAnimationGraph_v1;
+struct FRigUnit_AnimNextRunAnimationGraph_v2;
+class UAnimNextWorldSubsystem;
+struct FRigUnit_AnimNextInitializeEvent;
+struct FRigUnit_CopyModuleProxyVariables;
+struct FAnimNextNotifyDispatcherComponent;
+struct FAnimNextGraphInstance;
+
+namespace UE::UAF
+{
+	struct FModuleEventTickFunction;
+	struct FModuleEndTickFunction;
+}
+
+namespace UE::UAF::Debug
+{
+	struct FDebugDraw;
+}
+
+namespace UE::UAF
+{
+	using ModuleInstanceComponentMapType = TMap<FName, TInstancedStruct<FUAFModuleInstanceComponent>>;
+}
+
+namespace UE::UAF::Tests
+{
+	class FVariables_Shared;
+}
+
+// Root memory owner of a parameterized schedule 
+USTRUCT()
+struct FAnimNextModuleInstance : public FUAFAssetInstance
+{
+	GENERATED_BODY()
+
+	FAnimNextModuleInstance();
+	UAF_API FAnimNextModuleInstance(
+		UAnimNextModule* InModule,
+		UObject* InObject,
+		UE::UAF::TPool<FAnimNextModuleInstance>* InPool,
+		EAnimNextModuleInitMethod InInitMethod);
+	UAF_API ~FAnimNextModuleInstance();
+
+	// Checks to see if this entry is ticking
+	bool IsEnabled() const;
+
+	// Enables/disables the ticking of this entry
+	void Enable(bool bInEnabled);
+
+	// Queues an input trait event
+	// Input events will be processed in the next graph update after they are queued
+	UAF_API void QueueInputTraitEvent(FAnimNextTraitEventPtr Event);
+
+	// Queues an output trait event
+	// Output events will be processed at the end of the schedule tick
+	UAF_API void QueueOutputTraitEvent(FAnimNextTraitEventPtr Event);
+
+	// Get the object that this module is bound to
+	UObject* GetObject() const { return Object; }
+
+	// Get the module that this instance represents
+	UAF_API const UAnimNextModule* GetModule() const;
+
+#if UE_ENABLE_DEBUG_DRAWING
+	// Get the debug draw interface
+	UAF_API FRigVMDrawInterface* GetDebugDrawInterface();
+
+	// Set whether to show the module's debug drawing instructions in the viewport
+	UAF_API void ShowDebugDrawing(bool bInShowDebugDrawing);
+#endif
+
+	// Run a simple task on the GT via FFunctionGraphTask::CreateAndDispatchWhenReady
+	UAF_API static void RunTaskOnGameThread(TUniqueFunction<void(void)>&& InFunction);
+
+	// Find the tick function for the specified event
+	UE::UAF::FModuleEventTickFunction* FindTickFunctionByName(FName InEventName);
+	const UE::UAF::FModuleEventTickFunction* FindTickFunctionByName(FName InEventName) const;
+
+	// Find the first 'user' tick function
+	UAF_API UE::UAF::FModuleEventTickFunction* FindFirstUserTickFunction();
+
+	// Run the specified RigVM event
+	UAF_API void RunRigVMEvent(FName InEventName, float InDeltaTime);
+
+	// Get the world type that this module was instantiated within
+	EWorldType::Type GetWorldType() const { return WorldType; }
+
+	// Get tick functions for this module instance
+	UAF_API TArrayView<UE::UAF::FModuleEventTickFunction> GetTickFunctions();
+
+	UE::UAF::FModuleHandle GetHandle() const { return Handle; }
+
+	// Queue a task to run at a particular module event
+	// @param	InModuleEvent		The event in the module to run the supplied task relative to. If this is NAME_None or not found then the first 'user' event will be used.
+	// @param	InTaskFunction		The function to run
+	// @param	InLocation			Where to run the task, before or after
+	UAF_API void QueueTask(FName InEventName, TUniqueFunction<void(const UE::UAF::FModuleTaskContext&)>&& InTaskFunction, UE::UAF::ETaskRunLocation InLocation = UE::UAF::ETaskRunLocation::Before);
+
+	// Queue a task to run at a particular module event on some other module
+	// @param	InModuleHandle		The handle for the module to queue the task on
+	// @param	InModuleEventName	The event in the module to run the supplied task relative to. If this is NAME_None or not found then the first user 'event' will be used.
+	// @param	InTaskFunction		The function to run
+	// @param	InLocation			Where to run the task, before or after
+	UAF_API void QueueTaskOnOtherModule(const UE::UAF::FModuleHandle InOtherModuleHandle, FName InEventName, TUniqueFunction<void(const UE::UAF::FModuleTaskContext&)>&& InTaskFunction, UE::UAF::ETaskRunLocation InLocation = UE::UAF::ETaskRunLocation::Before);
+
+private:
+	// Setup the entry
+	UAF_API void Initialize();
+
+	// Clear data that binds the schedule to a runtime (e.g. tick functions) and any instance data
+	void ResetBindingsAndInstanceData();
+
+	// Remove all of our tick function's dependencies
+	void RemoveAllTickDependencies();
+
+	// Flip proxy variable buffers then copy any dirty values
+	void CopyProxyVariables();
+
+#if ANIMNEXT_TRACE_ENABLED
+	// trace this module if it hasn't already been traced this frame
+	UAF_API void Trace();
+#endif
+
+#if WITH_EDITOR
+	// Resets internal state if the module we are bound to is recompiled in editor
+	void OnCompileJobFinished();
+#endif
+
+	// Give the module components a chance to handle events and call their OnEndExecution extension point
+	void EndExecution(float InDeltaTime);
+
+	// Handles trait events at the module level.
+	// Called in the module end tick function before any events are dispatched.
+	void RaiseTraitEvents(const UE::UAF::FTraitEventList& EventList);
+
+	// Set the value of the specified variable via the proxy double-buffer. Called from any thread.
+	EPropertyBagResult SetProxyVariable(const FAnimNextVariableReference& InVariable, const FAnimNextParamType& InType, TConstArrayView<uint8> InData);
+
+	// Accesses value of the specified variable via the proxy double-buffer. Called from any thread.
+	EPropertyBagResult WriteProxyVariable(const FAnimNextVariableReference& InVariable, const FAnimNextParamType& InType, TFunctionRef<void(TArrayView<uint8>)> InFunction);
+
+private:
+	friend UE::UAF::FModuleEventTickFunction;
+	friend UE::UAF::FModuleEndTickFunction;
+	friend FRigUnit_AnimNextRunAnimationGraph_v1;
+	friend FRigUnit_AnimNextRunAnimationGraph_v2;
+	friend UAnimNextWorldSubsystem;
+	friend FRigUnit_AnimNextInitializeEvent;
+	friend FRigUnit_CopyModuleProxyVariables;
+	friend FAnimNextNotifyDispatcherComponent;
+	friend FAnimNextGraphInstance;
+	friend UE::UAF::Tests::FVariables_Shared;
+
+	// Object this entry is bound to
+	UPROPERTY(Transient)
+	TObjectPtr<UObject> Object = nullptr;
+
+	// The pool that this module instance exists in
+	UE::UAF::TPool<FAnimNextModuleInstance>* Pool = nullptr;
+
+	// Copy of the handle that represents this entry to client systems
+	UE::UAF::FModuleHandle Handle;
+
+	// Pre-allocated graph of tick functions
+	TArray<UE::UAF::FModuleEventTickFunction> TickFunctions;
+
+	// Input event list to be processed on the next update
+	UE::UAF::FTraitEventList InputEventList;
+
+	// Output event list to be processed at the end of the schedule tick
+	UE::UAF::FTraitEventList OutputEventList;
+
+	// Lock to ensure event list actions are thread safe
+	FRWLock EventListLock;
+
+	// Lock to ensure proxy access is thread safe
+	FRWLock ProxyLock;
+
+	// Proxy public variables, double-buffered
+	UPROPERTY(Transient)
+	FUAFInstanceVariableDataProxy ProxyVariables[2];
+
+	// Index into current buffer to write the proxy variables
+	int32 ProxyWriteIndex = 0;
+
+#if UE_ENABLE_DEBUG_DRAWING
+	// Debug draw support
+	// TODO: Move this to a 'module component' once that code is checked in
+	TUniquePtr<UE::UAF::Debug::FDebugDraw> DebugDraw;
+#endif
+
+	enum class EInitState : uint8
+	{
+		NotInitialized,
+
+		CreatingTasks,
+
+		BindingTasks,
+
+		SetupVariables,
+
+		PendingInitializeEvent,
+
+		FirstUpdate,
+
+		Initialized
+	};
+
+	// Current initialization state
+	EInitState InitState = EInitState::NotInitialized;
+
+	enum class ERunState : uint8
+	{
+		NotInitialized,
+
+		Running,
+
+		Paused,
+	};
+
+	// Current running state
+	ERunState RunState = ERunState::NotInitialized;
+
+	// Transition to the specified init state, verifying that the current state is valid
+	void TransitionToInitState(EInitState InNewState);
+
+	// Transition to the specified run state, verifying that the current state is valid
+	void TransitionToRunState(ERunState InNewState);
+
+	// How this entry initializes
+	EAnimNextModuleInitMethod InitMethod = EAnimNextModuleInitMethod::InitializeAndPauseInEditor;
+
+	// Whether this represents an editor object 
+	EWorldType::Type WorldType = EWorldType::None;
+
+#if WITH_EDITOR
+	// Whether we are currently recreating this instance because of compilation/reinstancing
+	bool bIsRecreatingOnCompile : 1 = false;
+#endif
+
+private:
+#if ANIMNEXT_TRACE_ENABLED
+	bool bTracedThisFrame = false;
+#endif
+};
+
+template<>
+struct TStructOpsTypeTraits<FAnimNextModuleInstance> : public TStructOpsTypeTraitsBase2<FAnimNextModuleInstance>
+{
+	enum
+	{
+		WithCopy = false
+	};
+};
+
+#if ANIMNEXT_TRACE_ENABLED
+#define TRACE_ANIMNEXT_MODULE(ModuleInstance) ModuleInstance.Trace();
+#else
+#define TRACE_ANIMNEXT_MODULE(ModuleInstance)
+#endif
diff --git a/Engine/Plugins/Experimental/UAF/UAF/Source/UAF/Internal/Module/ModuleEventTickFunctionBindings.h b/Engine/Plugins/Experimental/UAF/UAF/Source/UAF/Internal/Module/ModuleEventTickFunctionBindings.h
new file mode 100644
--- /dev/null
+++ b/Engine/Plugins/Experimental/UAF/UAF/Source/UAF/Internal/Module/ModuleEventTickFunctionBindings.h
@@ -0,0 +1,36 @@
+﻿// Copyright Epic Games, Inc. All Rights Reserved.
+
+#pragma once
+
+#include "CoreMinimal.h"
+#include "Templates/Function.h"
+
+struct FTickFunction;
+struct FAnimNextModuleInstance;
+
+namespace UE::UAF
+{
+
+// Binding context for tick functions
+struct FTickFunctionBindingContext
+{
+	FTickFunctionBindingContext(FAnimNextModuleInstance& InModuleInstance, UObject* InObject, UWorld* InWorld)
+		: ModuleInstance(InModuleInstance)
+		, Object(InObject)
+		, World(InWorld)
+	{}
+
+	// The module instance
+	FAnimNextModuleInstance& ModuleInstance;
+
+	// Object that a module is bound on (e.g. a component)
+	UObject* Object = nullptr;
+
+	// World that the object is contained within
+	UWorld* World = nullptr;
+};
+
+// Function called to bind a tick function to a module event, set up prerequisites etc.
+using FModuleEventBindingFunction = TFunction<void(const FTickFunctionBindingContext& InContext, FTickFunction& InTickFunction)>;
+
+}
diff --git a/Engine/Plugins/Experimental/UAF/UAF/Source/UAF/Private/Component/AnimNextComponent.cpp b/Engine/Plugins/Experimental/UAF/UAF/Source/UAF/Private/Component/AnimNextComponent.cpp
new file mode 100644
--- /dev/null
+++ b/Engine/Plugins/Experimental/UAF/UAF/Source/UAF/Private/Component/AnimNextComponent.cpp
@@ -0,0 +1,469 @@
+// Copyright Epic Games, Inc. All Rights Reserved.
+
+#include "Component/AnimNextComponent.h"
+
+#include "AnimNextComponentWorldSubsystem.h"
+#include "Blueprint/BlueprintExceptionInfo.h"
+#include "Engine/World.h"
+#include "Module/ModuleTaskContext.h"
+#include "Variables/AnimNextVariableReference.h"
+
+void UAnimNextComponent::OnRegister()
+{
+	using namespace UE::UAF;
+
+	Super::OnRegister();
+
+	Subsystem = GetWorld()->GetSubsystem<UAnimNextComponentWorldSubsystem>();
+
+	if (Subsystem && Module)
+	{
+		check(!ModuleHandle.IsValid());
+
+		Subsystem->Register(this);
+#if UE_ENABLE_DEBUG_DRAWING
+		Subsystem->ShowDebugDrawing(this, bShowDebugDrawing);
+#endif
+	}
+}
+
+void UAnimNextComponent::OnUnregister()
+{
+	Super::OnUnregister();
+
+	if(Subsystem)
+	{
+		Subsystem->Unregister(this);
+		Subsystem = nullptr;
+	}
+}
+
+void UAnimNextComponent::BeginPlay()
+{
+	Super::BeginPlay();
+
+	UWorld* World = GetWorld();
+	check(World);
+
+	if (InitMethod == EAnimNextModuleInitMethod::InitializeAndRun
+	|| (InitMethod == EAnimNextModuleInitMethod::InitializeAndPauseInEditor && World->WorldType != EWorldType::Editor && World->WorldType != EWorldType::EditorPreview))
+	{
+		SetEnabled(true);
+	}
+}
+
+void UAnimNextComponent::EndPlay(const EEndPlayReason::Type EndPlayReason)
+{
+	Super::EndPlay(EndPlayReason);
+
+	SetEnabled(false);
+}
+
+void UAnimNextComponent::BlueprintSetVariable(FName Name, int32 Value)
+{
+	checkNoEntry();
+}
+
+DEFINE_FUNCTION(UAnimNextComponent::execBlueprintSetVariable)
+{
+	using namespace UE::UAF;
+
+	// Read wildcard Value input.
+	Stack.MostRecentPropertyAddress = nullptr;
+	Stack.MostRecentPropertyContainer = nullptr;
+
+	P_GET_PROPERTY(FNameProperty, Name);
+
+	Stack.StepCompiledIn<FProperty>(nullptr);
+	const FProperty* ValueProp = CastField<FProperty>(Stack.MostRecentProperty);
+	const void* ContainerPtr = Stack.MostRecentPropertyContainer;
+
+	P_FINISH;
+
+	if (!ValueProp || !ContainerPtr)
+	{
+		FBlueprintExceptionInfo ExceptionInfo(
+			EBlueprintExceptionType::AbortExecution,
+			NSLOCTEXT("UAFComponent", "UAFComponent_SetVariableError", "Failed to resolve the Value for Set Variable")
+		);
+
+		FBlueprintCoreDelegates::ThrowScriptException(P_THIS, Stack, ExceptionInfo);
+		return;
+	}
+
+	if (Name == NAME_None)
+	{
+		FBlueprintExceptionInfo ExceptionInfo(
+			EBlueprintExceptionType::NonFatalError,
+			NSLOCTEXT("UAFComponent", "UAFComponent_SetVariableInvalidNameWarning", "Invalid variable name supplied to Set Variable")
+		);
+
+		FBlueprintCoreDelegates::ThrowScriptException(P_THIS, Stack, ExceptionInfo);
+		return;
+	}
+
+	EPropertyBagResult Result = EPropertyBagResult::Success;
+
+	P_NATIVE_BEGIN;
+
+	{
+		FAnimNextParamType Type = FAnimNextParamType::FromProperty(ValueProp);
+		const uint8* ValuePtr = ValueProp->ContainerPtrToValuePtr<uint8>(ContainerPtr);
+		PRAGMA_DISABLE_DEPRECATION_WARNINGS
+		Result = P_THIS->SetVariableInternal(FAnimNextVariableReference(Name), Type, TConstArrayView<uint8>(ValuePtr, ValueProp->GetElementSize()));
+		PRAGMA_ENABLE_DEPRECATION_WARNINGS
+	}
+
+	P_NATIVE_END;
+
+	switch (Result)
+	{
+	case EPropertyBagResult::Success:
+		break;
+	case EPropertyBagResult::TypeMismatch:
+		{
+			FBlueprintExceptionInfo ExceptionInfo(
+				EBlueprintExceptionType::NonFatalError,
+				FText::Format(NSLOCTEXT("UAFComponent", "UAFComponent_SetVariableTypeMismatch", "Incompatible type supplied for variable '{0}'"), FText::FromName(Name))
+			);
+
+			FBlueprintCoreDelegates::ThrowScriptException(P_THIS, Stack, ExceptionInfo);
+			break;
+		}
+	case EPropertyBagResult::PropertyNotFound:
+		{
+			FBlueprintExceptionInfo ExceptionInfo(
+				EBlueprintExceptionType::NonFatalError,
+				FText::Format(NSLOCTEXT("UAFComponent", "UAFComponent_SetVariableNotFoundWarning", "Unknown variable name '{0}' supplied to Set Variable"), FText::FromName(Name))
+			);
+
+			FBlueprintCoreDelegates::ThrowScriptException(P_THIS, Stack, ExceptionInfo);
+			break;
+		}
+	case EPropertyBagResult::OutOfBounds:
+	case EPropertyBagResult::DuplicatedValue:
+		{
+			FBlueprintExceptionInfo ExceptionInfo(
+				EBlueprintExceptionType::NonFatalError,
+				FText::Format(NSLOCTEXT("UAFComponent", "UAFComponent_SetVariableUnknownWarning", "Unexpected internal error when calling Set Variable for '{0}'"), FText::FromName(Name))
+			);
+
+			FBlueprintCoreDelegates::ThrowScriptException(P_THIS, Stack, ExceptionInfo);
+			break;
+		}
+	}
+}
+
+void UAnimNextComponent::BlueprintSetVariableReference(const FAnimNextVariableReference& Variable, const int32& Value)
+{
+	checkNoEntry();
+}
+
+DEFINE_FUNCTION(UAnimNextComponent::execBlueprintSetVariableReference)
+{
+	using namespace UE::UAF;
+
+	P_GET_STRUCT_REF(FAnimNextVariableReference, Variable);
+
+	const FProperty* VariableProperty = Variable.ResolveProperty();
+	if (VariableProperty == nullptr)
+	{
+		P_FINISH;
+
+		FBlueprintExceptionInfo ExceptionInfo(
+			EBlueprintExceptionType::AbortExecution,
+			FText::Format(NSLOCTEXT("UAFComponent", "UAFComponent_SetVariablePropertyError", "Failed to resolve the property {0} of variable reference"), FText::FromName(Variable.GetName()))
+		);
+
+		FBlueprintCoreDelegates::ThrowScriptException(P_THIS, Stack, ExceptionInfo);
+		return;
+	}
+
+	if (Variable.GetName() == NAME_None || Variable.GetObject() == nullptr)
+	{
+		P_FINISH;
+
+		FBlueprintExceptionInfo ExceptionInfo(
+			EBlueprintExceptionType::NonFatalError,
+			NSLOCTEXT("UAFComponent", "UAFComponent_SetVariableInvalidWarning", "Invalid variable supplied to Set Variable")
+		);
+
+		FBlueprintCoreDelegates::ThrowScriptException(P_THIS, Stack, ExceptionInfo);
+		return;
+	}
+
+	// Read wildcard Value input.
+	Stack.MostRecentPropertyAddress = nullptr;
+	Stack.MostRecentPropertyContainer = nullptr;
+
+	EPropertyBagResult Result = EPropertyBagResult::Success;
+	{
+		uint8* TempStorage = (uint8*)FMemory_Alloca_Aligned(VariableProperty->GetElementSize(), VariableProperty->GetMinAlignment());
+		VariableProperty->InitializeValue(TempStorage);
+
+		Stack.StepCompiledIn(TempStorage, VariableProperty->GetClass());
+
+		P_FINISH;
+		
+		P_NATIVE_BEGIN;
+
+		{
+			FAnimNextParamType Type = FAnimNextParamType::FromProperty(VariableProperty);
+			Result = P_THIS->SetVariableInternal(Variable, Type, TConstArrayView<uint8>(TempStorage, VariableProperty->GetElementSize()));
+		}
+
+		P_NATIVE_END;
+
+		VariableProperty->DestroyValue(TempStorage);
+	}
+
+	switch (Result)
+	{
+	case EPropertyBagResult::Success:
+		break;
+	case EPropertyBagResult::TypeMismatch:
+		{
+			FBlueprintExceptionInfo ExceptionInfo(
+				EBlueprintExceptionType::NonFatalError,
+				FText::Format(NSLOCTEXT("UAFComponent", "UAFComponent_SetVariableTypeMismatch", "Incompatible type supplied for variable '{0}'"), FText::FromName(Variable.GetName()))
+			);
+
+			FBlueprintCoreDelegates::ThrowScriptException(P_THIS, Stack, ExceptionInfo);
+			break;
+		}
+	case EPropertyBagResult::PropertyNotFound:
+		{
+			FBlueprintExceptionInfo ExceptionInfo(
+				EBlueprintExceptionType::NonFatalError,
+				FText::Format(NSLOCTEXT("UAFComponent", "UAFComponent_SetVariableNotFoundWarning", "Unknown variable name '{0}' supplied to Set Variable"), FText::FromName(Variable.GetName()))
+			);
+
+			FBlueprintCoreDelegates::ThrowScriptException(P_THIS, Stack, ExceptionInfo);
+			break;
+		}
+	case EPropertyBagResult::OutOfBounds:
+	case EPropertyBagResult::DuplicatedValue:
+		{
+			FBlueprintExceptionInfo ExceptionInfo(
+				EBlueprintExceptionType::NonFatalError,
+				FText::Format(NSLOCTEXT("UAFComponent", "UAFComponent_SetVariableUnknownWarning", "Unexpected internal error when calling Set Variable for '{0}'"), FText::FromName(Variable.GetName()))
+			);
+
+			FBlueprintCoreDelegates::ThrowScriptException(P_THIS, Stack, ExceptionInfo);
+			break;
+		}
+	}
+}
+
+EPropertyBagResult UAnimNextComponent::SetVariableInternal(const FAnimNextVariableReference& InVariable, const FAnimNextParamType& InType, TConstArrayView<uint8> InNewValue)
+{
+	if (Subsystem)
+	{
+		return Subsystem->SetVariable(this, InVariable, InType, InNewValue);
+	}
+	return EPropertyBagResult::PropertyNotFound;
+}
+
+EPropertyBagResult UAnimNextComponent::WriteVariableInternal(const FAnimNextVariableReference& InVariable, const FAnimNextParamType& InType, TFunctionRef<void(TArrayView<uint8>)> InFunction)
+{
+	if (Subsystem)
+	{
+		return Subsystem->WriteVariable(this, InVariable, InType, InFunction);
+	}
+	return EPropertyBagResult::PropertyNotFound;
+}
+
+bool UAnimNextComponent::IsEnabled() const
+{
+	if (Subsystem)
+	{
+		return Subsystem->IsEnabled(this);
+	}
+	return false;
+}
+
+void UAnimNextComponent::SetEnabled(bool bEnabled)
+{
+	if(Subsystem)
+	{
+		Subsystem->SetEnabled(this, bEnabled);
+	}
+}
+
+void UAnimNextComponent::ShowDebugDrawing(bool bInShowDebugDrawing)
+{
+#if UE_ENABLE_DEBUG_DRAWING
+	bShowDebugDrawing = bInShowDebugDrawing;
+	if(Subsystem)
+	{
+		Subsystem->ShowDebugDrawing(this, bInShowDebugDrawing);
+	}
+#endif
+}
+
+void UAnimNextComponent::QueueTask(FName InModuleEventName, TUniqueFunction<void(const UE::UAF::FModuleTaskContext&)>&& InTaskFunction, UE::UAF::ETaskRunLocation InLocation)
+{
+	using namespace UE::UAF;
+
+	if(Subsystem)
+	{
+		Subsystem->QueueTask(this, InModuleEventName, MoveTemp(InTaskFunction), InLocation);
+	}
+}
+
+void UAnimNextComponent::QueueInputTraitEvent(FAnimNextTraitEventPtr Event)
+{
+	using namespace UE::UAF;
+
+	if(Subsystem)
+	{
+		Subsystem->QueueInputTraitEvent(this, Event);
+	}
+}
+
+const FTickFunction* UAnimNextComponent::FindTickFunction(FName InEventName) const
+{
+	if(Subsystem)
+	{
+		return Subsystem->FindTickFunction(this, InEventName);
+	}
+	return nullptr;
+}
+
+void UAnimNextComponent::AddPrerequisite(UObject* InObject, FTickFunction& InTickFunction, FName InEventName)
+{
+	if(Subsystem)
+	{
+		Subsystem->AddDependency(this, InObject, InTickFunction, InEventName, UAnimNextWorldSubsystem::EDependency::Prerequisite);
+	}
+}
+
+void UAnimNextComponent::AddComponentPrerequisite(UActorComponent* InComponent, FName InEventName)
+{
+	if (InComponent)
+	{
+		AddPrerequisite(InComponent, InComponent->PrimaryComponentTick, InEventName);
+	}
+}
+
+void UAnimNextComponent::AddSubsequent(UObject* InObject, FTickFunction& InTickFunction, FName InEventName)
+{
+	if(Subsystem)
+	{
+		Subsystem->AddDependency(this, InObject, InTickFunction, InEventName, UAnimNextWorldSubsystem::EDependency::Subsequent);
+	}
+}
+
+void UAnimNextComponent::AddComponentSubsequent(UActorComponent* InComponent, FName InEventName)
+{
+	if (InComponent)
+	{
+		AddSubsequent(InComponent, InComponent->PrimaryComponentTick, InEventName);
+	}
+}
+
+void UAnimNextComponent::RemovePrerequisite(UObject* InObject, FTickFunction& InTickFunction, FName InEventName)
+{
+	if(Subsystem)
+	{
+		Subsystem->RemoveDependency(this, InObject, InTickFunction, InEventName, UAnimNextWorldSubsystem::EDependency::Prerequisite);
+	}
+}
+
+void UAnimNextComponent::RemoveComponentPrerequisite(UActorComponent* InComponent, FName InEventName)
+{
+	if (InComponent)
+	{
+		RemovePrerequisite(InComponent, InComponent->PrimaryComponentTick, InEventName);
+	}
+}
+
+void UAnimNextComponent::RemoveSubsequent(UObject* InObject, FTickFunction& InTickFunction, FName InEventName)
+{
+	if(Subsystem)
+	{
+		Subsystem->RemoveDependency(this, InObject, InTickFunction, InEventName, UAnimNextWorldSubsystem::EDependency::Subsequent);
+	}
+}
+
+void UAnimNextComponent::RemoveComponentSubsequent(UActorComponent* InComponent, FName InEventName)
+{
+	if (InComponent)
+	{
+		RemoveSubsequent(InComponent, InComponent->PrimaryComponentTick, InEventName);
+	}
+}
+
+void UAnimNextComponent::AddModuleEventPrerequisite(FName InEventName, UAnimNextComponent* OtherAnimNextComponent, FName OtherEventName)
+{
+	if (!OtherAnimNextComponent)
+	{
+		UE_LOG(LogAnimation, Warning, TEXT("UAFComponent::AddModuleEventPrerequisite called with null OtherAnimNextComponent"));
+	}
+	else if (OtherAnimNextComponent == this)
+	{
+		UE_LOG(LogAnimation, Warning, TEXT("UAFComponent::AddModuleEventPrerequisite called using the same component"));
+	}
+	else if (Subsystem)
+	{
+		Subsystem->AddModuleEventDependency(this, InEventName, OtherAnimNextComponent, OtherEventName, UAnimNextWorldSubsystem::EDependency::Prerequisite);
+	}
+}
+
+void UAnimNextComponent::AddModuleEventSubsequent(FName InEventName, UAnimNextComponent* OtherAnimNextComponent, FName OtherEventName)
+{
+	if (!OtherAnimNextComponent)
+	{
+		UE_LOG(LogAnimation, Warning, TEXT("UAFComponent::AddModuleEventSubsequent called with null OtherAnimNextComponent"));
+	}
+	else if (OtherAnimNextComponent == this)
+	{
+		UE_LOG(LogAnimation, Warning, TEXT("UAFComponent::AddModuleEventPrerequisite called using the same component"));
+	}
+	else if (Subsystem)
+	{
+		Subsystem->AddModuleEventDependency(this, InEventName, OtherAnimNextComponent, OtherEventName, UAnimNextWorldSubsystem::EDependency::Subsequent);
+	}
+}
+
+void UAnimNextComponent::RemoveModuleEventPrerequisite(FName InEventName, UAnimNextComponent* OtherAnimNextComponent, FName OtherEventName)
+{
+	if (!OtherAnimNextComponent)
+	{
+		UE_LOG(LogAnimation, Warning, TEXT("UAnimNextComponent::RemoveModuleEventPrerequisite called with null OtherAnimNextComponent"));
+	}
+	else if (OtherAnimNextComponent == this)
+	{
+		UE_LOG(LogAnimation, Warning, TEXT("UAnimNextComponent::AddModuleEventPrerequisite called using the same component"));
+	}
+	else if (Subsystem)
+	{
+		Subsystem->RemoveModuleEventDependency(this, InEventName, OtherAnimNextComponent, OtherEventName, UAnimNextWorldSubsystem::EDependency::Prerequisite);
+	}
+}
+
+void UAnimNextComponent::RemoveModuleEventSubsequent(FName InEventName, UAnimNextComponent* OtherAnimNextComponent, FName OtherEventName)
+{
+	if (!OtherAnimNextComponent)
+	{
+		UE_LOG(LogAnimation, Warning, TEXT("UAnimNextComponent::RemoveModuleEventSubsequent called with null OtherAnimNextComponent"));
+	}
+	else if (OtherAnimNextComponent == this)
+	{
+		UE_LOG(LogAnimation, Warning, TEXT("UAnimNextComponent::AddModuleEventPrerequisite called using the same component"));
+	}
+	else if (Subsystem)
+	{
+		Subsystem->RemoveModuleEventDependency(this, InEventName, OtherAnimNextComponent, OtherEventName, UAnimNextWorldSubsystem::EDependency::Subsequent);
+	}
+}
+
+FAnimNextModuleHandle UAnimNextComponent::BlueprintGetModuleHandle() const
+{
+	return FAnimNextModuleHandle(ModuleHandle);
+}
+
+void UAnimNextComponent::SetModule(TObjectPtr<UAnimNextModule> InModule)
+{
+	Module = InModule;
+}
diff --git a/Engine/Plugins/Experimental/UAF/UAF/Source/UAF/Private/Graph/RigUnit_AnimNextRemapPose.cpp b/Engine/Plugins/Experimental/UAF/UAF/Source/UAF/Private/Graph/RigUnit_AnimNextRemapPose.cpp
new file mode 100644
--- /dev/null
+++ b/Engine/Plugins/Experimental/UAF/UAF/Source/UAF/Private/Graph/RigUnit_AnimNextRemapPose.cpp
@@ -0,0 +1,56 @@
+// Copyright Epic Games, Inc. All Rights Reserved.
+
+#include "RigUnit_AnimNextRemapPose.h"
+
+#include "GenerationTools.h"
+#include "Components/SkeletalMeshComponent.h"
+#include "Engine/SkeletalMesh.h"
+
+#include UE_INLINE_GENERATED_CPP_BY_NAME(RigUnit_AnimNextRemapPose)
+DEFINE_STAT(STAT_AnimNext_RigUnit_RemapPose);
+
+FRigUnit_AnimNextRemapPose_Execute()
+{
+	SCOPE_CYCLE_COUNTER(STAT_AnimNext_RigUnit_RemapPose);
+
+	using namespace UE::UAF;
+
+	if(!TargetAnimGraphRefPose.ReferencePose.IsValid())
+	{
+		return;
+	}
+	
+	if(!Pose.LODPose.IsValid())
+	{
+		return;
+	}
+
+	FMemMark MemMark(FMemStack::Get());
+
+	const UE::UAF::FLODPoseHeap& SourcePose = Pose.LODPose;
+	const UE::UAF::FReferencePose& SourceRefPose = SourcePose.GetRefPose();
+	const USkeletalMesh* SourceMesh = SourceRefPose.SkeletalMesh.Get();
+
+	const UE::UAF::FReferencePose& TargetRefPose = TargetAnimGraphRefPose.ReferencePose.GetRef<UE::UAF::FReferencePose>();
+	UE::UAF::FLODPoseHeap& TargetPose = Result.LODPose;
+	const USkeletalMesh* TargetMesh = TargetRefPose.SkeletalMesh.Get();
+
+	if (SourceMesh == TargetMesh)
+	{
+		TargetPose.PrepareForLOD(TargetRefPose, SourcePose.LODLevel, /*bSetRefPose=*/false, SourcePose.IsAdditive());
+		// Just copy the pose to the target in case there is nothing to convert.
+		TargetPose.CopyFrom(SourcePose);
+	}
+	else
+	{
+		// Pre-create the mapping and cache it to avoid runtime lookups.
+		if (RemapPoseData.ShouldReinit(SourceRefPose, TargetRefPose))
+		{
+			RemapPoseData.Reinit(SourceRefPose, TargetRefPose);
+		}
+
+		Result.Curves.CopyFrom(Pose.Curves);
+		RemapPoseData.RemapPose(SourcePose, TargetPose);
+		RemapPoseData.RemapAttributes(Pose.LODPose, Pose.Attributes, Result.LODPose, Result.Attributes);
+	}
+}
\ No newline at end of file
diff --git a/Engine/Plugins/Experimental/UAF/UAF/Source/UAF/Private/Module/AnimNextActorComponentReferenceComponent.cpp b/Engine/Plugins/Experimental/UAF/UAF/Source/UAF/Private/Module/AnimNextActorComponentReferenceComponent.cpp
new file mode 100644
--- /dev/null
+++ b/Engine/Plugins/Experimental/UAF/UAF/Source/UAF/Private/Module/AnimNextActorComponentReferenceComponent.cpp
@@ -0,0 +1,44 @@
+// Copyright Epic Games, Inc. All Rights Reserved.
+
+#include "Module/AnimNextActorComponentReferenceComponent.h"
+
+#include "Component/AnimNextComponent.h"
+#include "Module/AnimNextModuleInstance.h"
+
+void FAnimNextActorComponentReferenceComponent::OnInitializeHelper(UScriptStruct* InScriptStruct)
+{
+	check(InScriptStruct->IsChildOf(FAnimNextActorComponentReferenceComponent::StaticStruct()));
+
+	FAnimNextModuleInstance* ModuleInstance = GetModuleInstancePtr();
+	if (ModuleInstance == nullptr)
+	{
+		return;
+	}
+
+	UAnimNextComponent* AnimNextComponent = Cast<UAnimNextComponent>(ModuleInstance->GetObject());
+
+	FAnimNextModuleInstance::RunTaskOnGameThread(
+		[WeakAnimNextComponent = TWeakObjectPtr<UAnimNextComponent>(AnimNextComponent), ComponentType = ComponentType, InScriptStruct]()
+		{
+			UAnimNextComponent* PinnedAnimNextComponent = WeakAnimNextComponent.Get();
+			if(PinnedAnimNextComponent == nullptr)
+			{
+				return;
+			}
+
+			AActor* Owner = PinnedAnimNextComponent->GetOwner();
+			if(Owner == nullptr)
+			{
+				return;
+			}
+
+			UActorComponent* FoundComponent = Owner->FindComponentByClass(ComponentType);
+			PinnedAnimNextComponent->QueueTask(NAME_None, [FoundComponent, InScriptStruct](const UE::UAF::FModuleTaskContext& InContext)
+			{
+				InContext.TryAccessComponent(InScriptStruct, [FoundComponent](FUAFModuleInstanceComponent& InComponent)
+				{
+					static_cast<FAnimNextActorComponentReferenceComponent&>(InComponent).Component = FoundComponent;
+				});
+			});
+		});
+}
diff --git a/Engine/Plugins/Experimental/UAF/UAF/Source/UAF/Private/Module/AnimNextModuleContextData.cpp b/Engine/Plugins/Experimental/UAF/UAF/Source/UAF/Private/Module/AnimNextModuleContextData.cpp
new file mode 100644
--- /dev/null
+++ b/Engine/Plugins/Experimental/UAF/UAF/Source/UAF/Private/Module/AnimNextModuleContextData.cpp
@@ -0,0 +1,21 @@
+﻿// Copyright Epic Games, Inc. All Rights Reserved.
+
+#include "Module/AnimNextModuleContextData.h"
+#include "Module/AnimNextModuleInstance.h"
+
+FAnimNextModuleContextData::FAnimNextModuleContextData(FAnimNextModuleInstance* InModuleInstance)
+	: ModuleInstance(InModuleInstance)
+	, Instance(InModuleInstance)
+{
+}
+
+FAnimNextModuleContextData::FAnimNextModuleContextData(FAnimNextModuleInstance* InModuleInstance, const FUAFAssetInstance* InInstance)
+	: ModuleInstance(InModuleInstance)
+	, Instance(InInstance)
+{
+}
+
+UObject* FAnimNextModuleContextData::GetObject() const
+{
+	return ModuleInstance ? ModuleInstance->GetObject() : nullptr;
+}
\ No newline at end of file
