UAF: Moved/renamed AnimNext and AnimNextAnimGraph plugins

diff --git a/Engine/Plugins/Experimental/UAF/UAF/Source/UAF/Private/Module/AnimNextModuleInstance.cpp b/Engine/Plugins/Experimental/UAF/UAF/Source/UAF/Private/Module/AnimNextModuleInstance.cpp
new file mode 100644
--- /dev/null
+++ b/Engine/Plugins/Experimental/UAF/UAF/Source/UAF/Private/Module/AnimNextModuleInstance.cpp
@@ -0,0 +1,717 @@
+// Copyright Epic Games, Inc. All Rights Reserved.
+
+#include "Module/AnimNextModuleInstance.h"
+
+#if UE_ENABLE_DEBUG_DRAWING
+#include "AnimNextDebugDraw.h"
+#endif
+#include "AnimNextPool.h"
+#include "Engine/World.h"
+#include "AnimNextStats.h"
+#include "Async/TaskGraphInterfaces.h"
+#include "SceneInterface.h"
+#include "UAFRigVMComponent.h"
+#include "Algo/TopologicalSort.h"
+#include "Logging/StructuredLog.h"
+#include "Misc/EnumerateRange.h"
+#include "Module/RigUnit_AnimNextModuleEvents.h"
+#include "Module/ModuleTickFunction.h"
+#include "UObject/UObjectIterator.h"
+#include "RewindDebugger/AnimNextTrace.h"
+#include "Variables/AnimNextVariableReference.h"
+
+DEFINE_STAT(STAT_AnimNext_InitializeInstance);
+
+FAnimNextModuleInstance::FAnimNextModuleInstance() = default;
+
+FAnimNextModuleInstance::FAnimNextModuleInstance(
+		UAnimNextModule* InModule,
+		UObject* InObject,
+		UE::UAF::TPool<FAnimNextModuleInstance>* InPool,
+		EAnimNextModuleInitMethod InInitMethod)
+	: Object(InObject)
+	, Pool(InPool)
+	, InitState(EInitState::NotInitialized)
+	, RunState(ERunState::NotInitialized)
+	, InitMethod(InInitMethod)
+{
+	check(InModule);
+	check(InObject);
+
+	Asset = InModule;
+
+#if UE_ENABLE_DEBUG_DRAWING
+	if(Object && Object->GetWorld())
+	{
+		DebugDraw = MakeUnique<UE::UAF::Debug::FDebugDraw>(Object);
+	}
+#endif
+}
+
+FAnimNextModuleInstance::~FAnimNextModuleInstance()
+{
+	ResetBindingsAndInstanceData();
+
+#if UE_ENABLE_DEBUG_DRAWING
+	DebugDraw.Reset();
+#endif
+
+	Object = nullptr;
+	Asset = nullptr;
+	Handle.Reset();
+}
+
+namespace UE::UAF::Private
+{
+
+struct FImplementedModuleEvent
+{
+	UScriptStruct* Struct = nullptr;
+	FModuleEventBindingFunction Binding;
+	FName EventName;
+	EModuleEventPhase Phase = EModuleEventPhase::Execute;
+	ETickingGroup TickGroup = ETickingGroup::TG_PrePhysics;
+	int32 SortOrder = 0;
+	bool bUserEvent = false;
+	bool bIsTask = false;
+	bool bIsGameThreadTask = false;
+};
+
+static TArray<FImplementedModuleEvent> GImplementedModuleEvents;
+
+// Gets information about the module events that are implemented by the supplied VM, sorted by execution order in the frame
+static TConstArrayView<FImplementedModuleEvent> GetImplementedModuleEvents(const URigVM* VM)
+{
+	check(IsInGameThread());	// This function cannot be run concurrently because of static usage
+	GImplementedModuleEvents.Reset();
+
+	const FRigVMByteCode& ByteCode = VM->GetByteCode();
+	const TArray<const FRigVMFunction*>& Functions = VM->GetFunctions();
+	FRigVMInstructionArray Instructions = ByteCode.GetInstructions();
+	for (int32 EntryIndex = 0; EntryIndex < ByteCode.NumEntries(); EntryIndex++)
+	{
+		const FRigVMByteCodeEntry& Entry = ByteCode.GetEntry(EntryIndex);
+		const FRigVMInstruction& Instruction = Instructions[Entry.InstructionIndex];
+		const FRigVMExecuteOp& Op = ByteCode.GetOpAt<FRigVMExecuteOp>(Instruction);
+		const FRigVMFunction* Function = Functions[Op.FunctionIndex];
+		check(Function != nullptr);
+
+		if (Function->Struct->IsChildOf(FRigUnit_AnimNextModuleEventBase::StaticStruct()))
+		{
+			TInstancedStruct<FRigUnit_AnimNextModuleEventBase> StructInstance;
+			StructInstance.InitializeAsScriptStruct(Function->Struct);
+			const FRigUnit_AnimNextModuleEventBase& Event = StructInstance.Get();
+			FImplementedModuleEvent& NewEvent = GImplementedModuleEvents.AddDefaulted_GetRef();
+			NewEvent.Struct = Function->Struct;
+			NewEvent.Binding = Event.GetBindingFunction();
+			NewEvent.EventName = Event.GetEventName();
+			NewEvent.Phase = Event.GetEventPhase();
+			NewEvent.TickGroup = Event.GetTickGroup();
+			NewEvent.SortOrder = Event.GetSortOrder();
+			NewEvent.bUserEvent = Event.IsUserEvent();
+			NewEvent.bIsTask = Event.IsTask();
+			NewEvent.bIsGameThreadTask = Event.IsGameThreadTask();
+
+			// User events can override their event name etc. via parameters 
+			if (Function->Struct->IsChildOf(FRigUnit_AnimNextUserEvent::StaticStruct()))
+			{
+				// Pull the values out of the literal memory
+				FRigVMOperandArray Operands = ByteCode.GetOperandsForExecuteOp(Instruction);
+				check(Function->ArgumentNames.Num() == Operands.Num());
+				int32 NumOperands = Operands.Num();
+				URigVM* NonConstVM = const_cast<URigVM*>(VM);
+				for (int32 OperandIndex = 0; OperandIndex < NumOperands; ++OperandIndex)
+				{
+					const FRigVMOperand& Operand = Operands[OperandIndex];
+					FName OperandName = Function->ArgumentNames[OperandIndex];
+					if (OperandName == GET_MEMBER_NAME_CHECKED(FRigUnit_AnimNextUserEvent, Name))
+					{
+						check(Operand.GetMemoryType() == ERigVMMemoryType::Literal);
+						NewEvent.EventName = *NonConstVM->LiteralMemoryStorage.GetData<FName>(Operand.GetRegisterIndex());
+					}
+					else if (OperandName == GET_MEMBER_NAME_CHECKED(FRigUnit_AnimNextUserEvent, SortOrder))
+					{
+						check(Operand.GetMemoryType() == ERigVMMemoryType::Literal);
+						NewEvent.SortOrder = *NonConstVM->LiteralMemoryStorage.GetData<int32>(Operand.GetRegisterIndex());
+					}
+				}
+			}
+		}
+	}
+
+	Algo::Sort(GImplementedModuleEvents, [](const FImplementedModuleEvent& InA, const FImplementedModuleEvent& InB)
+	{
+		if (InA.Phase != InB.Phase)
+		{
+			return InA.Phase < InB.Phase;
+		}
+		else if (InA.TickGroup != InB.TickGroup)
+		{
+			return InA.TickGroup < InB.TickGroup;
+		}
+		else if (InA.SortOrder != InB.SortOrder)
+		{
+			return InA.SortOrder < InB.SortOrder;
+		}
+
+		// Tie-break sorting on event name for determinism
+		return InA.EventName.Compare(InB.EventName) < 0;
+	});
+	
+	return GImplementedModuleEvents;
+}
+
+}
+
+void FAnimNextModuleInstance::Initialize()
+{
+	SCOPE_CYCLE_COUNTER(STAT_AnimNext_InitializeInstance);
+	
+	using namespace UE::UAF;
+
+	check(IsInGameThread());
+
+	check(Object)
+	const UAnimNextModule* Module = GetModule();
+	check(Module);
+
+	UWorld* World = Object->GetWorld();
+	if(World)
+	{
+		WorldType = World->WorldType;
+	}
+
+	// Get all the module events from the VM entry points, sorted by ordering in the frame
+	URigVM* VM = Module->RigVM;
+	TConstArrayView<Private::FImplementedModuleEvent> ImplementedModuleEvents = Private::GetImplementedModuleEvents(VM);
+
+	// Setup tick function graph using module events
+	if (ImplementedModuleEvents.Num() > 0)
+	{
+		TransitionToInitState(EInitState::CreatingTasks);
+
+		// Allocate tick functions
+		TickFunctions.Reserve(ImplementedModuleEvents.Num());
+		bool bFoundFirstUserEvent = false;
+		FModuleEventTickFunction* PrevTickFunction = nullptr;
+		for (int32 EventIndex = 0; EventIndex < ImplementedModuleEvents.Num(); EventIndex++)
+		{
+			const Private::FImplementedModuleEvent& ModuleEvent = ImplementedModuleEvents[EventIndex];
+			if (!ModuleEvent.bIsTask)
+			{
+				continue;
+			}
+
+			FModuleEventTickFunction& TickFunction = TickFunctions.AddDefaulted_GetRef();
+			TickFunction.bRunOnAnyThread = !ModuleEvent.bIsGameThreadTask;
+			TickFunction.ModuleInstance = this;
+			TickFunction.EventName = ModuleEvent.EventName;
+			TickFunction.TickGroup = ModuleEvent.TickGroup;
+			TickFunction.bUserEvent = ModuleEvent.bUserEvent;
+
+			// Perform custom setup
+			FTickFunctionBindingContext Context(*this, Object, World);
+			ModuleEvent.Binding(Context, TickFunction);
+
+			// Establish linear dependency chain
+			if (PrevTickFunction != nullptr)
+			{
+				TickFunction.AddPrerequisite(Object, *PrevTickFunction);
+			}
+			PrevTickFunction = &TickFunction;
+
+			// Set up dependencies, if any
+			for (const TInstancedStruct<FRigVMTrait_ModuleEventDependency>& DependencyInstance : Module->Dependencies)
+			{
+				const FRigVMTrait_ModuleEventDependency* Dependency = DependencyInstance.GetPtr<FRigVMTrait_ModuleEventDependency>();
+				if (Dependency != nullptr && Dependency->EventName == ModuleEvent.EventName)
+				{
+					FModuleDependencyContext ModuleDependencyContext(Object, TickFunction);
+					Dependency->OnAddDependency(ModuleDependencyContext);
+				}
+			}
+
+			if (ModuleEvent.bUserEvent && !ModuleEvent.bIsGameThreadTask)
+			{
+				// Hook up with Task Sync Manager if possible
+				TickFunction.InitializeBatchedWork(World);
+			}
+
+			if (ModuleEvent.bUserEvent && !bFoundFirstUserEvent)
+			{
+				TickFunction.bFirstUserEvent = true;
+				bFoundFirstUserEvent = true;
+
+				// Set this first user event to run the bindings event, if it exists
+				auto IsExecuteBindingsEvent = [](const Private::FImplementedModuleEvent& InEvent)
+				{
+					return InEvent.EventName == FRigUnit_AnimNextExecuteBindings_WT::EventName;
+				};
+
+				TickFunction.bRunBindingsEvent = ImplementedModuleEvents.ContainsByPredicate(IsExecuteBindingsEvent);
+			}
+		}
+
+		// Find the last user event - 'end' logic will be called from here
+		for (int32 EventIndex = TickFunctions.Num() - 1; EventIndex >= 0; EventIndex--)
+		{
+			FModuleEventTickFunction& TickFunction = TickFunctions[EventIndex];
+			if (TickFunction.bUserEvent)
+			{
+				TickFunction.bLastUserEvent = true;
+				break;
+			}
+		}
+
+		TransitionToInitState(EInitState::BindingTasks);
+
+		// Register our tick functions
+		if(World)
+		{
+			ULevel* Level = World->PersistentLevel;
+			for (FModuleEventTickFunction& TickFunction : TickFunctions)
+			{
+				TickFunction.RegisterTickFunction(Level);
+			}
+		}
+
+		TransitionToInitState(EInitState::SetupVariables);
+
+		// TODO: code in EInitState::SetupVariables phase below can probably move to FModuleEventTickFunction::Initialize
+
+		// Initialize variables
+#if WITH_EDITOR
+		if(bIsRecreatingOnCompile && Variables.bHasBeenInitialized)
+		{
+			MigrateVariables();
+		}
+		else
+#endif
+		{
+			InitializeVariables();
+		}
+
+		ProxyVariables[0].Initialize(Variables);
+		ProxyVariables[1].Initialize(Variables);
+
+		// Allocate compiled-in module components
+		{
+			FUAFAssetInstanceComponent::FScopedConstructorHelper ConstructorHelper(*this);
+			for(const UScriptStruct* ComponentStruct : Module->RequiredComponents)
+			{
+				TInstancedStruct<FUAFModuleInstanceComponent> InstancedStruct(ComponentStruct);
+				ComponentMap.Add(ComponentStruct, MoveTemp(InstancedStruct));
+			}
+		}
+
+		TransitionToInitState(EInitState::PendingInitializeEvent);
+		TransitionToRunState(ERunState::Running);
+
+		// Just pause now if we arent needing an initial update
+		if(InitMethod == EAnimNextModuleInitMethod::None)
+		{
+			Enable(false);
+		}
+#if WITH_EDITOR
+		else if(World)
+		{
+			// In editor worlds we run a linearized 'initial tick' to ensure we generate an initial output pose, as these worlds dont always tick
+			if( World->WorldType == EWorldType::Editor ||
+				World->WorldType == EWorldType::EditorPreview)
+			{
+				FModuleEventTickFunction::InitializeAndRunModule(*this);
+			}
+		}
+#endif
+	}
+}
+
+void FAnimNextModuleInstance::RemoveAllTickDependencies()
+{
+	using namespace UE::UAF;
+
+	check(IsInGameThread());
+
+	for (FModuleEventTickFunction& TickFunction : TickFunctions)
+	{
+		TickFunction.RemoveAllExternalSubsequents();
+	}
+}
+
+void FAnimNextModuleInstance::ResetBindingsAndInstanceData()
+{
+	using namespace UE::UAF;
+
+	check(IsInGameThread());
+
+	TransitionToInitState(EInitState::NotInitialized);
+	TransitionToRunState(ERunState::NotInitialized);
+
+#if WITH_EDITOR
+	if(!bIsRecreatingOnCompile)
+#endif
+	{
+		for (FModuleEventTickFunction& TickFunction : TickFunctions)
+		{
+			// We should have released all external dependencies by now via RemoveAllTickDependencies
+			check(TickFunction.ExternalSubsequents.Num() == 0);
+			TickFunction.UnRegisterTickFunction();
+		}
+		ProxyVariables[0].Reset();
+		ProxyVariables[1].Reset();
+	}
+
+	TickFunctions.Reset();
+
+	ReleaseComponents();
+}
+
+void FAnimNextModuleInstance::QueueInputTraitEvent(FAnimNextTraitEventPtr Event)
+{
+	InputEventList.Push(MoveTemp(Event));
+}
+
+void FAnimNextModuleInstance::QueueOutputTraitEvent(FAnimNextTraitEventPtr Event)
+{
+	OutputEventList.Push(MoveTemp(Event));
+}
+
+bool FAnimNextModuleInstance::IsEnabled() const
+{
+	using namespace UE::UAF;
+
+	check(IsInGameThread());
+
+	return RunState == ERunState::Running;
+}
+
+void FAnimNextModuleInstance::Enable(bool bInEnabled)
+{
+	using namespace UE::UAF;
+
+	check(IsInGameThread());
+
+	if(RunState == ERunState::Paused || RunState == ERunState::Running)
+	{
+		for (FModuleEventTickFunction& TickFunction : TickFunctions)
+		{
+			TickFunction.SetTickFunctionEnable(bInEnabled);
+		}
+
+		TransitionToRunState(bInEnabled ? ERunState::Running : ERunState::Paused);
+	}
+}
+
+void FAnimNextModuleInstance::TransitionToInitState(EInitState InNewState)
+{
+	switch(InNewState)
+	{
+	case EInitState::NotInitialized:
+		check(InitState == EInitState::NotInitialized || InitState == EInitState::PendingInitializeEvent || InitState == EInitState::SetupVariables || InitState == EInitState::FirstUpdate || InitState == EInitState::Initialized);
+		break;
+	case EInitState::CreatingTasks:
+		check(InitState == EInitState::NotInitialized);
+		break;
+	case EInitState::BindingTasks:
+		check(InitState == EInitState::CreatingTasks);
+		break;
+	case EInitState::SetupVariables:
+		check(InitState == EInitState::BindingTasks);
+		break;
+	case EInitState::PendingInitializeEvent:
+		check(InitState == EInitState::SetupVariables);
+		break;
+	case EInitState::FirstUpdate:
+		check(InitState == EInitState::PendingInitializeEvent);
+		break;
+	case EInitState::Initialized:
+		check(InitState == EInitState::FirstUpdate);
+		break;
+	default:
+		checkNoEntry();
+	}
+
+	InitState = InNewState;
+}
+
+void FAnimNextModuleInstance::TransitionToRunState(ERunState InNewState)
+{
+	switch(InNewState)
+	{
+	case ERunState::Running:
+		check(RunState == ERunState::NotInitialized || RunState == ERunState::Paused || RunState == ERunState::Running);
+		break;
+	case ERunState::Paused:
+		check(RunState == ERunState::Paused || RunState == ERunState::Running);
+		break;
+	case ERunState::NotInitialized:
+		check(RunState == ERunState::NotInitialized || RunState == ERunState::Paused || RunState == ERunState::Running);
+		break;
+	default:
+		checkNoEntry();
+	}
+
+	RunState = InNewState;
+}
+
+void FAnimNextModuleInstance::CopyProxyVariables()
+{
+	int32 ProxyReadIndex = 0;
+	{
+		FWriteScopeLock WriteLock(ProxyLock);
+
+		// Flip the write buffer index
+		ProxyReadIndex = ProxyWriteIndex;
+		ProxyWriteIndex = 1 - ProxyWriteIndex;
+	}
+
+	FUAFInstanceVariableDataProxy& ProxyVariablesRead = ProxyVariables[ProxyReadIndex];
+	ProxyVariablesRead.CopyDirty();
+}
+
+#if ANIMNEXT_TRACE_ENABLED
+void FAnimNextModuleInstance::Trace()
+{
+	if (!bTracedThisFrame)
+	{
+		TRACE_ANIMNEXT_VARIABLES(this, Object);
+		bTracedThisFrame = true;
+	}
+}
+#endif
+
+const UAnimNextModule* FAnimNextModuleInstance::GetModule() const
+{
+	return CastChecked<UAnimNextModule>(Asset);
+}
+
+#if WITH_EDITOR
+void FAnimNextModuleInstance::OnCompileJobFinished()
+{
+	using namespace UE::UAF;
+
+	FGuardValue_Bitfield(bIsRecreatingOnCompile, true);
+
+	ResetBindingsAndInstanceData();
+	Initialize();
+}
+#endif
+
+#if UE_ENABLE_DEBUG_DRAWING
+FRigVMDrawInterface* FAnimNextModuleInstance::GetDebugDrawInterface()
+{
+	return DebugDraw ? &DebugDraw->DrawInterface : nullptr;
+}
+
+void FAnimNextModuleInstance::ShowDebugDrawing(bool bInShowDebugDrawing)
+{
+	if(DebugDraw)
+	{
+		DebugDraw->SetEnabled(bInShowDebugDrawing);
+	}
+}
+#endif
+
+void FAnimNextModuleInstance::RunTaskOnGameThread(TUniqueFunction<void(void)>&& InFunction)
+{
+	UE::UAF::FModuleEventTickFunction::RunTaskOnGameThread(MoveTemp(InFunction));
+}
+
+UE::UAF::FModuleEventTickFunction* FAnimNextModuleInstance::FindTickFunctionByName(FName InEventName)
+{
+	using namespace UE::UAF;
+
+	return const_cast<FModuleEventTickFunction*>(const_cast<const FAnimNextModuleInstance*>(this)->FindTickFunctionByName(InEventName));
+}
+
+const UE::UAF::FModuleEventTickFunction* FAnimNextModuleInstance::FindTickFunctionByName(FName InEventName) const
+{
+	using namespace UE::UAF;
+
+	for(const FModuleEventTickFunction& TickFunction : TickFunctions)
+	{
+		if (TickFunction.EventName == InEventName)
+		{
+			return &TickFunction;
+		}
+	}
+	return nullptr;
+}
+
+void FAnimNextModuleInstance::EndExecution(float InDeltaTime)
+{
+	// Give the module a chance to handle events
+	RaiseTraitEvents(OutputEventList);
+
+	// Give each component a chance to finalize execution
+	for(auto It = ComponentMap.CreateIterator(); It; ++It)
+	{
+		FUAFModuleInstanceComponent* Component = It->Value.GetMutablePtr<FUAFModuleInstanceComponent>();
+		if (Component == nullptr)
+		{
+			continue;
+		}
+		Component->OnEndExecution(InDeltaTime);
+	}
+}
+
+void FAnimNextModuleInstance::RaiseTraitEvents(const UE::UAF::FTraitEventList& EventList)
+{
+	for(auto It = ComponentMap.CreateIterator(); It; ++It)
+	{
+		FUAFModuleInstanceComponent* Component = It->Value.GetMutablePtr<FUAFModuleInstanceComponent>();
+		if (Component == nullptr)
+		{
+			continue;
+		}
+
+		// Event handlers can raise events and as such the list may change while we iterate
+		// However, if an event is added while we iterate, we will not visit it
+		const int32 NumEvents = EventList.Num();
+		for (int32 EventIndex = 0; EventIndex < NumEvents; ++EventIndex)
+		{
+			const FAnimNextTraitEventPtr Event = EventList[EventIndex];
+			if (Event->IsValid())
+			{
+				Component->OnTraitEvent(*Event);
+			}
+		}
+	}
+}
+
+EPropertyBagResult FAnimNextModuleInstance::SetProxyVariable(const FAnimNextVariableReference& InVariable, const FAnimNextParamType& InType, TConstArrayView<uint8> InData)
+{
+	FWriteScopeLock WriteLock(ProxyLock);
+	FUAFInstanceVariableDataProxy& Proxy = ProxyVariables[ProxyWriteIndex];
+	return Proxy.SetVariable(InVariable, InType, InData);
+}
+
+EPropertyBagResult FAnimNextModuleInstance::WriteProxyVariable(const FAnimNextVariableReference& InVariable, const FAnimNextParamType& InType, TFunctionRef<void(TArrayView<uint8>)> InFunction)
+{
+	FWriteScopeLock WriteLock(ProxyLock);
+	FUAFInstanceVariableDataProxy& Proxy = ProxyVariables[ProxyWriteIndex];
+	TArrayView<uint8> Data;
+	EPropertyBagResult Result = Proxy.WriteVariable(InVariable, InType, Data);
+	if (Result == EPropertyBagResult::Success)
+	{
+		InFunction(Data);
+	}
+	return Result;
+}
+
+void FAnimNextModuleInstance::RunRigVMEvent(FName InEventName, float InDeltaTime)
+{
+	URigVM* VM = GetModule()->RigVM;
+	if(VM == nullptr)
+	{
+		return;
+	}
+
+	if(!VM->ContainsEntry(InEventName))
+	{
+		return;
+	}
+
+	FUAFRigVMComponent& RigVMComponent = GetComponent<FUAFRigVMComponent>();
+	FRigVMExtendedExecuteContext& ExtendedExecuteContext = RigVMComponent.GetExtendedExecuteContext();
+	check(ExtendedExecuteContext.VMHash == VM->GetVMHash());
+
+	FAnimNextExecuteContext& AnimNextContext = ExtendedExecuteContext.GetPublicDataSafe<FAnimNextExecuteContext>();
+
+	// RigVM setup
+	AnimNextContext.SetDeltaTime(InDeltaTime);
+	AnimNextContext.SetOwningObject(Object);
+
+#if UE_ENABLE_DEBUG_DRAWING
+	AnimNextContext.SetDrawInterface(GetDebugDrawInterface());
+#endif
+
+	// Insert our context data for the scope of execution
+	FAnimNextModuleContextData ContextData(this);
+	UE::UAF::FScopedExecuteContextData ContextDataScope(AnimNextContext, ContextData);
+
+	// Run the VM for this event
+	VM->ExecuteVM(ExtendedExecuteContext, InEventName);
+}
+
+
+TArrayView<UE::UAF::FModuleEventTickFunction> FAnimNextModuleInstance::GetTickFunctions()
+{
+	return TickFunctions;
+}
+
+UE::UAF::FModuleEventTickFunction* FAnimNextModuleInstance::FindFirstUserTickFunction()
+{
+	UE::UAF::FModuleEventTickFunction* FoundTickFunction = TickFunctions.FindByPredicate([](const UE::UAF::FModuleEventTickFunction& InTickFunction)
+	{
+		return InTickFunction.bFirstUserEvent;
+	});
+
+	return FoundTickFunction;
+}
+
+void FAnimNextModuleInstance::QueueTask(FName InEventName, TUniqueFunction<void(const UE::UAF::FModuleTaskContext&)>&& InTaskFunction, UE::UAF::ETaskRunLocation InLocation)
+{
+	using namespace UE::UAF;
+
+	FModuleEventTickFunction* FoundTickFunction = nullptr;
+	if(TickFunctions.Num() > 0)
+	{
+		if(!InEventName.IsNone())
+		{
+			// Match according to event desc
+			FoundTickFunction = TickFunctions.FindByPredicate([InEventName](const FModuleEventTickFunction& InTickFunction)
+			{
+				return InTickFunction.EventName == InEventName;
+			});
+		}
+
+		if (FoundTickFunction == nullptr)
+		{
+			// Fall back to first user function
+			FoundTickFunction = TickFunctions.FindByPredicate([](const UE::UAF::FModuleEventTickFunction& InTickFunction)
+			{
+				return InTickFunction.bFirstUserEvent;
+			});
+		}
+	}
+
+	TSpscQueue<TUniqueFunction<void(const FModuleTaskContext&)>>* Queue = nullptr;
+	if (FoundTickFunction)
+	{
+		switch (InLocation)
+		{
+		case ETaskRunLocation::Before:
+			Queue = &FoundTickFunction->PreExecuteTasks;
+			break;
+		case ETaskRunLocation::After:
+			Queue = &FoundTickFunction->PostExecuteTasks;
+			break;
+		}
+	}
+
+	if (Queue)
+	{
+		Queue->Enqueue(MoveTemp(InTaskFunction));
+	}
+	else
+	{
+		UE_LOGFMT(LogAnimation, Warning, "QueueTask: Could not find event '{EventName}' in module '{ModuleName}'", InEventName, GetAssetName());
+	}
+}
+
+void FAnimNextModuleInstance::QueueTaskOnOtherModule(const UE::UAF::FModuleHandle InOtherModuleHandle, const FName InEventName, TUniqueFunction<void(const UE::UAF::FModuleTaskContext&)>&& InTaskFunction, UE::UAF::ETaskRunLocation InLocation)
+{
+	if (Pool == nullptr)
+	{
+		return;
+	}
+
+	FAnimNextModuleInstance* OtherModuleInstance = Pool->TryGet(InOtherModuleHandle);
+	if (OtherModuleInstance == nullptr)
+	{
+		return;
+	}
+
+	OtherModuleInstance->QueueTask(InEventName, MoveTemp(InTaskFunction), InLocation);
+}
\ No newline at end of file
diff --git a/Engine/Plugins/Experimental/UAF/UAF/Source/UAF/Private/Module/ModuleTaskContext.cpp b/Engine/Plugins/Experimental/UAF/UAF/Source/UAF/Private/Module/ModuleTaskContext.cpp
new file mode 100644
--- /dev/null
+++ b/Engine/Plugins/Experimental/UAF/UAF/Source/UAF/Private/Module/ModuleTaskContext.cpp
@@ -0,0 +1,37 @@
+// Copyright Epic Games, Inc. All Rights Reserved.
+
+#include "Module/ModuleTaskContext.h"
+
+#include "Module/AnimNextModuleInstance.h"
+
+namespace UE::UAF
+{
+
+FModuleTaskContext::FModuleTaskContext(FAnimNextModuleInstance& InModuleInstance)
+	: FInstanceTaskContext(InModuleInstance)
+	, ModuleInstance(InModuleInstance)
+{
+}
+
+void FModuleTaskContext::QueueInputTraitEvent(FAnimNextTraitEventPtr Event) const
+{
+	ModuleInstance.QueueInputTraitEvent(MoveTemp(Event));
+}
+
+void FModuleTaskContext::TryAccessComponent(UScriptStruct* InComponentType, TFunctionRef<void(FUAFModuleInstanceComponent&)> InFunction) const
+{
+	FUAFModuleInstanceComponent* Component = static_cast<FUAFModuleInstanceComponent*>(ModuleInstance.TryGetComponent(InComponentType));
+	if(Component == nullptr)
+	{
+		return;
+	}
+
+	InFunction(*Component);
+}
+
+FAnimNextModuleInstance* const FModuleTaskContext::GetModuleInstance() const
+{
+	return &ModuleInstance;
+}
+
+}
diff --git a/Engine/Plugins/Experimental/UAF/UAF/Source/UAF/Private/UAFAssetInstanceComponent.cpp b/Engine/Plugins/Experimental/UAF/UAF/Source/UAF/Private/UAFAssetInstanceComponent.cpp
new file mode 100644
--- /dev/null
+++ b/Engine/Plugins/Experimental/UAF/UAF/Source/UAF/Private/UAFAssetInstanceComponent.cpp
@@ -0,0 +1,23 @@
+﻿// Copyright Epic Games, Inc. All Rights Reserved.
+
+#include "UAFAssetInstanceComponent.h"
+
+namespace UE::UAF::Private
+{
+	static thread_local FUAFAssetInstance* GCurrentAssetInstance = nullptr;
+}
+
+FUAFAssetInstanceComponent::FScopedConstructorHelper::FScopedConstructorHelper(FUAFAssetInstance& InInstance)
+{
+	UE::UAF::Private::GCurrentAssetInstance = &InInstance;
+}
+
+FUAFAssetInstanceComponent::FScopedConstructorHelper::~FScopedConstructorHelper()
+{
+	UE::UAF::Private::GCurrentAssetInstance = nullptr;
+}
+
+FUAFAssetInstanceComponent::FUAFAssetInstanceComponent()
+{
+	Instance = UE::UAF::Private::GCurrentAssetInstance;
+}
\ No newline at end of file
diff --git a/Engine/Plugins/Experimental/UAF/UAF/Source/UAF/Private/UAFRigVMComponent.cpp b/Engine/Plugins/Experimental/UAF/UAF/Source/UAF/Private/UAFRigVMComponent.cpp
new file mode 100644
--- /dev/null
+++ b/Engine/Plugins/Experimental/UAF/UAF/Source/UAF/Private/UAFRigVMComponent.cpp
@@ -0,0 +1,49 @@
+﻿// Copyright Epic Games, Inc. All Rights Reserved.
+
+#include "UAFRigVMComponent.h"
+
+#include "AnimNextRigVMAsset.h"
+#include "UAFAssetInstance.h"
+#include "Variables/UAFInstanceVariableContainer.h"
+
+FUAFRigVMComponent::FUAFRigVMComponent()
+{
+	FUAFAssetInstance* AssetInstance = GetAssetInstancePtr();
+	if (AssetInstance == nullptr)
+	{
+		return;
+	}
+
+	const UAnimNextRigVMAsset* Asset = AssetInstance->GetAsset<UAnimNextRigVMAsset>();
+	URigVM* VM = Asset->RigVM;
+	if (VM == nullptr)
+	{
+		return;
+	}
+
+	// Initialize the RigVM context
+	ExtendedExecuteContext = Asset->GetRigVMExtendedExecuteContext();
+
+	// Hookup the runtime data ptrs
+	const TArray<FRigVMExternalVariableDef>& ExternalVariableDefs = VM->GetExternalVariableDefs();
+	const int32 NumExternalVariables = ExternalVariableDefs.Num();
+	check(AssetInstance->Variables.NumInternalVariables == NumExternalVariables);
+	if(NumExternalVariables > 0)
+	{
+		TArray<FRigVMExternalVariableRuntimeData> ExternalVariableRuntimeData;
+		ExternalVariableRuntimeData.Reserve(AssetInstance->Variables.NumInternalVariables);
+		const int32 NumInternalSharedVariableContainers = AssetInstance->Variables.InternalSharedVariableContainers.Num();
+		int32 CurrentBaseIndex = 0;
+		for (int32 SharedVariableSetIndex = 0; SharedVariableSetIndex < NumInternalSharedVariableContainers; ++SharedVariableSetIndex)
+		{
+			const TSharedRef<FUAFInstanceVariableContainer>& SharedVariableSet = AssetInstance->Variables.InternalSharedVariableContainers[SharedVariableSetIndex].Pin().ToSharedRef();
+			TConstArrayView<FRigVMExternalVariableDef> DefsForSet(ExternalVariableDefs.GetData() + CurrentBaseIndex, SharedVariableSet->NumVariables);
+			SharedVariableSet->GetRigVMMemoryForVariables(DefsForSet, ExternalVariableRuntimeData);
+			CurrentBaseIndex += SharedVariableSet->NumVariables;
+		}
+		ExtendedExecuteContext.ExternalVariableRuntimeData = MoveTemp(ExternalVariableRuntimeData);
+	}
+
+	// Now initialize the 'instance', cache memory handles etc. in the context
+	VM->InitializeInstance(ExtendedExecuteContext);
+}
\ No newline at end of file
diff --git a/Engine/Plugins/Experimental/UAF/UAF/Source/UAF/Private/Variables/AnimNextSharedVariables.cpp b/Engine/Plugins/Experimental/UAF/UAF/Source/UAF/Private/Variables/AnimNextSharedVariables.cpp
new file mode 100644
--- /dev/null
+++ b/Engine/Plugins/Experimental/UAF/UAF/Source/UAF/Private/Variables/AnimNextSharedVariables.cpp
@@ -0,0 +1,15 @@
+﻿// Copyright Epic Games, Inc. All Rights Reserved.
+
+#include "Variables/AnimNextSharedVariables.h"
+#include "AnimNextExecuteContext.h"
+
+#if WITH_EDITOR	
+#include "Engine/ExternalAssetDependencyGatherer.h"
+REGISTER_ASSETDEPENDENCY_GATHERER(FExternalAssetDependencyGatherer, UAnimNextSharedVariables);
+#endif // WITH_EDITOR
+
+UAnimNextSharedVariables::UAnimNextSharedVariables(const FObjectInitializer& ObjectInitializer)
+	: Super(ObjectInitializer)
+{
+	ExtendedExecuteContext.SetContextPublicDataStruct(FAnimNextExecuteContext::StaticStruct());
+}
diff --git a/Engine/Plugins/Experimental/UAF/UAF/Source/UAF/Private/Variables/AnimNextSoftVariableReference.cpp b/Engine/Plugins/Experimental/UAF/UAF/Source/UAF/Private/Variables/AnimNextSoftVariableReference.cpp
new file mode 100644
--- /dev/null
+++ b/Engine/Plugins/Experimental/UAF/UAF/Source/UAF/Private/Variables/AnimNextSoftVariableReference.cpp
@@ -0,0 +1,14 @@
+﻿// Copyright Epic Games, Inc. All Rights Reserved.
+
+#include "Variables/AnimNextSoftVariableReference.h"
+#include "AnimNextRigVMAsset.h"
+#include "Variables/AnimNextVariableReference.h"
+
+FAnimNextSoftVariableReference::FAnimNextSoftVariableReference(const FAnimNextVariableReference& InVariableReference)
+	: Name(InVariableReference.GetName())
+	, SoftObjectPath(InVariableReference.GetObject())
+{
+	const UObject* Object = InVariableReference.GetObject();
+	check(Object == nullptr || Object->IsA<UAnimNextRigVMAsset>() || Object->IsA<UScriptStruct>());
+}
+
diff --git a/Engine/Plugins/Experimental/UAF/UAF/Source/UAF/Private/Variables/RigVMTrait_OverrideVariables.h b/Engine/Plugins/Experimental/UAF/UAF/Source/UAF/Private/Variables/RigVMTrait_OverrideVariables.h
new file mode 100644
--- /dev/null
+++ b/Engine/Plugins/Experimental/UAF/UAF/Source/UAF/Private/Variables/RigVMTrait_OverrideVariables.h
@@ -0,0 +1,88 @@
+// Copyright Epic Games, Inc. All Rights Reserved.
+
+#pragma once
+
+#include "RigVMCore/RigVMTrait.h"
+#include "Param/ParamType.h"
+#include "RigVMTrait_OverrideVariables.generated.h"
+
+class UAnimNextRigVMAsset;
+
+namespace UE::UAF
+{
+	struct FVariableOverrides;
+}
+
+// Override info for a single variable
+USTRUCT(BlueprintInternalUseOnly)
+struct FAnimNextVariableOverrideInfo
+{
+	GENERATED_BODY()
+
+	FAnimNextVariableOverrideInfo() = default;
+
+	FAnimNextVariableOverrideInfo(FName InName, const FAnimNextParamType& InType, const FString& InDefaultValue)
+		: Name(InName)
+		, Type(InType)
+		, DefaultValue(InDefaultValue)
+	{}
+
+	UPROPERTY(EditAnywhere, Category = "Variables")
+	FName Name = "Variable";
+
+	UPROPERTY(EditAnywhere, Category = "Variables")
+	FAnimNextParamType Type = FAnimNextParamType(FAnimNextParamType::EValueType::Bool);
+
+	UPROPERTY()
+	FString DefaultValue;
+};
+
+// A trait that allows insertion of specific asset/struct variable overrides onto an override variables node
+USTRUCT(meta=(Hidden))
+struct UAF_API FRigVMTrait_OverrideVariables : public FRigVMTrait
+{
+	GENERATED_BODY()
+
+	// FRigVMTrait interface
+#if WITH_EDITOR
+	virtual void GetProgrammaticPins(URigVMController* InController, int32 InParentPinIndex, const URigVMPin* InTraitPin, const FString& InDefaultValue, FRigVMPinInfoArray& OutPinArray) const override;
+	virtual bool ShouldCreatePinForProperty(const FProperty* InProperty) const override;
+#endif
+
+	// Override point for derived types to generate overrides
+	virtual void GenerateOverrides(const FRigVMTraitScope& InTraitScope, TArray<UE::UAF::FVariableOverrides>& OutOverrides) const {}
+
+	// Cached override data
+	UPROPERTY(EditAnywhere, Category = "Variables", meta=(Input, DetailsOnly, TitleProperty = "Name"))
+	TArray<FAnimNextVariableOverrideInfo> Overrides;
+};
+
+USTRUCT()
+struct UAF_API FRigVMTrait_OverrideVariablesForAsset : public FRigVMTrait_OverrideVariables
+{
+	GENERATED_BODY()
+
+#if WITH_EDITOR
+	virtual FString GetDisplayName() const override;
+#endif
+	virtual void GenerateOverrides(const FRigVMTraitScope& InTraitScope, TArray<UE::UAF::FVariableOverrides>& OutOverrides) const override;
+
+	// The asset to use to generate overrides
+	UPROPERTY(EditAnywhere, Category = "Variables", meta=(Input, Constant, DetailsOnly))
+	TObjectPtr<const UAnimNextRigVMAsset> Asset;
+};
+
+USTRUCT()
+struct UAF_API FRigVMTrait_OverrideVariablesForStruct : public FRigVMTrait_OverrideVariables
+{
+	GENERATED_BODY()
+
+#if WITH_EDITOR
+	virtual FString GetDisplayName() const override;
+#endif
+	virtual void GenerateOverrides(const FRigVMTraitScope& InTraitScope, TArray<UE::UAF::FVariableOverrides>& OutOverrides) const override;
+
+	// The asset to use to generate overrides
+	UPROPERTY(EditAnywhere, Category = "Variables", meta=(Input, Constant, DetailsOnly))
+	TObjectPtr<const UScriptStruct> Struct;
+};
diff --git a/Engine/Plugins/Experimental/UAF/UAF/Source/UAF/Private/Variables/VariableOverrides.cpp b/Engine/Plugins/Experimental/UAF/UAF/Source/UAF/Private/Variables/VariableOverrides.cpp
new file mode 100644
--- /dev/null
+++ b/Engine/Plugins/Experimental/UAF/UAF/Source/UAF/Private/Variables/VariableOverrides.cpp
@@ -0,0 +1,25 @@
+﻿// Copyright Epic Games, Inc. All Rights Reserved.
+
+#include "VariableOverrides.h"
+
+#include "AnimNextRigVMAsset.h"
+
+namespace UE::UAF
+{
+
+FVariableOverrides::FVariableOverrides(const UAnimNextRigVMAsset* InAsset, TArray<FOverride>&& InOverrides)
+	: Overrides(MoveTemp(InOverrides))
+{
+	check(InAsset);
+
+	AssetOrStructData.Set<FAssetType>(InAsset);
+}
+
+FVariableOverrides::FVariableOverrides(const UScriptStruct* InStruct, TArray<FOverride>&& InOverrides)
+	: Overrides(MoveTemp(InOverrides))
+{
+	AssetOrStructData.Set<FStructType>(InStruct);
+
+}
+
+}
diff --git a/Engine/Plugins/Experimental/UAF/UAF/Source/UAF/Public/AnimNextFunctionHandle.h b/Engine/Plugins/Experimental/UAF/UAF/Source/UAF/Public/AnimNextFunctionHandle.h
new file mode 100644
--- /dev/null
+++ b/Engine/Plugins/Experimental/UAF/UAF/Source/UAF/Public/AnimNextFunctionHandle.h
@@ -0,0 +1,38 @@
+﻿// Copyright Epic Games, Inc. All Rights Reserved.
+
+#pragma once
+
+#include "CoreTypes.h"
+
+class UAnimNextRigVMAsset;
+
+namespace UE::UAF
+{
+
+// Opaque handle that can be used to call RigVM functions on UAnimNextRigVMAsset
+struct FFunctionHandle
+{
+	// Check to see if this is a valid handle
+	bool IsValid() const
+	{
+		return FunctionIndex != INDEX_NONE;
+	}
+
+#if WITH_EDITOR
+	// Check to see if this is a valid handle and that it can be used for the specified VM (hash comparison)
+	bool IsValidForVM(uint32 InVMHash) const
+	{
+		return IsValid() && VMHash == InVMHash;
+	}
+#endif
+
+private:
+	friend ::UAnimNextRigVMAsset;
+
+	int32 FunctionIndex = INDEX_NONE;
+#if WITH_EDITOR
+	uint32 VMHash = 0;
+#endif
+};
+
+}
\ No newline at end of file
diff --git a/Engine/Plugins/Experimental/UAF/UAF/Source/UAF/Public/Module/UAFModuleInstanceComponent.h b/Engine/Plugins/Experimental/UAF/UAF/Source/UAF/Public/Module/UAFModuleInstanceComponent.h
new file mode 100644
--- /dev/null
+++ b/Engine/Plugins/Experimental/UAF/UAF/Source/UAF/Public/Module/UAFModuleInstanceComponent.h
@@ -0,0 +1,45 @@
+// Copyright Epic Games, Inc. All Rights Reserved.
+
+#pragma once
+
+#include "CoreMinimal.h"
+#include "UAFAssetInstanceComponent.h"
+#include "UAFModuleInstanceComponent.generated.h"
+
+struct FAnimNextModuleInstance;
+struct FAnimNextTraitEvent;
+
+#define UE_API UAF_API
+
+/** A module instance component is attached and owned by a module instance. */
+USTRUCT()
+struct FUAFModuleInstanceComponent : public FUAFAssetInstanceComponent
+{
+	GENERATED_BODY()
+
+	using ContainerType = FAnimNextModuleInstance;
+
+	FUAFModuleInstanceComponent() = default;
+
+	virtual ~FUAFModuleInstanceComponent() override = default;
+
+	// Returns the owning module instance this component lives on
+	UE_API FAnimNextModuleInstance* GetModuleInstancePtr();
+
+	// Returns the owning module instance this component lives on
+	UE_API FAnimNextModuleInstance& GetModuleInstance();
+
+	// Returns the owning module instance this component lives on
+	UE_API const FAnimNextModuleInstance& GetModuleInstance() const;
+
+	// Called during module execution for any events to be handled
+	virtual void OnTraitEvent(FAnimNextTraitEvent& Event) {}
+
+	// Called at end of module execution each frame
+	virtual void OnEndExecution(float InDeltaTime) {}
+
+private:
+	friend struct FAnimNextModuleInstance;
+};
+
+#undef UE_API
\ No newline at end of file
diff --git a/Engine/Plugins/Experimental/UAF/UAF/Source/UAF/Public/RewindDebugger/AnimNextTrace.h b/Engine/Plugins/Experimental/UAF/UAF/Source/UAF/Public/RewindDebugger/AnimNextTrace.h
new file mode 100644
--- /dev/null
+++ b/Engine/Plugins/Experimental/UAF/UAF/Source/UAF/Public/RewindDebugger/AnimNextTrace.h
@@ -0,0 +1,52 @@
+// Copyright Epic Games, Inc. All Rights Reserved.
+
+#pragma once
+
+#if UE_ENABLE_INCLUDE_ORDER_DEPRECATED_IN_5_7
+#include "CoreMinimal.h"
+#include "ObjectTrace.h"
+#include "StructUtils/PropertyBag.h"
+#endif // UE_ENABLE_INCLUDE_ORDER_DEPRECATED_IN_5_7
+
+#include "CoreTypes.h"
+#include "Misc/Guid.h"
+#include "ObjectTraceDefines.h"
+#include "Trace/Trace.h"
+#include "Templates/SharedPointer.h"
+
+#define ANIMNEXT_TRACE_ENABLED (OBJECT_TRACE_ENABLED && !(UE_BUILD_SHIPPING || UE_BUILD_TEST))
+
+#if ANIMNEXT_TRACE_ENABLED
+
+UE_TRACE_CHANNEL_EXTERN(UAFChannel, UAF_API);
+
+class UObject;
+struct FAnimNextModuleInstance;
+struct FUAFAssetInstance;
+struct FAnimNextGraphInstance;
+struct FUAFInstanceVariableContainer;
+
+namespace UE::UAF
+{
+struct FAnimNextTrace
+{
+	UAF_API static const FGuid CustomVersionGUID;
+	
+	UAF_API static void Reset();
+	
+	UAF_API static void OutputAnimNextInstance(const FUAFAssetInstance* DataInterface, const UObject* OuterObject);
+	UAF_API static void OutputAnimNextVariables(const FUAFAssetInstance* DataInterface, const UObject* OuterObject);
+
+private:
+	UAF_API static bool OutputAnimNextVariableSet(const TSharedRef<FUAFInstanceVariableContainer>& VariableSet, uint64 InstanceId, const UObject* OuterObject);
+};
+
+}
+
+
+#define TRACE_ANIMNEXT_INSTANCE(DataInterface, OuterObject) UE::UAF::FAnimNextTrace::OutputAnimNextInstance(DataInterface, OuterObject);
+#define TRACE_ANIMNEXT_VARIABLES(DataInterface, OuterObject) UE::UAF::FAnimNextTrace::OutputAnimNextVariables(DataInterface, OuterObject);
+#else
+#define TRACE_ANIMNEXT_INSTANCE(DataInterface, OuterObject)
+#define TRACE_ANIMNEXT_VARIABLES(DataInterface, OuterObject)
+#endif
diff --git a/Engine/Plugins/Experimental/UAF/UAF/Source/UAF/Public/Variables/AnimNextSharedVariables.h b/Engine/Plugins/Experimental/UAF/UAF/Source/UAF/Public/Variables/AnimNextSharedVariables.h
new file mode 100644
--- /dev/null
+++ b/Engine/Plugins/Experimental/UAF/UAF/Source/UAF/Public/Variables/AnimNextSharedVariables.h
@@ -0,0 +1,35 @@
+// Copyright Epic Games, Inc. All Rights Reserved.
+
+#pragma once
+
+#include "CoreMinimal.h"
+#include "AnimNextRigVMAsset.h"
+#include "AnimNextSharedVariables.generated.h"
+
+class UAnimNextSharedVariablesFactory;
+class UAnimNextSharedVariables;
+
+namespace UE::UAF
+{
+struct FInjectionInfo;
+}
+
+namespace UE::UAF::UncookedOnly
+{
+struct FUtils;
+}
+
+// Shared variables are just bundles of variables and their default values
+UCLASS(BlueprintType)
+class UAF_API UAnimNextSharedVariables : public UAnimNextRigVMAsset
+{
+	GENERATED_BODY()
+
+public:
+	UAnimNextSharedVariables(const FObjectInitializer& ObjectInitializer);
+
+private:
+	friend class UAnimNextSharedVariablesFactory;
+	friend struct UE::UAF::UncookedOnly::FUtils;
+	friend UE::UAF::FInjectionInfo;
+};
diff --git a/Engine/Plugins/Experimental/UAF/UAF/Source/UAF/Public/Variables/AnimNextSoftVariableReference.h b/Engine/Plugins/Experimental/UAF/UAF/Source/UAF/Public/Variables/AnimNextSoftVariableReference.h
new file mode 100644
--- /dev/null
+++ b/Engine/Plugins/Experimental/UAF/UAF/Source/UAF/Public/Variables/AnimNextSoftVariableReference.h
@@ -0,0 +1,75 @@
+// Copyright Epic Games, Inc. All Rights Reserved.
+
+#pragma once
+
+#include "AnimNextSoftVariableReference.generated.h"
+
+class UAnimNextRigVMAsset;
+struct FAnimNextVariableReference;
+
+// A soft reference to an AnimNext variable
+USTRUCT(BlueprintType)
+struct FAnimNextSoftVariableReference
+{
+	GENERATED_BODY()
+
+	FAnimNextSoftVariableReference() = default;
+
+	// Construct a soft reference to an asset/struct variable
+	explicit FAnimNextSoftVariableReference(FName InName, const FSoftObjectPath& InSoftObjectPath)
+		: Name(InName)
+		, SoftObjectPath(InSoftObjectPath)
+	{}
+
+	// Construct a soft reference from a hard reference
+	UAF_API explicit FAnimNextSoftVariableReference(const FAnimNextVariableReference& InVariableReference);
+
+	// Get the name of the variable
+	FName GetName() const
+	{
+		return Name;
+	}
+
+	// Get the asset or struct that the variable reference is contained in
+	const FSoftObjectPath& GetSoftObjectPath() const
+	{
+		return SoftObjectPath;
+	}
+
+	// Check whether this can ever refer to a valid variable
+	bool IsNone() const
+	{
+		return Name.IsNone() || SoftObjectPath.IsNull();
+	}
+
+	// Set this reference to None
+	void Reset()
+	{
+		Name = NAME_None; 
+		SoftObjectPath.Reset();
+	}
+
+	bool operator==(const FAnimNextSoftVariableReference& InOther) const
+	{
+		return Name == InOther.Name && SoftObjectPath == InOther.SoftObjectPath;
+	}
+
+	bool operator!=(const FAnimNextSoftVariableReference& InOther) const
+	{
+		return !(*this == InOther);
+	}
+
+	friend uint32 GetTypeHash(const FAnimNextSoftVariableReference& InKey)
+	{
+		return HashCombineFast(GetTypeHash(InKey.Name), GetTypeHash(InKey.SoftObjectPath));
+	}
+
+private:
+	// The name of the variable
+	UPROPERTY(EditAnywhere, Category = Variable)
+	FName Name;
+
+	// The asset or struct that the variable reference is contained in
+	UPROPERTY(EditAnywhere, Category = Variable)
+	FSoftObjectPath SoftObjectPath;
+};
\ No newline at end of file
diff --git a/Engine/Plugins/Experimental/UAF/UAF/Source/UAFEditor/Internal/AnimNextStyle.h b/Engine/Plugins/Experimental/UAF/UAF/Source/UAFEditor/Internal/AnimNextStyle.h
new file mode 100644
--- /dev/null
+++ b/Engine/Plugins/Experimental/UAF/UAF/Source/UAFEditor/Internal/AnimNextStyle.h
@@ -0,0 +1,22 @@
+﻿// Copyright Epic Games, Inc. All Rights Reserved.
+
+#pragma once
+
+#include "Styling/SlateStyle.h"
+
+class ISlateStyle;
+
+class UAFEDITOR_API FAnimNextStyle : public FSlateStyleSet
+{
+public:
+	static FAnimNextStyle& Get();
+
+protected:
+	static void Register();
+	static void Unregister();
+
+private:
+	FAnimNextStyle();
+
+	void SetAssetWizardStyles();
+};
diff --git a/Engine/Plugins/Experimental/UAF/UAF/Source/UAFEditor/Private/AnimNextEditorStyle.h b/Engine/Plugins/Experimental/UAF/UAF/Source/UAFEditor/Private/AnimNextEditorStyle.h
new file mode 100644
--- /dev/null
+++ b/Engine/Plugins/Experimental/UAF/UAF/Source/UAFEditor/Private/AnimNextEditorStyle.h
@@ -0,0 +1,68 @@
+// Copyright Epic Games, Inc. All Rights Reserved.
+
+#pragma once
+
+#include "Styling/SlateStyle.h"
+#include "Styling/SlateStyleRegistry.h"
+#include "Styling/SlateStyleMacros.h"
+#include "Styling/AppStyle.h"
+#include "Styling/CoreStyle.h"
+#include "Styling/SlateTypes.h"
+#include "Styling/StyleColors.h"
+#include "Interfaces/IPluginManager.h"
+
+#define IMAGE_PLUGIN_BRUSH_SVG( RelativePath, ... ) FSlateVectorImageBrush( FAnimNextEditorStyle::InContent(RelativePath, ".svg"), __VA_ARGS__)
+#define PLUGIN_BOX_BRUSH_BRUSH( RelativePath, ... ) FSlateBoxBrush( FAnimNextEditorStyle::InContent(RelativePath, ".png"), __VA_ARGS__)
+
+class FAnimNextEditorStyle final : public FSlateStyleSet
+{
+public:	
+	FAnimNextEditorStyle() : FSlateStyleSet("AnimNextEditorStyle")
+	{
+		const FVector2D Icon16x16(16.0f, 16.0f);
+		
+		SetContentRoot(FPaths::EngineContentDir() / TEXT("Editor/Slate"));
+		SetCoreContentRoot(FPaths::EngineContentDir() / TEXT("Slate"));
+
+		// Class Icons
+		{
+			Set("ClassIcon.AnimNextSharedVariables", new IMAGE_PLUGIN_BRUSH_SVG("Slate/Archive_16", Icon16x16));
+			Set("ClassIcon.AnimNextModule", new IMAGE_BRUSH_SVG( "Starship/Common/Modules", Icon16x16 ) );
+			Set("ClassIcon.AnimNextAnimationGraph", new IMAGE_BRUSH_SVG("Starship/Blueprints/icon_BlueprintEditor_EventGraph", Icon16x16));
+			Set("ClassIcon.AnimNextStateTree", new IMAGE_BRUSH_SVG(TEXT("UMG/TreeView"), Icon16x16));
+		}
+
+		Set("VariablesOutliner.SharedVariablesPill.Border", new PLUGIN_BOX_BRUSH_BRUSH("Slate/NamespaceBorder", FMargin(4.0f / 16.0f)));
+
+		const FTextBlockStyle NormalText = FAppStyle::GetWidgetStyle<FTextBlockStyle>("NormalText");
+		Set("VariablesOutliner.SharedVariablesPill.Text", FTextBlockStyle(NormalText)
+			.SetFont(DEFAULT_FONT("Regular", 8))
+			.SetColorAndOpacity(FSlateColor::UseForeground()));
+
+		const FVector2D Icon14x14(14.0f, 14.0f);
+		const FButtonStyle CloseButtonStyle = FButtonStyle()
+		.SetNormal(CORE_IMAGE_BRUSH_SVG("Starship/Common/close-small", Icon14x14, FStyleColors::Foreground))
+		.SetPressed(CORE_IMAGE_BRUSH_SVG("Starship/Common/close-small", Icon14x14, FStyleColors::Foreground))
+		.SetHovered(CORE_IMAGE_BRUSH_SVG("Starship/Common/close-small", Icon14x14, FStyleColors::ForegroundHover));
+		Set("VariablesOutliner.SharedVariablesPill.CloseButton", CloseButtonStyle);
+
+		FSlateStyleRegistry::RegisterSlateStyle(*this);
+	}
+
+	static FAnimNextEditorStyle& Get()
+	{ 
+		static FAnimNextEditorStyle Instance;
+		return Instance;
+	}
+	
+	static FString InContent(const FString& RelativePath, const ANSICHAR* Extension)
+	{
+		static FString ContentDir = IPluginManager::Get().FindPlugin(TEXT("UAF"))->GetContentDir();
+		return (ContentDir / RelativePath) + Extension;
+	}
+
+	virtual ~FAnimNextEditorStyle() override
+	{
+		FSlateStyleRegistry::UnRegisterSlateStyle(*this);
+	}
+};
\ No newline at end of file
diff --git a/Engine/Plugins/Experimental/UAF/UAF/Source/UAFEditor/Private/AnimNextStyle.cpp b/Engine/Plugins/Experimental/UAF/UAF/Source/UAFEditor/Private/AnimNextStyle.cpp
new file mode 100644
--- /dev/null
+++ b/Engine/Plugins/Experimental/UAF/UAF/Source/UAFEditor/Private/AnimNextStyle.cpp
@@ -0,0 +1,82 @@
+﻿// Copyright Epic Games, Inc. All Rights Reserved.
+
+#include "AnimNextStyle.h"
+#include "Brushes/SlateBoxBrush.h"
+#include "Styling/SlateStyleRegistry.h"
+#include "Brushes/SlateImageBrush.h"
+#include "Styling/CoreStyle.h"
+#include "Brushes/SlateRoundedBoxBrush.h"
+#include "Styling/SlateTypes.h"
+#include "Misc/Paths.h"
+#include "Styling/StyleColors.h"
+#include "Styling/SlateStyleMacros.h"
+
+
+FAnimNextStyle::FAnimNextStyle()
+	: FSlateStyleSet(TEXT("AnimNextStyle"))
+{
+	const FString EngineSlateContentDir = FPaths::EngineContentDir() / TEXT("Slate");
+	const FString EngineEditorSlateContentDir = FPaths::EngineContentDir() / TEXT("Editor/Slate");
+	SetCoreContentRoot(EngineSlateContentDir);
+
+	const FString AnimNextPluginContentDir = FPaths::EnginePluginsDir() / TEXT("Experimental/AnimNext/Resources");
+	SetContentRoot(AnimNextPluginContentDir);
+
+	const FTextBlockStyle& NormalText = FAppStyle::Get().GetWidgetStyle<FTextBlockStyle>("NormalText");
+	const FComboButtonStyle& ComboButtonStyle = FCoreStyle::Get().GetWidgetStyle<FComboButtonStyle>("ComboButton");
+
+	// Trait details labels
+	Set("AnimNext.TraitDetails.Label", FTextBlockStyle(NormalText)
+		.SetFont(FAppStyle::GetFontStyle(TEXT("PropertyWindow.BoldFont")))
+		.SetFontSize(7));
+
+	Set("AnimNext.TraitDetails.Background", new FSlateRoundedBoxBrush(FStyleColors::Hover, 6.f));
+
+	// Add trait combo button
+	const FButtonStyle AddButtonStyle = FButtonStyle()
+		.SetNormal(FSlateRoundedBoxBrush(FStyleColors::Hover, 3.0f))
+		.SetHovered(FSlateRoundedBoxBrush(FStyleColors::Primary, 3.0f))
+		.SetPressed(FSlateRoundedBoxBrush(FStyleColors::Primary, 3.0f))
+		.SetNormalForeground(FStyleColors::Transparent)
+		.SetHoveredForeground(FStyleColors::Hover)
+		.SetPressedForeground(FStyleColors::Foreground)
+		.SetNormalPadding(FMargin(6, 2, 6, 2))
+		.SetPressedPadding(FMargin(6, 3, 6, 1));
+
+	Set("AnimNext.TraitDetails.AddButton", FComboButtonStyle(ComboButtonStyle).SetButtonStyle(AddButtonStyle));
+
+	SetAssetWizardStyles();
+}
+
+void FAnimNextStyle::SetAssetWizardStyles()
+{
+	const FVector2D Icon16x16(16.0f, 16.0f);
+	const FVector2D Icon64x64(64.0f, 64.0f);
+	
+	Set("AssetWizard.OpenExternal", new FSlateVectorImageBrush(RootToContentDir("Icons/Wizard/OpenExternal.svg"),Icon16x16));
+	Set("AssetWizard.Template", new FSlateVectorImageBrush(RootToContentDir("Icons/Wizard/Template.svg"),Icon64x64));
+
+	Set("AssetWizard.TemplateTag.OuterBorder", new FSlateRoundedBoxBrush(FStyleColors::AccentGray, 4.0f));
+	Set("AssetWizard.TemplateTag.InnerBorder", new FSlateRoundedBoxBrush(FStyleColors::Recessed, 4.0f));
+
+	Set("AssetWizard.TemplateFolderBackground", new FSlateRoundedBoxBrush(EStyleColor::Panel, 4.0f));
+	Set("AssetWizard.TemplateFolderHoverBackground", new FSlateRoundedBoxBrush(FStyleColors::Hover, 4.0f));
+	Set("AssetWizard.TemplateFolderSelectedBackground", new FSlateRoundedBoxBrush(FStyleColors::Primary, 4.0f));
+	Set("AssetWizard.TemplateFolderHoverSelectedBackground", new FSlateRoundedBoxBrush(FStyleColors::PrimaryHover, 4.0f));
+}
+
+void FAnimNextStyle::Register()
+{
+	FSlateStyleRegistry::RegisterSlateStyle(Get());
+}
+
+void FAnimNextStyle::Unregister()
+{
+	FSlateStyleRegistry::UnRegisterSlateStyle(Get());
+}
+
+FAnimNextStyle& FAnimNextStyle::Get()
+{
+	static FAnimNextStyle Instance;
+	return Instance;
+}
diff --git a/Engine/Plugins/Experimental/UAF/UAF/Source/UAFEditor/Private/AssetCompilationHandler.cpp b/Engine/Plugins/Experimental/UAF/UAF/Source/UAFEditor/Private/AssetCompilationHandler.cpp
new file mode 100644
--- /dev/null
+++ b/Engine/Plugins/Experimental/UAF/UAF/Source/UAFEditor/Private/AssetCompilationHandler.cpp
@@ -0,0 +1,71 @@
+﻿// Copyright Epic Games, Inc. All Rights Reserved.
+
+#include "AssetCompilationHandler.h"
+#include "AnimNextRigVMAsset.h"
+#include "AnimNextRigVMAssetEditorData.h"
+#include "AnimNextScopedCompileJob.h"
+#include "UncookedOnlyUtils.h"
+
+namespace UE::UAF::Editor
+{
+
+FAssetCompilationHandler::FAssetCompilationHandler(UObject* InAsset)
+{
+	UAnimNextRigVMAsset* AnimNextRigVMAsset = CastChecked<UAnimNextRigVMAsset>(InAsset);
+	UAnimNextRigVMAssetEditorData* EditorData = UE::UAF::UncookedOnly::FUtils::GetEditorData<UAnimNextRigVMAssetEditorData>(AnimNextRigVMAsset);
+
+	NumErrors = EditorData->bErrorsDuringCompilation ? 1 : 0;
+	NumWarnings = EditorData->bWarningsDuringCompilation ? 1 : 0;
+}
+
+void FAssetCompilationHandler::Compile(TSharedRef<Workspace::IWorkspaceEditor> InWorkspaceEditor, UObject* InAsset)
+{
+	UAnimNextRigVMAsset* AnimNextRigVMAsset = CastChecked<UAnimNextRigVMAsset>(InAsset);
+	UAnimNextRigVMAssetEditorData* EditorData = UE::UAF::UncookedOnly::FUtils::GetEditorData<UAnimNextRigVMAssetEditorData>(AnimNextRigVMAsset);
+
+	UncookedOnly::FScopedCompileJob CompileJob(InAsset);
+	EditorData->RecompileVM();
+
+	NumErrors = CompileJob.GetLog().NumErrors;
+	NumWarnings = CompileJob.GetLog().NumWarnings;
+
+	OnCompileStatusChanged().ExecuteIfBound();
+}
+
+void FAssetCompilationHandler::SetAutoCompile(TSharedRef<Workspace::IWorkspaceEditor> InWorkspaceEditor, UObject* InAsset, bool bInAutoCompile)
+{
+	UAnimNextRigVMAsset* AnimNextRigVMAsset = CastChecked<UAnimNextRigVMAsset>(InAsset);
+	UAnimNextRigVMAssetEditorData* EditorData = UE::UAF::UncookedOnly::FUtils::GetEditorData<UAnimNextRigVMAssetEditorData>(AnimNextRigVMAsset);
+	EditorData->SetAutoVMRecompile(bInAutoCompile);
+}
+
+bool FAssetCompilationHandler::GetAutoCompile(TSharedRef<Workspace::IWorkspaceEditor> InWorkspaceEditor, const UObject* InAsset) const
+{
+	const UAnimNextRigVMAsset* AnimNextRigVMAsset = CastChecked<UAnimNextRigVMAsset>(InAsset);
+	UAnimNextRigVMAssetEditorData* EditorData = UE::UAF::UncookedOnly::FUtils::GetEditorData<UAnimNextRigVMAssetEditorData>(AnimNextRigVMAsset);
+	return EditorData->GetAutoVMRecompile();
+}
+
+ECompileStatus FAssetCompilationHandler::GetCompileStatus(TSharedRef<Workspace::IWorkspaceEditor> InWorkspaceEditor, const UObject* InAsset) const
+{
+	const UAnimNextRigVMAsset* AnimNextRigVMAsset = CastChecked<UAnimNextRigVMAsset>(InAsset);
+	UAnimNextRigVMAssetEditorData* EditorData = UE::UAF::UncookedOnly::FUtils::GetEditorData<UAnimNextRigVMAssetEditorData>(AnimNextRigVMAsset);
+	if(EditorData->bErrorsDuringCompilation || NumErrors > 0)
+	{
+		return ECompileStatus::Error;
+	}
+
+	if(EditorData->bWarningsDuringCompilation || NumWarnings > 0)
+	{
+		return ECompileStatus::Warning;
+	}
+
+	if(EditorData->IsDirtyForRecompilation())
+	{
+		return ECompileStatus::Dirty;
+	}
+
+	return ECompileStatus::UpToDate;
+}
+
+}
diff --git a/Engine/Plugins/Experimental/UAF/UAF/Source/UAFEditor/Private/Common/AnimNextEntryAssetDefinitions.cpp b/Engine/Plugins/Experimental/UAF/UAF/Source/UAFEditor/Private/Common/AnimNextEntryAssetDefinitions.cpp
new file mode 100644
--- /dev/null
+++ b/Engine/Plugins/Experimental/UAF/UAF/Source/UAFEditor/Private/Common/AnimNextEntryAssetDefinitions.cpp
@@ -0,0 +1,25 @@
+﻿// Copyright Epic Games, Inc. All Rights Reserved.
+
+#include "AnimNextEntryAssetDefinitions.h"
+
+#define LOCTEXT_NAMESPACE "AnimNextAssetDefinitions"
+
+FText UAssetDefinition_AnimNextVariableEntry::GetObjectDisplayNameText(UObject* Object) const
+{
+	UAnimNextRigVMAssetEntry* Parameter = CastChecked<UAnimNextRigVMAssetEntry>(Object);
+	return Parameter->GetDisplayName();
+}
+
+FText UAssetDefinition_AnimNextEventGraphEntry::GetObjectDisplayNameText(UObject* Object) const
+{
+	UAnimNextRigVMAssetEntry* Variable = CastChecked<UAnimNextRigVMAssetEntry>(Object);
+	return Variable->GetDisplayName();
+}
+
+FText UAssetDefinition_AnimNextSharedVariablesEntry::GetObjectDisplayNameText(UObject* Object) const
+{
+	UAnimNextSharedVariablesEntry* SharedVariables = CastChecked<UAnimNextSharedVariablesEntry>(Object);
+	return SharedVariables->GetDisplayName();
+}
+
+#undef LOCTEXT_NAMESPACE
\ No newline at end of file
diff --git a/Engine/Plugins/Experimental/UAF/UAF/Source/UAFEditor/Private/Common/AnimNextRigVMAssetCommands.cpp b/Engine/Plugins/Experimental/UAF/UAF/Source/UAFEditor/Private/Common/AnimNextRigVMAssetCommands.cpp
new file mode 100644
--- /dev/null
+++ b/Engine/Plugins/Experimental/UAF/UAF/Source/UAFEditor/Private/Common/AnimNextRigVMAssetCommands.cpp
@@ -0,0 +1,23 @@
+﻿// Copyright Epic Games, Inc. All Rights Reserved.
+
+#include "Common/AnimNextRigVMAssetCommands.h"
+
+#define LOCTEXT_NAMESPACE "AnimNextRigVMAssetCommands"
+
+namespace UE::UAF
+{
+
+FAnimNextRigVMAssetCommands::FAnimNextRigVMAssetCommands()
+	: TCommands<FAnimNextRigVMAssetCommands>("UAFRigVMAssetCommand", LOCTEXT("UAFRigVMAssetCommands", "UAF Asset Commands"), NAME_None, "AnimNextStyle")
+{
+}
+
+void FAnimNextRigVMAssetCommands::RegisterCommands()
+{
+	UI_COMMAND(FindInAnimNextRigVMAsset, "FindInUAFRigVMAsset", "Search the current UAF Asset.", EUserInterfaceActionType::Button, FInputChord(EModifierKey::Control, EKeys::F));
+	UI_COMMAND(FindAndReplaceInAnimNextRigVMAsset, "FindAndReplaceInUAFRigVMAsset", "Find and Replace across UAF Assets.", EUserInterfaceActionType::Button, FInputChord(EModifierKey::Control | EModifierKey::Shift, EKeys::F));
+}
+
+}
+
+#undef LOCTEXT_NAMESPACE
diff --git a/Engine/Plugins/Experimental/UAF/UAF/Source/UAFEditor/Private/Module/AnimNextModuleFactory.cpp b/Engine/Plugins/Experimental/UAF/UAF/Source/UAFEditor/Private/Module/AnimNextModuleFactory.cpp
new file mode 100644
--- /dev/null
+++ b/Engine/Plugins/Experimental/UAF/UAF/Source/UAFEditor/Private/Module/AnimNextModuleFactory.cpp
@@ -0,0 +1,44 @@
+﻿// Copyright Epic Games, Inc. All Rights Reserved.
+
+#include "Module/AnimNextModuleFactory.h"
+#include "Module/AnimNextModule.h"
+#include "Module/AnimNextModule_EditorData.h"
+#include "UncookedOnlyUtils.h"
+#include "Module/RigUnit_AnimNextModuleEvents.h"
+
+UAnimNextModuleFactory::UAnimNextModuleFactory()
+{
+	bCreateNew = true;
+	bEditAfterNew = true;
+	SupportedClass = UAnimNextModule::StaticClass();
+}
+
+bool UAnimNextModuleFactory::ConfigureProperties()
+{
+	return true;
+}
+
+UObject* UAnimNextModuleFactory::FactoryCreateNew(UClass* Class, UObject* InParent, FName Name, EObjectFlags Flags, UObject* Context, FFeedbackContext* Warn, FName CallingContext)
+{
+	EObjectFlags FlagsToUse = Flags | RF_Public | RF_Standalone | RF_Transactional | RF_LoadCompleted;
+	if(InParent == GetTransientPackage())
+	{
+		FlagsToUse &= ~RF_Standalone;
+	}
+
+	UAnimNextModule* NewModule = NewObject<UAnimNextModule>(InParent, Class, Name, FlagsToUse);
+
+	// Create internal editor data
+	UAnimNextModule_EditorData* EditorData = NewObject<UAnimNextModule_EditorData>(NewModule, TEXT("EditorData"), RF_Transactional);
+	NewModule->EditorData = EditorData;
+	EditorData->bUsesExternalPackages = false;
+	EditorData->Initialize(/*bRecompileVM*/false);
+
+	EditorData->AddEventGraph(FRigUnit_AnimNextPrePhysicsEvent::DefaultEventName, FRigUnit_AnimNextPrePhysicsEvent::StaticStruct());
+
+	// Compile the initial skeleton
+	EditorData->RecompileVM();
+	check(!EditorData->bErrorsDuringCompilation);
+
+	return NewModule;
+}
\ No newline at end of file
diff --git a/Engine/Plugins/Experimental/UAF/UAF/Source/UAFEditor/Private/Variables/AnimNextSharedVariablesAssetDefinition.cpp b/Engine/Plugins/Experimental/UAF/UAF/Source/UAFEditor/Private/Variables/AnimNextSharedVariablesAssetDefinition.cpp
new file mode 100644
--- /dev/null
+++ b/Engine/Plugins/Experimental/UAF/UAF/Source/UAFEditor/Private/Variables/AnimNextSharedVariablesAssetDefinition.cpp
@@ -0,0 +1,153 @@
+﻿// Copyright Epic Games, Inc. All Rights Reserved.
+
+#include "AnimNextSharedVariablesAssetDefinition.h"
+
+#include "AnimNextRigVMAssetEditorData.h"
+#include "ContentBrowserMenuContexts.h"
+#include "Editor.h"
+#include "IWorkspaceEditorModule.h"
+#include "ToolMenus.h"
+#include "UncookedOnlyUtils.h"
+#include "UObject/SavePackage.h"
+#include "Workspace/AnimNextWorkspaceFactory.h"
+
+#define LOCTEXT_NAMESPACE "AnimNextAssetDefinitions"
+
+EAssetCommandResult UAssetDefinition_AnimNextSharedVariables::OpenAssets(const FAssetOpenArgs& OpenArgs) const
+{
+	using namespace UE::Workspace;
+
+	for (UAnimNextSharedVariables* Asset : OpenArgs.LoadObjects<UAnimNextSharedVariables>())
+	{
+		IWorkspaceEditorModule& WorkspaceEditorModule = FModuleManager::Get().LoadModuleChecked<IWorkspaceEditorModule>("WorkspaceEditor");
+		WorkspaceEditorModule.OpenWorkspaceForObject(Asset, EOpenWorkspaceMethod::Default, UAnimNextWorkspaceFactory::StaticClass());
+	}
+
+	return EAssetCommandResult::Handled;
+}
+
+namespace UE::UAF::Editor
+{
+
+static FDelayedAutoRegisterHelper AutoRegisterGraphMenuItems(EDelayedRegisterRunPhase::EndOfEngineInit, []
+{
+	UToolMenus::RegisterStartupCallback(FSimpleMulticastDelegate::FDelegate::CreateLambda([]()
+	{
+		FToolMenuOwnerScoped OwnerScoped(UE_MODULE_NAME);
+
+		UToolMenu* Menu = UE::ContentBrowser::ExtendToolMenu_AssetContextMenu(UAnimNextSharedVariables::StaticClass());
+
+		FToolMenuSection& Section = Menu->FindOrAddSection("GetAssetActions");
+		Section.AddDynamicEntry(NAME_None, FNewToolMenuSectionDelegate::CreateLambda([](FToolMenuSection& InSection)
+		{
+			auto SetExternalPackagingStatus = [](const FToolMenuContext& InContext, bool bInUseExternalPackaging)
+			{
+				UContentBrowserAssetContextMenuContext* AssetContext = InContext.FindContext<UContentBrowserAssetContextMenuContext>();
+				if(AssetContext == nullptr)
+				{
+					return;
+				}
+
+				TArray<UAnimNextRigVMAsset*> Assets;
+				for (const FAssetData& AssetData : AssetContext->SelectedAssets)
+				{
+					if(!AssetData.GetClass()->IsChildOf(UAnimNextRigVMAsset::StaticClass()))
+					{
+						continue;
+					}
+
+					bool bUsesExternalPackages = true;
+					AssetData.GetTagValue<bool>(UAnimNextRigVMAssetEditorData::GetUsesExternalPackagesPropertyName(), bUsesExternalPackages);
+
+					if(bUsesExternalPackages == bInUseExternalPackaging)
+					{
+						// Status is as requested, skip
+						continue;
+					}
+					
+					UAnimNextRigVMAsset* Asset = Cast<UAnimNextRigVMAsset>(AssetData.GetAsset());
+					if(Asset == nullptr)
+					{
+						continue;
+					}
+
+					Assets.Add(Asset);
+
+					UAnimNextRigVMAssetEditorData* EditorData = UncookedOnly::FUtils::GetEditorData<UAnimNextRigVMAssetEditorData>(Asset);
+					if(EditorData == nullptr)
+					{
+						continue;
+					}
+
+					ensure(EditorData->IsUsingExternalPackages() == bUsesExternalPackages);
+				}
+
+				UAnimNextRigVMAssetEditorData::SetUseExternalPackages(Assets, bInUseExternalPackaging);
+			};
+
+			auto ExternalPackagingStatusMatches = [](const FToolMenuContext& InContext, bool bInUsesExternalPackaging)
+			{
+				UContentBrowserAssetContextMenuContext* AssetContext = InContext.FindContext<UContentBrowserAssetContextMenuContext>();
+				if(AssetContext == nullptr)
+				{
+					return false;
+				}
+
+				for (const FAssetData& AssetData : AssetContext->SelectedAssets)
+				{
+					if(!AssetData.GetClass()->IsChildOf(UAnimNextRigVMAsset::StaticClass()))
+					{
+						continue;
+					}
+
+					bool bUsesExternalPackages = true;
+					AssetData.GetTagValue<bool>(UAnimNextRigVMAssetEditorData::GetUsesExternalPackagesPropertyName(), bUsesExternalPackages);
+
+					if(bUsesExternalPackages == bInUsesExternalPackaging)
+					{
+						return true;
+					}
+				}
+
+				return false;
+			};
+
+			FToolUIAction UIAction;
+			UIAction.ExecuteAction = FToolMenuExecuteAction::CreateLambda([&SetExternalPackagingStatus](const FToolMenuContext& InContext)
+			{
+				SetExternalPackagingStatus(InContext, true);
+			});
+			UIAction.IsActionVisibleDelegate = FToolMenuIsActionButtonVisible::CreateLambda([&ExternalPackagingStatusMatches](const FToolMenuContext& InContext)
+			{
+				return ExternalPackagingStatusMatches(InContext, false);
+			});
+
+			InSection.AddMenuEntry(
+				"EnableExternalPackages",
+				LOCTEXT("EnableExternalPackagesLabel", "Use External Packages"),
+				LOCTEXT("EnableExternalPackagesTooltip", "Set the asset(s) to use external packaging for its entries (graphs, variables etc.)\nThis will create the external packages for all entries, add them to version control if enabled and save all packages.\nWarning: This operation cannot be undone, so a connection to version control is recommended."),
+				FSlateIcon(FAppStyle::Get().GetStyleSetName(), "Icons.Edit"),
+				UIAction);
+
+			UIAction.ExecuteAction = FToolMenuExecuteAction::CreateLambda([&SetExternalPackagingStatus](const FToolMenuContext& InContext)
+			{
+				SetExternalPackagingStatus(InContext, false);
+			});
+			UIAction.IsActionVisibleDelegate = FToolMenuIsActionButtonVisible::CreateLambda([&ExternalPackagingStatusMatches](const FToolMenuContext& InContext)
+			{
+				return ExternalPackagingStatusMatches(InContext, true);
+			});
+
+			InSection.AddMenuEntry(
+				"DisableExternalPackages",
+				LOCTEXT("DisableExternalPackagesLabel", "Use Single Package"),
+				LOCTEXT("DisableExternalPackagesTooltip", "Set the asset(s) to use a single package.\nThis will remove any external packages for existing entries, remove them from version control if enabled and save all packages.\nWarning: This operation cannot be undone, so a connection to version control is recommended."),
+				FSlateIcon(FAppStyle::Get().GetStyleSetName(), "Icons.Edit"),
+				UIAction);
+		}));
+	}));
+});
+
+}
+
+#undef LOCTEXT_NAMESPACE
\ No newline at end of file
diff --git a/Engine/Plugins/Experimental/UAF/UAF/Source/UAFEditor/Private/Variables/AnimNextSharedVariablesFactory.cpp b/Engine/Plugins/Experimental/UAF/UAF/Source/UAFEditor/Private/Variables/AnimNextSharedVariablesFactory.cpp
new file mode 100644
--- /dev/null
+++ b/Engine/Plugins/Experimental/UAF/UAF/Source/UAFEditor/Private/Variables/AnimNextSharedVariablesFactory.cpp
@@ -0,0 +1,41 @@
+﻿// Copyright Epic Games, Inc. All Rights Reserved.
+
+#include "Variables/AnimNextSharedVariablesFactory.h"
+#include "Variables/AnimNextSharedVariables.h"
+#include "Variables/AnimNextSharedVariables_EditorData.h"
+#include "Module/RigUnit_AnimNextModuleEvents.h"
+
+UAnimNextSharedVariablesFactory::UAnimNextSharedVariablesFactory()
+{
+	bCreateNew = true;
+	bEditAfterNew = true;
+	SupportedClass = UAnimNextSharedVariables::StaticClass();
+}
+
+bool UAnimNextSharedVariablesFactory::ConfigureProperties()
+{
+	return true;
+}
+
+UObject* UAnimNextSharedVariablesFactory::FactoryCreateNew(UClass* Class, UObject* InParent, FName Name, EObjectFlags Flags, UObject* Context, FFeedbackContext* Warn, FName CallingContext)
+{
+	EObjectFlags FlagsToUse = Flags | RF_Public | RF_Standalone | RF_Transactional | RF_LoadCompleted;
+	if(InParent == GetTransientPackage())
+	{
+		FlagsToUse &= ~RF_Standalone;
+	}
+
+	UAnimNextSharedVariables* NewDataInterface = NewObject<UAnimNextSharedVariables>(InParent, Class, Name, FlagsToUse);
+
+	// Create internal editor data
+	UAnimNextSharedVariables_EditorData* EditorData = NewObject<UAnimNextSharedVariables_EditorData>(NewDataInterface, TEXT("EditorData"), RF_Transactional);
+	NewDataInterface->EditorData = EditorData;
+	EditorData->bUsesExternalPackages = false;
+	EditorData->Initialize(/*bRecompileVM*/false);
+
+	// Compile the initial skeleton
+	EditorData->RecompileVM();
+	check(!EditorData->bErrorsDuringCompilation);
+
+	return NewDataInterface;
+}
\ No newline at end of file
diff --git a/Engine/Plugins/Experimental/UAF/UAF/Source/UAFEditor/Private/Variables/Outliner/VariablesOutlinerCommands.cpp b/Engine/Plugins/Experimental/UAF/UAF/Source/UAFEditor/Private/Variables/Outliner/VariablesOutlinerCommands.cpp
new file mode 100644
--- /dev/null
+++ b/Engine/Plugins/Experimental/UAF/UAF/Source/UAFEditor/Private/Variables/Outliner/VariablesOutlinerCommands.cpp
@@ -0,0 +1,22 @@
+// Copyright Epic Games, Inc. All Rights Reserved.
+
+
+#include "VariablesOutlinerCommands.h"
+
+#define LOCTEXT_NAMESPACE "VariablesOutlinerCommands"
+
+void FVariablesOutlinerCommands::RegisterCommands()
+{
+	UI_COMMAND(AddNewVariable, "Add Variable", "Adds a single new variable to the selected asset.", EUserInterfaceActionType::Button, FInputChord());
+	UI_COMMAND(AddNewVariables, "Add Variable(s)", "Adds variables to assets.\nIf multiple assets are selected, then variables will be added to each.\nIf no assets are selected and there are multiple assets, variables will be added to all assets.", EUserInterfaceActionType::Button, FInputChord());
+
+	UI_COMMAND(FindReferences, "In Project", "Finds all references to the selected variable variables across assets in the project.", EUserInterfaceActionType::Button, FInputChord());
+	UI_COMMAND(FindReferencesInWorkspace, "In Workspace", "Finds all references to the selected variable variables across assets in the current workspace.", EUserInterfaceActionType::Button, FInputChord());
+	UI_COMMAND(FindReferencesInAsset, "In Asset", "Finds all references to the selected variable variables in the currently opened asset.", EUserInterfaceActionType::Button, FInputChord());
+
+	UI_COMMAND(SaveAsset, "Save Asset", "Save the selected asset.", EUserInterfaceActionType::Button, FInputChord());
+
+	UI_COMMAND(ToggleVariableExport, "Toggle Variable Export", "Sets Variable to be either Public or Private.", EUserInterfaceActionType::Button, FInputChord());
+}
+
+#undef LOCTEXT_NAMESPACE
\ No newline at end of file
diff --git a/Engine/Plugins/Experimental/UAF/UAF/Source/UAFEditor/Private/Variables/Outliner/VariablesOutlinerEntryItem.h b/Engine/Plugins/Experimental/UAF/UAF/Source/UAFEditor/Private/Variables/Outliner/VariablesOutlinerEntryItem.h
new file mode 100644
--- /dev/null
+++ b/Engine/Plugins/Experimental/UAF/UAF/Source/UAFEditor/Private/Variables/Outliner/VariablesOutlinerEntryItem.h
@@ -0,0 +1,51 @@
+// Copyright Epic Games, Inc. All Rights Reserved.
+
+#pragma once
+
+#include "ISceneOutlinerTreeItem.h"
+
+class UAnimNextSharedVariablesEntry;
+class UAnimNextVariableEntry;
+
+namespace UE::UAF::Editor
+{
+
+struct FVariablesOutlinerEntryItem : ISceneOutlinerTreeItem
+{
+	static const FSceneOutlinerTreeItemType Type;
+
+	FVariablesOutlinerEntryItem(UAnimNextVariableEntry* InEntry);
+	FVariablesOutlinerEntryItem(const FProperty* InProperty);
+
+	// Begin ISceneOutlinerTreeItem overrides
+	virtual bool IsValid() const override;
+	virtual FSceneOutlinerTreeItemID GetID() const override;
+	virtual FString GetDisplayString() const override;
+	virtual bool CanInteract() const override { return true; }
+	virtual TSharedRef<SWidget> GenerateLabelWidget(ISceneOutliner& Outliner, const STableRow<FSceneOutlinerTreeItemPtr>& InRow) override;
+	virtual FString GetPackageName() const override;
+	// End ISceneOutlinerTreeItem overrides
+
+	// Renames the item to the specified name
+	void Rename(const FText& InNewName) const;
+
+	// Validates the new item name
+	bool ValidateName(const FText& InNewName, FText& OutErrorMessage) const;
+
+	bool HasStructOwner() const { return bStructOwner; }
+
+	// Ptr to the underlying entry if this variable came from an asset
+	TWeakObjectPtr<UAnimNextVariableEntry> WeakEntry;
+
+	// The data interface entry this entry is from, if any
+	TWeakObjectPtr<const UAnimNextSharedVariablesEntry> WeakSharedVariablesEntry;
+
+	// Ptr to the property if this variable came from a struct
+	TFieldPath<const FProperty> PropertyPath;
+
+	bool bStructOwner = false;
+
+	int32 SortValue = 0;
+};
+
+}
diff --git a/Engine/Plugins/Experimental/UAF/UAF/Source/UAFEditor/Private/Variables/Outliner/VariablesOutlinerHierarchy.cpp b/Engine/Plugins/Experimental/UAF/UAF/Source/UAFEditor/Private/Variables/Outliner/VariablesOutlinerHierarchy.cpp
new file mode 100644
--- /dev/null
+++ b/Engine/Plugins/Experimental/UAF/UAF/Source/UAFEditor/Private/Variables/Outliner/VariablesOutlinerHierarchy.cpp
@@ -0,0 +1,233 @@
+// Copyright Epic Games, Inc. All Rights Reserved.
+
+#include "VariablesOutlinerHierarchy.h"
+
+#include "AnimNextRigVMAsset.h"
+#include "ISceneOutlinerMode.h"
+#include "AnimNextRigVMAssetEditorData.h"
+#include "VariablesOutlinerAssetItem.h"
+#include "VariablesOutlinerCategoryItem.h"
+#include "VariablesOutlinerStructSharedVariablesItem.h"
+#include "VariablesOutlinerMode.h"
+#include "VariablesOutlinerEntryItem.h"
+#include "Entries/AnimNextSharedVariablesEntry.h"
+#include "Entries/AnimNextRigVMAssetEntry.h"
+#include "Entries/AnimNextVariableEntry.h"
+#include "Variables/SVariablesView.h"
+#include "Variables/AnimNextSharedVariables.h"
+
+namespace UE::UAF::Editor
+{
+
+TAutoConsoleVariable<bool> GOutputHierarchyDebugInformationCvar(TEXT("VariablesOutliner.OutputHierarchyDebugInformation"), false, TEXT("If enabled will output debug information about the hierarchy generation to the output log")); 
+
+FVariablesOutlinerHierarchy::FVariablesOutlinerHierarchy(ISceneOutlinerMode* Mode)
+	: ISceneOutlinerHierarchy(Mode)
+{
+}
+
+void FVariablesOutlinerHierarchy::CreateItems(TArray<FSceneOutlinerTreeItemPtr>& OutItems) const
+{
+	SVariablesOutliner* Outliner = static_cast<FVariablesOutlinerMode* const>(Mode)->GetOutliner();
+
+	if(Outliner->Assets.Num())
+	{
+		// Process first asset only for now
+		const int32 NumAssets = Outliner->Assets.Num();
+		for(int32 Index = 0; Index < NumAssets; Index++)
+		{
+			const TSoftObjectPtr<UAnimNextRigVMAsset>& SoftAsset = Outliner->Assets[Index];
+			TArray<const FSoftObjectPath> SharedVariableSourcePaths;
+			ProcessAsset(SoftAsset, (NumAssets - Index), OutItems, SharedVariableSourcePaths);
+		}
+	}
+}
+
+FSceneOutlinerTreeItemPtr FVariablesOutlinerHierarchy::FindOrCreateParentItem(const ISceneOutlinerTreeItem& Item, const TMap<FSceneOutlinerTreeItemID, FSceneOutlinerTreeItemPtr>& Items, bool bCreate)
+{
+	uint32 ParentHash = INDEX_NONE;
+	
+	if (const FVariablesOutlinerEntryItem* EntryItem = Item.CastTo<FVariablesOutlinerEntryItem>())
+	{
+		const UAnimNextSharedVariablesEntry* SharedVariablesEntry = EntryItem->WeakSharedVariablesEntry.Get();
+		
+		FSoftObjectPath OwnerPath;
+		if(SharedVariablesEntry != nullptr)
+		{
+			OwnerPath = SharedVariablesEntry->GetObjectPath();
+		}
+		else
+		{
+			UAnimNextRigVMAssetEntry* Entry = EntryItem->WeakEntry.Get();
+			if (Entry == nullptr)
+			{
+				return nullptr;
+			}
+
+			UAnimNextRigVMAsset* Asset = Entry->GetTypedOuter<UAnimNextRigVMAsset>();
+			if (Entry == nullptr)
+			{
+				return nullptr;
+			}
+
+			TSoftObjectPtr<UAnimNextRigVMAsset> SoftObjectPtr(Asset);
+			OwnerPath = Asset;
+			
+		}
+
+		uint32 CategoryHash = INDEX_NONE;
+
+		// Now we have the parent asset hash, determine if we need to be child-ed to a category
+		if (UAnimNextVariableEntry* VariableEntry = EntryItem->WeakEntry.Get())
+		{
+			if (!VariableEntry->GetVariableCategory().IsEmpty())
+			{
+				CategoryHash = GetTypeHash(VariableEntry->GetVariableCategory());
+			}
+		}
+
+		check(OwnerPath.IsValid());
+		ParentHash = CategoryHash != INDEX_NONE ? HashCombine(GetTypeHash(OwnerPath), CategoryHash) : GetTypeHash(OwnerPath);		
+		
+	}
+
+	if (const FVariablesOutlinerCategoryItem* CategoryItem = Item.CastTo<FVariablesOutlinerCategoryItem>())
+	{
+		if (UAnimNextRigVMAsset* Asset = CategoryItem->WeakOwner.Get())
+		{
+			const FSoftObjectPath OwnerPath = Asset;
+			ParentHash = GetTypeHash(OwnerPath);
+		}
+	}
+
+	if (const FSceneOutlinerTreeItemPtr* ParentItem = Items.Find(ParentHash))
+	{
+		return *ParentItem;
+	}
+
+	return nullptr;
+}
+
+void FVariablesOutlinerHierarchy::ProcessAsset(const TSoftObjectPtr<UAnimNextRigVMAsset>& InSoftAsset, int32 Depth, TArray<FSceneOutlinerTreeItemPtr>& OutItems, TArray<const FSoftObjectPath>& OutSharedVariableSourcePaths) const
+{
+	UAnimNextRigVMAsset* Asset = InSoftAsset.Get();
+	if(Asset == nullptr)
+	{
+		return;
+	}
+		
+	const UAnimNextRigVMAssetEditorData* EditorData = UncookedOnly::FUtils::GetEditorData<UAnimNextRigVMAssetEditorData>(Asset);
+	if (EditorData == nullptr)
+	{
+		return;
+	}
+
+	UE_CLOG(GOutputHierarchyDebugInformationCvar.GetValueOnAnyThread(),LogAnimation, Warning, TEXT("ProcessAsset: %s %i"), *InSoftAsset.GetAssetName(), Depth);
+		
+	TArray<FSceneOutlinerTreeItemPtr> ParentItems;
+	EditorData->ForEachEntryOfType<UAnimNextSharedVariablesEntry>([this, &OutSharedVariableSourcePaths, &ParentItems, Depth, &OutItems](UAnimNextSharedVariablesEntry* InSharedVariablesEntry)
+	{
+		SVariablesOutliner* Outliner = static_cast<FVariablesOutlinerMode* const>(Mode)->GetOutliner();
+
+		switch (InSharedVariablesEntry->Type)
+		{
+		case EAnimNextSharedVariablesType::Asset:
+			{
+				if (const UAnimNextSharedVariables* SharedVariables = InSharedVariablesEntry->GetAsset())
+				{
+					OutSharedVariableSourcePaths.Add(SharedVariables);
+				
+					// Would have expected shared variable (assets) to be part of the outer export chain, so should have been handled as regular assets.
+					if (!Outliner->Assets.Contains(SharedVariables))
+					{	
+						TArray<const FSoftObjectPath> SharedVariableSourcePaths;
+						ProcessAsset( const_cast<UAnimNextSharedVariables*>(SharedVariables), Depth - 1, ParentItems, SharedVariableSourcePaths);
+					}
+				}				
+			}
+			break;
+
+		case EAnimNextSharedVariablesType::Struct:
+			{				
+				if (const UScriptStruct* Struct = InSharedVariablesEntry->GetStruct())
+				{
+					if (FSceneOutlinerTreeItemPtr Item = Mode->CreateItemFor<FVariablesOutlinerStructSharedVariablesItem>(InSharedVariablesEntry))
+					{
+						OutItems.Add(Item);
+					}
+					
+					UE_CLOG(GOutputHierarchyDebugInformationCvar.GetValueOnAnyThread(), LogAnimation, Warning, TEXT("~ FVariablesOutlinerStructSharedVariablesItem: %s %i"), *InSharedVariablesEntry->GetObjectPath().ToString(), Depth);
+					
+					OutSharedVariableSourcePaths.Add(Struct);
+					int32 Index = 0;
+					for (TFieldIterator<FProperty> It(Struct); It; ++It)
+					{
+						const FProperty* Property = *It;
+						if (FSceneOutlinerTreeItemPtr Item = Mode->CreateItemFor<FVariablesOutlinerEntryItem>(Property))
+						{
+							FVariablesOutlinerEntryItem* EntryItem = Item->CastTo<FVariablesOutlinerEntryItem>();
+							EntryItem->WeakSharedVariablesEntry = InSharedVariablesEntry;
+							EntryItem->SortValue = Index++;
+							OutItems.Add(Item);
+						}
+					}
+				}				
+			}
+			break;
+		}		
+		
+		return true;
+	});
+
+	OutItems.Append(ParentItems);
+
+	TMap<FString, FSceneOutlinerTreeItemPtr> CategoryItems;
+
+	for (int32 CategoryIndex = 0; CategoryIndex < EditorData->VariableAndFunctionCategories.Num(); ++CategoryIndex)
+	{
+		const FString& CategoryName = EditorData->VariableAndFunctionCategories[CategoryIndex];
+		if (FSceneOutlinerTreeItemPtr Item = Mode->CreateItemFor<FVariablesOutlinerCategoryItem>(FVariablesOutlinerCategoryItem::FItemData{Asset, CategoryName}))
+		{
+			Item->CastTo<FVariablesOutlinerCategoryItem>()->SortValue = CategoryIndex;
+			CategoryItems.Add(CategoryName, Item);
+
+			UE_CLOG(GOutputHierarchyDebugInformationCvar.GetValueOnAnyThread(),LogAnimation, Warning, TEXT("+ FVariablesOutlinerCategoryItem: %s %i"), *CategoryName, CategoryIndex);
+		}
+	}
+
+	// Adjust asset depth to ensure SharedVariables are always pinned to 0 depth
+	const int32 AssetDepth = Asset->GetClass() && Asset->GetClass() == UAnimNextSharedVariables::StaticClass() ? 0 : Depth;
+
+	constexpr int32 NumberOfVariableSortValueEntriesInCategory = 1024;
+
+	if (FSceneOutlinerTreeItemPtr Item = Mode->CreateItemFor<FVariablesOutlinerAssetItem>(FVariablesOutlinerAssetItem::FItemData{InSoftAsset, AssetDepth, OutSharedVariableSourcePaths}))
+	{
+		OutItems.Add(Item);
+
+		UE_CLOG(GOutputHierarchyDebugInformationCvar.GetValueOnAnyThread(), LogAnimation, Warning, TEXT("* FVariablesOutlinerAssetItem: %s %i"), *InSoftAsset.GetAssetName(), AssetDepth);
+
+		int32 Index = 1;
+		EditorData->ForEachEntryOfType<UAnimNextVariableEntry>([this, &OutItems, &CategoryItems, &Index, EditorData](UAnimNextVariableEntry* InVariable)
+		{
+			if (FSceneOutlinerTreeItemPtr Item = Mode->CreateItemFor<FVariablesOutlinerEntryItem>(CastChecked<UAnimNextVariableEntry>(InVariable)))
+			{
+				FVariablesOutlinerEntryItem* EntryItem = Item->CastTo<FVariablesOutlinerEntryItem>();
+				const int32 CategoryIndex = EditorData->VariableAndFunctionCategories.IndexOfByKey(InVariable->Category);
+				EntryItem->SortValue = (CategoryIndex != INDEX_NONE ? ((CategoryIndex + 1) * NumberOfVariableSortValueEntriesInCategory) : 0) + Index++;
+				OutItems.Add(Item);
+				
+				OutItems.Add(Item);
+				UE_CLOG(GOutputHierarchyDebugInformationCvar.GetValueOnAnyThread(), LogAnimation, Warning, TEXT("\t- FVariablesOutlinerEntryItem: %s %i"), *InVariable->GetEntryName().ToString(), Item->CastTo<FVariablesOutlinerEntryItem>()->SortValue);
+
+				// Add category item on demand (to ensure we do not end up with empty category item hierarchies)
+				if (!InVariable->Category.IsEmpty())
+				{
+					FSceneOutlinerTreeItemPtr& CategoryItemToAdd = CategoryItems.FindChecked(InVariable->Category);
+					OutItems.AddUnique(CategoryItemToAdd);
+				}
+			}
+			return true;
+		});
+	}
+}
+}
diff --git a/Engine/Plugins/Experimental/UAF/UAF/Source/UAFEditor/Private/Variables/Outliner/VariablesOutlinerHierarchy.h b/Engine/Plugins/Experimental/UAF/UAF/Source/UAFEditor/Private/Variables/Outliner/VariablesOutlinerHierarchy.h
new file mode 100644
--- /dev/null
+++ b/Engine/Plugins/Experimental/UAF/UAF/Source/UAFEditor/Private/Variables/Outliner/VariablesOutlinerHierarchy.h
@@ -0,0 +1,30 @@
+// Copyright Epic Games, Inc. All Rights Reserved.
+
+#pragma once
+
+#include "CoreMinimal.h"
+#include "ISceneOutlinerHierarchy.h"
+
+class UAnimNextRigVMAssetEditorData;
+class UAnimNextRigVMAsset;
+
+namespace UE::UAF::Editor
+{
+
+class FVariablesOutlinerHierarchy : public ISceneOutlinerHierarchy
+{
+public:
+	FVariablesOutlinerHierarchy(ISceneOutlinerMode* Mode);
+	FVariablesOutlinerHierarchy(const FVariablesOutlinerHierarchy&) = delete;
+	FVariablesOutlinerHierarchy& operator=(const FVariablesOutlinerHierarchy&) = delete;
+
+	// Begin ISceneOutlinerHierarchy overrides
+	virtual void CreateItems(TArray<FSceneOutlinerTreeItemPtr>& OutItems) const override;
+	virtual void CreateChildren(const FSceneOutlinerTreeItemPtr& Item, TArray<FSceneOutlinerTreeItemPtr>& OutChildren) const override {}
+	virtual FSceneOutlinerTreeItemPtr FindOrCreateParentItem(const ISceneOutlinerTreeItem& Item, const TMap<FSceneOutlinerTreeItemID, FSceneOutlinerTreeItemPtr>& Items, bool bCreate = false) override;
+	// End ISceneOutlinerHierarchy overrides
+protected:
+	void ProcessAsset(const TSoftObjectPtr<UAnimNextRigVMAsset>& InSoftAsset, int32 Depth, TArray<FSceneOutlinerTreeItemPtr>& OutItems, TArray<const FSoftObjectPath>& OutSharedVariableSourcePaths) const;
+};
+
+}
\ No newline at end of file
