UAF: Moved/renamed AnimNext and AnimNextAnimGraph plugins

diff --git a/Engine/Plugins/Experimental/UAF/UAF/Source/UAFEditor/Private/Variables/Outliner/VariablesOutlinerMode.cpp b/Engine/Plugins/Experimental/UAF/UAF/Source/UAFEditor/Private/Variables/Outliner/VariablesOutlinerMode.cpp
new file mode 100644
--- /dev/null
+++ b/Engine/Plugins/Experimental/UAF/UAF/Source/UAFEditor/Private/Variables/Outliner/VariablesOutlinerMode.cpp
@@ -0,0 +1,1065 @@
+// Copyright Epic Games, Inc. All Rights Reserved.
+
+#include "VariablesOutlinerMode.h"
+
+#include "ISourceControlModule.h"
+#include "WorkspaceItemMenuContext.h"
+#include "VariablesOutlinerHierarchy.h"
+#include "VariablesOutlinerEntryItem.h"
+#include "Toolkits/AssetEditorToolkitMenuContext.h"
+#include "AnimNextRigVMAssetEditorData.h"
+#include "IWorkspaceEditor.h"
+#include "AnimNextEditorModule.h"
+#include "AnimNextEditorStyle.h"
+#include "AssetRegistry/AssetRegistryModule.h"
+#include "ScopedTransaction.h"
+#include "ToolMenus.h"
+#include "Framework/Commands/GenericCommands.h"
+#include "Entries/AnimNextRigVMAssetEntry.h"
+#include "Variables/AnimNextVariableItemMenuContext.h"
+#include "AnimNextRigVMAsset.h"
+#include "ContentBrowserModule.h"
+#include "FileHelpers.h"
+#include "StructViewerFilter.h"
+#include "StructViewerModule.h"
+#include "VariablesOutlinerAssetItem.h"
+#include "VariablesOutlinerCommands.h"
+#include "VariablesOutlinerStructSharedVariablesItem.h"
+#include "VariablesOutlinerCategoryItem.h"
+#include "VariablesOutlinerDragDrop.h"
+#include "Common/AnimNextAssetFindReplaceVariables.h"
+#include "Common/GraphEditorSchemaActions.h"
+#include "Entries/AnimNextVariableEntry.h"
+#include "Variables/SVariablesView.h"
+#include "Entries/AnimNextSharedVariablesEntry.h"
+#include "Styling/SlateIconFinder.h"
+#include "Variables/AnimNextSharedVariables.h"
+#include "Variables/AnimNextVariableSettings.h"
+#include "Widgets/Docking/SDockTab.h"
+
+#define LOCTEXT_NAMESPACE "VariablesOutlinerMode"
+
+namespace UE::UAF::Editor
+{
+
+FVariablesOutlinerMode::FVariablesOutlinerMode(SVariablesOutliner* InVariablesOutliner, const TSharedRef<UE::Workspace::IWorkspaceEditor>& InWorkspaceEditor)
+	: ISceneOutlinerMode(InVariablesOutliner)
+	, WeakWorkspaceEditor(InWorkspaceEditor)
+{
+	CommandList = MakeShared<FUICommandList>();
+
+	static const FName MenuName("VariablesOutliner.AddVariablesMenu");
+	UToolMenus* ToolMenus = UToolMenus::Get();
+	if (!ToolMenus->IsMenuRegistered(MenuName))
+	{
+		FToolMenuOwnerScoped ToolMenuOwnerScope(this);
+
+		if (UToolMenu* Menu = ToolMenus->RegisterMenu(MenuName))
+		{
+			Menu->AddDynamicSection(TEXT("Variables"), FNewToolMenuDelegate::CreateLambda([](UToolMenu* InMenu)
+			{
+				PopulateNewVariableToolMenuEntries(InMenu, false);				
+			}));
+		}
+	}
+}
+
+void FVariablesOutlinerMode::Rebuild()
+{
+	Hierarchy = CreateHierarchy();
+}
+
+TSharedPtr<SWidget> FVariablesOutlinerMode::CreateContextMenu()
+{
+	static const FName MenuName("VariablesOutliner.ItemContextMenu");
+	UToolMenus* ToolMenus = UToolMenus::Get();
+	if (!ToolMenus->IsMenuRegistered(MenuName))
+	{
+		FToolMenuOwnerScoped ToolMenuOwnerScope(this);
+		if (UToolMenu* Menu = ToolMenus->RegisterMenu(MenuName))
+		{
+			Menu->AddDynamicSection(TEXT("Assets"), FNewToolMenuDelegate::CreateLambda([](UToolMenu* InMenu)
+			{
+				const UAssetEditorToolkitMenuContext* EditorContext = InMenu->FindContext<UAssetEditorToolkitMenuContext>();
+				const UAnimNextVariableItemMenuContext* MenuContext = InMenu->FindContext<UAnimNextVariableItemMenuContext>();
+				if(EditorContext == nullptr || MenuContext == nullptr)
+				{
+					return;
+				}
+
+				const int32 NumSelectedVariables = MenuContext->WeakEntries.Num();
+				
+				/** Asset specific actions */
+				FToolMenuSection& AssetSection = InMenu->AddSection("Asset", LOCTEXT("AssetSectionLabel", "Asset"));
+				{
+					AssetSection.AddMenuEntryWithCommandList(FVariablesOutlinerCommands::Get().SaveAsset, MenuContext->WeakCommandList.Pin(), TAttribute<FText>(), TAttribute<FText>(), FSlateIcon(FAppStyle::Get().GetStyleSetName(), "AssetEditor.SaveAsset"));
+					AssetSection.AddMenuEntryWithCommandList(FGenericCommands::Get().Paste, MenuContext->WeakCommandList.Pin());
+				}
+
+				/** Variable specific actions */
+				FToolMenuSection& VariablesSection = InMenu->AddSection("Variables", LOCTEXT("VariablesSectionLabel", "Variables"));
+				{
+					// Populates entries for adding of different variable types/sources					
+					PopulateNewVariableToolMenuEntries(InMenu, true);
+					
+					VariablesSection.AddSubMenu("FindReferencesMenu", LOCTEXT("FindReferencesMenuLabel", "Find References..."), LOCTEXT("FindReferencesMenuToolTip", "Find references to selected variables by different means."),
+					FNewToolMenuDelegate::CreateLambda(
+						[MenuContext](UToolMenu* InSubmenu)
+							{
+								InSubmenu->AddMenuEntry(NAME_None, FToolMenuEntry::InitMenuEntryWithCommandList(FVariablesOutlinerCommands::Get().FindReferences, MenuContext->WeakCommandList.Pin(), TAttribute<FText>(), TAttribute<FText>(), UAnimNextAssetFindReplaceVariables::GetIconFromSearchScope(ESearchScope::Global)));
+
+								InSubmenu->AddMenuEntry(NAME_None, FToolMenuEntry::InitMenuEntryWithCommandList(FVariablesOutlinerCommands::Get().FindReferencesInWorkspace, MenuContext->WeakCommandList.Pin(), TAttribute<FText>(), TAttribute<FText>(),
+									UAnimNextAssetFindReplaceVariables::GetIconFromSearchScope(ESearchScope::Workspace)));
+							
+								InSubmenu->AddMenuEntry(NAME_None, FToolMenuEntry::InitMenuEntryWithCommandList(FVariablesOutlinerCommands::Get().FindReferencesInAsset, MenuContext->WeakCommandList.Pin(), TAttribute<FText>(), TAttribute<FText>(),
+									UAnimNextAssetFindReplaceVariables::GetIconFromSearchScope(ESearchScope::Asset)));
+							}),
+							FUIAction(FExecuteAction(), FCanExecuteAction(), FGetActionCheckState(), FIsActionButtonVisible::CreateLambda([MenuContext]() -> bool
+							{
+								if (TSharedPtr<UE::UAF::Editor::SVariablesOutliner> SharedOutliner = MenuContext->WeakOutliner.Pin())
+								{
+									if (const FVariablesOutlinerMode* Mode = static_cast<const FVariablesOutlinerMode*>(SharedOutliner->GetMode()))
+									{
+										return Mode->CanFindReferences();
+									}
+								}
+
+								return false;
+							})),
+							EUserInterfaceActionType::Button,
+							false,
+							FSlateIcon(FAppStyle::GetAppStyleSetName(), TEXT("Icons.Search"))
+						);
+
+					VariablesSection.AddMenuEntryWithCommandList(FGenericCommands::Get().Rename, MenuContext->WeakCommandList.Pin());
+					VariablesSection.AddMenuEntryWithCommandList(FGenericCommands::Get().Delete, MenuContext->WeakCommandList.Pin());
+					VariablesSection.AddMenuEntryWithCommandList(FGenericCommands::Get().Duplicate, MenuContext->WeakCommandList.Pin());
+					
+					VariablesSection.AddMenuEntryWithCommandList(FGenericCommands::Get().Copy, MenuContext->WeakCommandList.Pin());
+
+					if (NumSelectedVariables == 1)
+					{
+						if (const UAnimNextVariableEntry* VariableEntry = Cast<UAnimNextVariableEntry>(MenuContext->WeakEntries[0].Get()))
+						{
+							VariablesSection.AddMenuEntryWithCommandList(FVariablesOutlinerCommands::Get().ToggleVariableExport, MenuContext->WeakCommandList.Pin(), VariableEntry->GetExportAccessSpecifier() == EAnimNextExportAccessSpecifier::Private ? LOCTEXT("MakePublicMenuLabel", "Expose as Public") : LOCTEXT("MakePrivateMenuLabel", "Make Private"), TAttribute<FText>(), FSlateIcon(FAppStyle::Get().GetStyleSetName(), VariableEntry->GetExportAccessSpecifier() == EAnimNextExportAccessSpecifier::Private ? "Level.VisibleIcon16x" : "Level.NotVisibleHighlightIcon16x"));
+						}
+					}
+				}
+
+				if (TSharedPtr<SVariablesOutliner> Outliner = MenuContext->WeakOutliner.Pin())
+				{
+					Outliner->AddSourceControlMenuOptions(InMenu);
+				}
+			}));
+		}
+	}
+
+	{
+		UAnimNextVariableItemMenuContext* MenuContext = NewObject<UAnimNextVariableItemMenuContext>();
+		MenuContext->WeakWorkspaceEditor = WeakWorkspaceEditor;
+		MenuContext->WeakOutliner = StaticCastSharedRef<SVariablesOutliner>(SceneOutliner->AsShared());
+		MenuContext->WeakCommandList = CommandList;
+		TArray<FSceneOutlinerTreeItemPtr> SelectedItems = GetOutliner()->GetSelectedItems();
+		for(const FSceneOutlinerTreeItemPtr& Item : SelectedItems)
+		{
+			if (const FVariablesOutlinerAssetItem* AssetItem = Item->CastTo<FVariablesOutlinerAssetItem>())
+			{
+				UAnimNextRigVMAsset* Asset = AssetItem->SoftAsset.Get();
+				if(Asset == nullptr)
+				{
+					continue;
+				}
+
+				UAnimNextRigVMAssetEditorData* EditorData = UncookedOnly::FUtils::GetEditorData<UAnimNextRigVMAssetEditorData>(Asset);
+				if(EditorData == nullptr)
+				{
+					continue;
+				}
+				
+				MenuContext->WeakEditorDatas.Add(EditorData);
+			}
+			else if (const FVariablesOutlinerEntryItem* EntryItem = Item->CastTo<FVariablesOutlinerEntryItem>())
+			{
+				UAnimNextRigVMAssetEntry* Entry = EntryItem->WeakEntry.Get();
+				if(Entry == nullptr)
+				{
+					continue;
+				}
+
+				MenuContext->WeakEntries.Add(Entry);
+			}
+		}
+
+		FToolMenuContext Context;
+		Context.AddObject(MenuContext);
+		WeakWorkspaceEditor.Pin()->InitToolMenuContext(Context);
+		return UToolMenus::Get()->GenerateWidget(MenuName, Context);
+	}
+}
+
+FReply FVariablesOutlinerMode::OnKeyDown(const FKeyEvent& InKeyEvent)
+{
+	if (CommandList->ProcessCommandBindings(InKeyEvent))
+	{
+		return FReply::Handled();
+	}
+
+	return FReply::Unhandled();
+}
+
+void FVariablesOutlinerMode::OnItemClicked(FSceneOutlinerTreeItemPtr Item)
+{
+	const FSceneOutlinerItemSelection& Selection = SceneOutliner->GetSelection();
+	HandleItemSelection(Selection);
+}
+
+void FVariablesOutlinerMode::HandleItemSelection(const FSceneOutlinerItemSelection& Selection) const
+{
+	if (TSharedPtr<UE::Workspace::IWorkspaceEditor> WorkspaceEditor = WeakWorkspaceEditor.Pin())
+	{
+		TArray<FSceneOutlinerTreeItemPtr> SelectedItems;
+		Selection.Get(SelectedItems);
+		TArray<UObject*> EntriesToShow;
+		EntriesToShow.Reserve(SelectedItems.Num());
+		for(const FSceneOutlinerTreeItemPtr& Item : SelectedItems)
+		{
+			if (FVariablesOutlinerEntryItem* VariablesItem = Item->CastTo<FVariablesOutlinerEntryItem>())
+			{
+				if(UAnimNextVariableEntry* VariableEntry = VariablesItem->WeakEntry.Get())
+				{
+					EntriesToShow.Add(VariableEntry);
+				}
+			}
+			else if(const FVariablesOutlinerStructSharedVariablesItem* SharedVariablesItem = Item->CastTo<FVariablesOutlinerStructSharedVariablesItem>())
+			{
+				if(const UAnimNextSharedVariablesEntry* SharedVariablesEntry = SharedVariablesItem->WeakEntry.Get())
+				{
+					EntriesToShow.Add(const_cast<UAnimNextSharedVariablesEntry*>(SharedVariablesEntry));
+				}
+			}
+		}
+
+		WorkspaceEditor->SetDetailsObjects(EntriesToShow);
+	}
+}
+
+void FVariablesOutlinerMode::OnItemSelectionChanged(FSceneOutlinerTreeItemPtr Item, ESelectInfo::Type SelectionType, const FSceneOutlinerItemSelection& Selection)
+{	
+	if (TSharedPtr<UE::Workspace::IWorkspaceEditor> SharedWorkspaceEditor = WeakWorkspaceEditor.Pin())
+	{
+		SharedWorkspaceEditor->SetGlobalSelection(SceneOutliner->AsShared(), UE::Workspace::FOnClearGlobalSelection::CreateRaw(this, &FVariablesOutlinerMode::ResetOutlinerSelection));
+	}
+	
+	HandleItemSelection(Selection);
+}
+
+bool FVariablesOutlinerMode::CanRenameItem(const ISceneOutlinerTreeItem& Item) const
+{
+	if (const FVariablesOutlinerEntryItem* EntryItem = Item.CastTo<FVariablesOutlinerEntryItem>())
+	{
+		return !EntryItem->HasStructOwner();
+	}
+	else if (Item.IsA<FVariablesOutlinerAssetItem>())
+	{
+		return true;
+	}
+	else if (Item.IsA<FVariablesOutlinerCategoryItem>())
+	{
+		return true;
+	}
+	return false;
+}
+
+void FVariablesOutlinerMode::BindCommands(const TSharedRef<FUICommandList>& OutCommandList)
+{
+	CommandList->MapAction( 
+		FGenericCommands::Get().Rename, 
+		FExecuteAction::CreateRaw(this, &FVariablesOutlinerMode::Rename),
+		FCanExecuteAction::CreateRaw(this, &FVariablesOutlinerMode::CanRename),
+		FIsActionChecked(),
+		FIsActionButtonVisible::CreateRaw(this, &FVariablesOutlinerMode::CanRename)
+	);
+
+	CommandList->MapAction(
+		FGenericCommands::Get().Delete,
+		FExecuteAction::CreateRaw(this, &FVariablesOutlinerMode::Delete),
+		FCanExecuteAction::CreateRaw(this, &FVariablesOutlinerMode::CanDelete),
+		FIsActionChecked(),
+		FIsActionButtonVisible::CreateRaw(this, &FVariablesOutlinerMode::CanDelete)	
+		);
+
+	// [TODO] implement copy/paste for variables
+	CommandList->MapAction(
+		FGenericCommands::Get().Copy, 
+		FExecuteAction::CreateRaw(this, &FVariablesOutlinerMode::Copy),
+		FCanExecuteAction::CreateRaw(this, &FVariablesOutlinerMode::CanCopy),
+		FIsActionChecked(),
+		FIsActionButtonVisible::CreateRaw(this, &FVariablesOutlinerMode::CanCopy)
+	);
+
+	CommandList->MapAction( 
+		FGenericCommands::Get().Paste, 
+		FExecuteAction::CreateRaw(this, &FVariablesOutlinerMode::Paste),
+		FCanExecuteAction::CreateRaw(this, &FVariablesOutlinerMode::CanPaste),
+		FIsActionChecked(),
+		FIsActionButtonVisible::CreateRaw(this, &FVariablesOutlinerMode::CanPaste)
+	);
+
+	CommandList->MapAction(
+		FGenericCommands::Get().Duplicate,
+		FExecuteAction::CreateRaw(this, &FVariablesOutlinerMode::Duplicate),
+		FCanExecuteAction::CreateRaw(this, &FVariablesOutlinerMode::CanDuplicate),
+		FIsActionChecked(),
+		FIsActionButtonVisible::CreateRaw(this, &FVariablesOutlinerMode::CanDuplicate)
+	);
+
+	CommandList->MapAction(
+		FVariablesOutlinerCommands::Get().FindReferences,
+		FExecuteAction::CreateRaw(this, &FVariablesOutlinerMode::FindReferences, ESearchScope::Global),
+		FCanExecuteAction::CreateRaw(this, &FVariablesOutlinerMode::CanFindReferences),
+		FIsActionChecked(),
+		FIsActionButtonVisible::CreateRaw(this, &FVariablesOutlinerMode::IsFindReferencesVisible, ESearchScope::Global)
+	);
+	
+	CommandList->MapAction(
+		FVariablesOutlinerCommands::Get().FindReferencesInWorkspace,
+		FExecuteAction::CreateRaw(this, &FVariablesOutlinerMode::FindReferences, ESearchScope::Workspace),
+		FCanExecuteAction::CreateRaw(this, &FVariablesOutlinerMode::CanFindReferences),
+		FIsActionChecked(),
+	FIsActionButtonVisible::CreateRaw(this, &FVariablesOutlinerMode::IsFindReferencesVisible, ESearchScope::Workspace)
+	);
+	
+	CommandList->MapAction(
+		FVariablesOutlinerCommands::Get().FindReferencesInAsset,
+		FExecuteAction::CreateRaw(this, &FVariablesOutlinerMode::FindReferences, ESearchScope::Asset),
+		FCanExecuteAction::CreateRaw(this, &FVariablesOutlinerMode::CanFindReferences),
+		FIsActionChecked(),
+		FIsActionButtonVisible::CreateRaw(this, &FVariablesOutlinerMode::IsFindReferencesVisible, ESearchScope::Asset)
+	);
+
+	CommandList->MapAction(
+		FVariablesOutlinerCommands::Get().SaveAsset,
+		FExecuteAction::CreateRaw(this, &FVariablesOutlinerMode::SaveAsset),
+		FCanExecuteAction::CreateRaw(this, &FVariablesOutlinerMode::CanSaveAsset),
+		FIsActionChecked(),
+		FIsActionButtonVisible::CreateRaw(this, &FVariablesOutlinerMode::IsAsset)
+		);
+	
+	CommandList->MapAction(
+		FVariablesOutlinerCommands::Get().ToggleVariableExport,
+		FExecuteAction::CreateRaw(this, &FVariablesOutlinerMode::ToggleVariableExport),
+		FCanExecuteAction::CreateRaw(this, &FVariablesOutlinerMode::CanToggleVariableExport),
+		FIsActionChecked(),
+		FIsActionButtonVisible::CreateRaw(this, &FVariablesOutlinerMode::CanToggleVariableExport)
+		);
+}
+
+TUniquePtr<ISceneOutlinerHierarchy> FVariablesOutlinerMode::CreateHierarchy()
+{
+	return MakeUnique<FVariablesOutlinerHierarchy>(this);
+}
+
+void FVariablesOutlinerMode::ResetOutlinerSelection() const
+{
+	SceneOutliner->ClearSelection();
+}
+
+SVariablesOutliner* FVariablesOutlinerMode::GetOutliner() const
+{
+	return static_cast<SVariablesOutliner*>(SceneOutliner);
+}
+
+void FVariablesOutlinerMode::Rename() const
+{
+	const FSceneOutlinerItemSelection& Selection = SceneOutliner->GetSelection();
+	if (Selection.Num() == 1)
+	{
+		FSceneOutlinerTreeItemPtr ItemToRename = Selection.SelectedItems[0].Pin();
+
+		if (ItemToRename.IsValid() && CanRenameItem(*ItemToRename) && ItemToRename->CanInteract())
+		{
+			SceneOutliner->SetPendingRenameItem(ItemToRename);
+			SceneOutliner->ScrollItemIntoView(ItemToRename);
+		}
+	}
+}
+
+bool FVariablesOutlinerMode::CanRename() const
+{
+	const FSceneOutlinerItemSelection& Selection = SceneOutliner->GetSelection();
+	if (Selection.Num() == 1)
+	{
+		FSceneOutlinerTreeItemPtr ItemToRename = Selection.SelectedItems[0].Pin();
+		return ItemToRename.IsValid() && CanRenameItem(*ItemToRename) && ItemToRename->CanInteract();
+	}
+	return false;
+}
+
+void FVariablesOutlinerMode::Delete() const
+{
+	int32 NumEntries = 0;
+	TMap<UAnimNextRigVMAssetEditorData*, TArray<UAnimNextRigVMAssetEntry*>> EntriesToDeletePerAsset;
+	TArray<FSceneOutlinerTreeItemPtr> SelectedItems = SceneOutliner->GetSelectedItems();
+	for(const FSceneOutlinerTreeItemPtr& Item : SelectedItems)
+	{
+		if (const FVariablesOutlinerEntryItem* VariablesItem = Item->CastTo<FVariablesOutlinerEntryItem>())
+		{
+			if (!VariablesItem->HasStructOwner())
+			{
+				UAnimNextVariableEntry* VariableEntry = VariablesItem->WeakEntry.Get();
+				const UAnimNextSharedVariablesEntry* SharedVariablesEntry = VariablesItem->WeakSharedVariablesEntry.Get();
+				if(VariableEntry == nullptr || SharedVariablesEntry != nullptr)	// Cant delete variables in other data interfaces
+				{
+					continue;
+				}
+
+				UAnimNextRigVMAssetEditorData* EditorData = VariableEntry->GetTypedOuter<UAnimNextRigVMAssetEditorData>();
+				if(EditorData == nullptr)
+				{
+					continue;
+				}
+
+				TArray<UAnimNextRigVMAssetEntry*>& EntriesToDelete = EntriesToDeletePerAsset.FindOrAdd(EditorData);
+				EntriesToDelete.Add(VariableEntry);
+				NumEntries++;
+			}
+		}
+		else if (const FVariablesOutlinerStructSharedVariablesItem* SharedVariablesItem = Item->CastTo<FVariablesOutlinerStructSharedVariablesItem>())
+		{
+			UAnimNextSharedVariablesEntry* SharedVariablesEntry = SharedVariablesItem->WeakEntry.Get();
+			if(SharedVariablesEntry == nullptr)
+			{
+				continue;
+			}
+
+			UAnimNextRigVMAssetEditorData* EditorData = SharedVariablesEntry->GetTypedOuter<UAnimNextRigVMAssetEditorData>();
+			if(EditorData == nullptr)
+			{
+				continue;
+			}
+
+			TArray<UAnimNextRigVMAssetEntry*>& EntriesToDelete = EntriesToDeletePerAsset.FindOrAdd(EditorData);
+			EntriesToDelete.Add(SharedVariablesEntry);
+			NumEntries++;
+		}
+	}
+
+	if(NumEntries > 0)
+	{
+		// [TODO] prompt user when deletion impacts assets not part of the current workspace (local context), as these have a hidden (read unknown) impact
+		FScopedTransaction Transaction(FText::Format(LOCTEXT("DeleteVariablesFormat", "Delete {0}|plural(one=variable, other=variables)"), NumEntries));
+		for(const TPair<UAnimNextRigVMAssetEditorData*, TArray<UAnimNextRigVMAssetEntry*>>& EntriesPair : EntriesToDeletePerAsset)
+		{
+			for (UAnimNextRigVMAssetEntry* Entry : EntriesPair.Value)
+			{
+				UncookedOnly::FUtils::DeleteVariable(CastChecked<UAnimNextVariableEntry>(Entry), true, true);
+			}
+		}
+	}
+}
+
+bool FVariablesOutlinerMode::CanDelete() const
+{
+	TArray<FSceneOutlinerTreeItemPtr> SelectedItems = SceneOutliner->GetSelectedItems();
+
+	bool bCanDelete = false;
+	for(const FSceneOutlinerTreeItemPtr& Item : SelectedItems)
+	{
+		if (const FVariablesOutlinerEntryItem* VariablesItem = Item->CastTo<FVariablesOutlinerEntryItem>())
+		{
+			if (!VariablesItem->HasStructOwner())
+			{
+				bCanDelete = true;
+			}
+		}
+	}
+	
+	return bCanDelete;
+}
+
+void FVariablesOutlinerMode::Copy() const
+{
+	const FSceneOutlinerItemSelection& Selection = SceneOutliner->GetSelection();
+	Selection.ForEachItem([](const FSceneOutlinerTreeItemPtr& SelectedItem)
+	{
+		
+	});
+}
+
+bool FVariablesOutlinerMode::CanCopy() const
+{
+	// [TODO] implement copy behaviour
+	return false;
+}
+
+void FVariablesOutlinerMode::Paste() const
+{
+}
+
+bool FVariablesOutlinerMode::CanPaste() const
+{
+	// [TODO] implement paste behaviour
+	return false;
+}
+
+void FVariablesOutlinerMode::ToggleVariableExport() const
+{
+	const FSceneOutlinerItemSelection& Selection = SceneOutliner->GetSelection();
+	if (Selection.Num())
+	{		
+		FScopedTransaction Transaction(LOCTEXT("SetAccessSpecifierTransaction", "Set Access Specifier"));
+		Selection.ForEachItem<FVariablesOutlinerEntryItem>([](const FVariablesOutlinerEntryItem& EntryItem)
+		{			
+			if (IAnimNextRigVMExportInterface* Export = Cast<IAnimNextRigVMExportInterface>(EntryItem.WeakEntry.Get()))
+			{
+				if (UAnimNextRigVMAsset* OuterAsset = EntryItem.WeakEntry.Get()->GetTypedOuter<UAnimNextRigVMAsset>())
+				{
+					if (OuterAsset->GetClass() && OuterAsset->GetClass() != UAnimNextSharedVariables::StaticClass())
+					{
+						Export->SetExportAccessSpecifier(Export->GetExportAccessSpecifier() == EAnimNextExportAccessSpecifier::Public ? EAnimNextExportAccessSpecifier::Private : EAnimNextExportAccessSpecifier::Public);
+					}
+				}
+			}
+		});
+	}
+}
+
+bool FVariablesOutlinerMode::CanToggleVariableExport() const
+{
+	const FSceneOutlinerItemSelection& Selection = SceneOutliner->GetSelection();
+	if (Selection.Num())
+	{
+		bool bCanToggle = false;
+		
+		FScopedTransaction Transaction(LOCTEXT("SetAccessSpecifierTransaction", "Set Access Specifier"));
+		Selection.ForEachItem<FVariablesOutlinerEntryItem>([&bCanToggle](const FVariablesOutlinerEntryItem& EntryItem)
+		{
+			if (UAnimNextVariableEntry* VariableEntry = Cast<UAnimNextVariableEntry>(EntryItem.WeakEntry.Get()))
+			{
+				if (UAnimNextRigVMAsset* OuterAsset = VariableEntry->GetTypedOuter<UAnimNextRigVMAsset>())
+				{
+					// Disable toggling public/private for pure UAnimNextSharedVariables objects				
+					if (OuterAsset->GetClass() && OuterAsset->GetClass() != UAnimNextSharedVariables::StaticClass())
+					{
+						bCanToggle = true;
+					}
+				}
+			}
+		});
+
+		return bCanToggle;
+	}
+
+	return false;
+}
+
+void FVariablesOutlinerMode::Duplicate() const
+{	
+	const FSceneOutlinerItemSelection& Selection = SceneOutliner->GetSelection();
+	if (Selection.Num())
+	{		
+		FScopedTransaction Transaction(FText::Format(LOCTEXT("DuplicateVariableFormat", "Duplicate {0}|plural(one=variable, other=variables)"), Selection.Num()));
+
+		Selection.ForEachItem<FVariablesOutlinerEntryItem>([](const FVariablesOutlinerEntryItem& EntryItem)
+		{
+			if (const UAnimNextVariableEntry* VariableEntry = EntryItem.WeakEntry.Get())
+			{
+				if (UAnimNextRigVMAssetEditorData* EditorData = VariableEntry->GetTypedOuter<UAnimNextRigVMAssetEditorData>())
+				{	
+					const FName VariableName = UncookedOnly::FUtils::GetValidVariableName(EditorData, VariableEntry->GetVariableName());		
+					EditorData->AddVariable(VariableName, VariableEntry->GetType());
+				}
+			}
+		});
+	}
+}	
+
+bool FVariablesOutlinerMode::CanDuplicate() const
+{
+	const FSceneOutlinerItemSelection& Selection = SceneOutliner->GetSelection();
+
+	bool bCanDuplicate = true;
+	Selection.ForEachItem([&bCanDuplicate](const FSceneOutlinerTreeItemPtr& Item)
+	{
+		const FVariablesOutlinerEntryItem* EntryItem = Item->CastTo<FVariablesOutlinerEntryItem>();
+		if (EntryItem == nullptr || EntryItem->HasStructOwner())
+		{
+			bCanDuplicate = false;
+		}
+	});
+	
+	return bCanDuplicate;
+}
+
+void FVariablesOutlinerMode::FindReferences(ESearchScope InSearchScope) const
+{
+	const FSceneOutlinerItemSelection& Selection = SceneOutliner->GetSelection();
+	if (Selection.Num() == 1)
+	{
+		FSceneOutlinerTreeItemPtr ItemToFind = Selection.SelectedItems[0].Pin();
+		if (const FVariablesOutlinerEntryItem* EntryItem = ItemToFind->CastTo<FVariablesOutlinerEntryItem>())
+		{
+			if (TSharedPtr<SDockTab> FindAndReplaceTab = WeakWorkspaceEditor.Pin()->GetTabManager()->TryInvokeTab(FTabId(UE::UAF::Editor::FindAndReplaceTabName)))
+			{
+				TSharedRef<IAnimAssetFindReplace> FindAndReplace = StaticCastSharedRef<IAnimAssetFindReplace>(FindAndReplaceTab->GetContent());
+				FindAndReplace->SetCurrentProcessor(UAnimNextAssetFindReplaceVariables::StaticClass());
+				UAnimNextAssetFindReplaceVariables* AnimNextFindAndReplaceProcessor = FindAndReplace->GetProcessor<UAnimNextAssetFindReplaceVariables>();
+
+				if (TSharedPtr<UE::Workspace::IWorkspaceEditor> SharedWorkspaceEditor = WeakWorkspaceEditor.Pin())
+				{
+					AnimNextFindAndReplaceProcessor->SetWorkspaceEditor(SharedWorkspaceEditor.ToSharedRef());
+				}
+				
+				UAnimNextVariableEntry* VariableEntry = EntryItem->WeakEntry.Get();
+				AnimNextFindAndReplaceProcessor->SetSearchScope(InSearchScope);		
+				AnimNextFindAndReplaceProcessor->SetFindReferenceFromEntry(VariableEntry);				
+			}
+		}		
+	}
+}
+
+bool FVariablesOutlinerMode::CanFindReferences() const
+{
+	const FSceneOutlinerItemSelection& Selection = SceneOutliner->GetSelection();
+	if (Selection.Num() == 1)
+	{
+		FSceneOutlinerTreeItemPtr ItemToFind = Selection.SelectedItems[0].Pin();
+		if (const FVariablesOutlinerEntryItem* EntryItem = ItemToFind->CastTo<FVariablesOutlinerEntryItem>())
+		{
+			return true;
+		}
+	}
+	
+	return false;
+}
+
+bool FVariablesOutlinerMode::IsFindReferencesVisible(ESearchScope InSearchScope) const
+{
+	return CanFindReferences();
+}
+
+void FVariablesOutlinerMode::PopulateNewVariableToolMenuEntries(UToolMenu* InMenu, bool bAddSeparator)
+{
+	const UAnimNextVariableItemMenuContext* MenuContext = InMenu->FindContext<UAnimNextVariableItemMenuContext>();
+	if(MenuContext == nullptr)
+	{
+		return;
+	}
+
+	TMap<TWeakObjectPtr<UAnimNextRigVMAssetEditorData>, FString> AssetsAndCategoriesToAddTo;
+	if (MenuContext->WeakEditorDatas.Num())
+	{
+		if (UAnimNextRigVMAssetEditorData* EditorData = MenuContext->WeakEditorDatas[0].Get())
+		{
+			AssetsAndCategoriesToAddTo.Add(EditorData, TEXT(""));
+		}
+	}
+	else
+	{
+		TSharedPtr<UE::UAF::Editor::SVariablesOutliner> SharedOutliner = MenuContext->WeakOutliner.Pin();
+		if (SharedOutliner.IsValid())
+		{
+			TArray<FSceneOutlinerTreeItemPtr> SelectedItems = SharedOutliner->GetSelectedItems();
+		
+			for(const FSceneOutlinerTreeItemPtr& Item : SelectedItems)
+			{
+				const FVariablesOutlinerAssetItem* AssetItem = Item->CastTo<FVariablesOutlinerAssetItem>();
+				const FVariablesOutlinerCategoryItem* CategoryItem = Item->CastTo<FVariablesOutlinerCategoryItem>();
+				if (AssetItem == nullptr && CategoryItem == nullptr)
+				{
+					continue;
+				}
+
+				UAnimNextRigVMAsset* Asset = nullptr;
+				FString Category;
+		
+				if (CategoryItem)
+				{
+					Asset = CategoryItem->WeakOwner.Get();
+					Category = CategoryItem->GetDisplayString();
+				}
+				else
+				{
+					Asset = AssetItem->SoftAsset.Get();
+				}
+
+				if(Asset == nullptr)
+				{
+					continue;
+				}
+
+				UAnimNextRigVMAssetEditorData* EditorData = UncookedOnly::FUtils::GetEditorData<UAnimNextRigVMAssetEditorData>(Asset);
+				if(EditorData == nullptr)
+				{
+					continue;
+				}
+
+				AssetsAndCategoriesToAddTo.Add(EditorData, Category);
+			}
+		}
+	}
+	
+	if (AssetsAndCategoriesToAddTo.Num())
+	{
+		constexpr bool bShouldCloseWindowAfterMenuSelection = true;
+		FMenuBuilder MenuBuilder(bShouldCloseWindowAfterMenuSelection, TSharedPtr<FUICommandList>());
+
+		FToolMenuSection& Section = InMenu->AddSection("Variables", LOCTEXT("VariablesSectionLabel", "Variables"));
+		Section.AddMenuEntry("AddNewVariable", LOCTEXT("VariableLabel", "Add Variable"), LOCTEXT("VariableTooltip", "Adds a new Variable"),FSlateIcon(FAppStyle::GetAppStyleSetName(), "Icons.PlusCircle"),
+			FUIAction(FExecuteAction::CreateLambda([AssetsAndCategoriesToAddTo]()
+			{
+				for (auto& Pair : AssetsAndCategoriesToAddTo)
+				{
+					if (UAnimNextRigVMAssetEditorData* EditorData = Pair.Key.Get())
+                    {
+                        UAnimNextVariableSettings* Settings = GetMutableDefault<UAnimNextVariableSettings>();
+                        FScopedTransaction Transaction(LOCTEXT("AddVariable", "Add variable"));
+                        
+                        const FName VariableName = UncookedOnly::FUtils::GetValidVariableName(EditorData, Settings->GetLastVariableName());		
+
+						if (UAnimNextVariableEntry* VariableEntry = Pair.Key->AddVariable(VariableName, Settings->GetLastVariableType()))
+						{
+							VariableEntry->SetVariableCategory(Pair.Value);
+						}
+                        Settings->SetLastVariableName(VariableName);
+                    }
+				}
+			}))
+		);
+
+		if (AssetsAndCategoriesToAddTo.Num() == 1)
+		{
+			Section.AddSubMenu("AddSharedVariables", LOCTEXT("SharedVariableLabel", "Add Shared Variables"), LOCTEXT("SharedVariableTooltip", "Includes added SharedVariables"),
+			FNewMenuDelegate::CreateLambda([AssetsAndCategoriesToAddTo](FMenuBuilder& SubMenuBuilder)
+				{
+					TArray<FAssetData> ExistingSharedVariableReferences;
+					for (auto& Pair : AssetsAndCategoriesToAddTo)
+					{
+						if (UAnimNextRigVMAssetEditorData* EditorData = Pair.Key.Get())
+						{
+							EditorData->ForEachEntryOfType<UAnimNextSharedVariablesEntry>([&ExistingSharedVariableReferences](const UAnimNextSharedVariablesEntry* SharedVariablesEntry) -> bool
+							{
+								if (SharedVariablesEntry->GetType() == EAnimNextSharedVariablesType::Asset)
+								{
+									ExistingSharedVariableReferences.Add(SharedVariablesEntry->GetAsset());
+								}
+								return true;
+							});
+						}
+					}
+
+					FAssetPickerConfig AssetPickerConfig;
+					AssetPickerConfig.Filter.bRecursiveClasses = true;
+					AssetPickerConfig.Filter.ClassPaths.Add(UAnimNextSharedVariables::StaticClass()->GetClassPathName());
+					AssetPickerConfig.InitialAssetViewType = EAssetViewType::List;
+					AssetPickerConfig.OnAssetSelected = FOnAssetSelected::CreateLambda([AssetsAndCategoriesToAddTo](const FAssetData& InAssetData)
+					{
+						FSlateApplication::Get().DismissAllMenus();
+						if(UAnimNextSharedVariables* SharedVariables = Cast<UAnimNextSharedVariables>(InAssetData.GetAsset()))
+						{							
+							FScopedTransaction Transaction(LOCTEXT("AddSharedVariable", "Add Shared Variables"));
+							for (auto& Pair : AssetsAndCategoriesToAddTo)
+							{
+								if (UAnimNextRigVMAssetEditorData* EditorData = Pair.Key.Get())
+								{
+									EditorData->AddSharedVariables(SharedVariables);
+								}
+							}
+						}
+					});
+					AssetPickerConfig.OnShouldFilterAsset = FOnShouldFilterAsset::CreateLambda([ExistingSharedVariableReferences](const FAssetData& InAssetData)
+					{
+						FAnimNextAssetRegistryExports Exports;
+						UncookedOnly::FUtils::GetExportedVariablesForAsset(InAssetData, Exports);
+						if(Exports.Exports.Num() == 0 && InAssetData.GetClass(EResolveClass::Yes) != UAnimNextSharedVariables::StaticClass())
+						{
+							return true;
+						}
+
+						// Filter out already referenced assets
+						return ExistingSharedVariableReferences.ContainsByPredicate([InAssetData](const FAssetData& ExistingAssetData)
+						{
+							return ExistingAssetData == InAssetData;
+						});
+					});
+
+					FContentBrowserModule& ContentBrowserModule = FModuleManager::Get().LoadModuleChecked<FContentBrowserModule>(TEXT("ContentBrowser"));
+					TSharedPtr<SWidget> Widget = SNew(SBox)
+					.WidthOverride(300.0f)
+					.HeightOverride(400.0f)
+					[
+						ContentBrowserModule.Get().CreateAssetPicker(AssetPickerConfig)
+					];
+
+					SubMenuBuilder.AddWidget(Widget.ToSharedRef(), FText::GetEmpty(), true);
+				}),
+				false,
+				FSlateIcon(FAnimNextEditorStyle::Get().GetStyleSetName(), "ClassIcon.AnimNextSharedVariables")
+			);
+
+
+			Section.AddSubMenu("AddNativeSharedVariables",LOCTEXT("NativeSharedVariableLabel", "Add Native Shared Variables"), LOCTEXT("NativeSharedVariableTooltip", "Includes added Native SharedVariables"),
+		FNewMenuDelegate::CreateLambda([AssetsAndCategoriesToAddTo](FMenuBuilder& SubMenuBuilder)
+				{
+
+					TArray<FSoftObjectPath> ExistingSharedVariableStructPaths;
+					TArray<FAssetData> ExistingSharedVariableReferences;
+					for (auto& Pair : AssetsAndCategoriesToAddTo)
+					{
+						if (UAnimNextRigVMAssetEditorData* EditorData = Pair.Key.Get())
+						{
+							EditorData->ForEachEntryOfType<UAnimNextSharedVariablesEntry>([&ExistingSharedVariableStructPaths](const UAnimNextSharedVariablesEntry* SharedVariablesEntry) -> bool
+							{
+								if (SharedVariablesEntry->GetType() == EAnimNextSharedVariablesType::Struct)
+								{
+									ExistingSharedVariableStructPaths.Add(SharedVariablesEntry->GetObjectPath());
+								}
+								return true;
+							});
+						}
+					}
+			
+					class FStructFilter : public IStructViewerFilter
+					{
+					public:
+
+						FStructFilter(const TArray<FSoftObjectPath>& InFilteredStructPaths) : FilteredStructPaths(InFilteredStructPaths) {}
+						
+						virtual bool IsStructAllowed(const FStructViewerInitializationOptions& InInitOptions, const UScriptStruct* InStruct, TSharedRef<FStructViewerFilterFuncs> InFilterFuncs) override
+						{
+							if (InStruct->IsA<UUserDefinedStruct>())
+							{
+								return false;
+							}
+
+							if (!InStruct->HasMetaData(TEXT("BlueprintType")) || InStruct->HasMetaData(TEXT("Hidden")))
+							{
+								return false;
+							}
+
+							if (FilteredStructPaths.Contains(InStruct))
+							{
+								return false;
+							}
+
+							return true;
+						}
+
+						virtual bool IsUnloadedStructAllowed(const FStructViewerInitializationOptions& InInitOptions, const FSoftObjectPath& InStructPath, TSharedRef<class FStructViewerFilterFuncs> InFilterFuncs) override
+						{
+							return false;
+						};
+
+						TArray<FSoftObjectPath> FilteredStructPaths;
+					};
+
+					FStructViewerModule& StructViewerModule = FModuleManager::LoadModuleChecked<FStructViewerModule>("StructViewer");
+
+					FStructViewerInitializationOptions Options;
+					Options.Mode = EStructViewerMode::StructPicker;
+					Options.StructFilter = MakeShared<FStructFilter>(ExistingSharedVariableStructPaths);
+
+					TSharedPtr<SWidget> Widget = SNew(SBox)
+						.WidthOverride(300.0f)
+						.HeightOverride(400.0f)
+						[
+							StructViewerModule.CreateStructViewer(Options, FOnStructPicked::CreateLambda([AssetsAndCategoriesToAddTo](const UScriptStruct* InStruct)
+							{
+								FSlateApplication::Get().DismissAllMenus();
+
+								if(InStruct)
+								{
+									FScopedTransaction Transaction(LOCTEXT("AddNativeSharedVariable", "Add Native Shared Variables"));
+									for (auto& Pair : AssetsAndCategoriesToAddTo)
+									{
+										if (UAnimNextRigVMAssetEditorData* EditorData = Pair.Key.Get())
+										{									
+											EditorData->AddSharedVariablesStruct(InStruct);
+										}
+									}
+								}
+							}))
+						];
+
+					SubMenuBuilder.AddWidget(Widget.ToSharedRef(), FText::GetEmpty(), true);
+				}),
+				false,
+				FSlateIconFinder::FindIconForClass(UUserDefinedStruct::StaticClass())
+			);
+		}
+
+		if (bAddSeparator)
+		{
+			Section.AddSeparator("PostAddVariableSeparator");
+		}
+	}
+}
+
+void FVariablesOutlinerMode::SetHighlightedItem(FSceneOutlinerTreeItemID InID) const
+{
+	if (SVariablesOutliner* Outliner = static_cast<SVariablesOutliner*>(SceneOutliner))
+	{
+		if (FSceneOutlinerTreeItemPtr ItemPtr = Outliner->GetTreeItem(InID))
+		{
+			Outliner->FrameItem(InID);
+			Outliner->SetHighlightedItem(ItemPtr);
+		}
+	}
+}
+
+void FVariablesOutlinerMode::ClearHighlightedItem(FSceneOutlinerTreeItemID InID) const
+{
+	if (SVariablesOutliner* Outliner = static_cast<SVariablesOutliner*>(SceneOutliner))
+	{
+		if (FSceneOutlinerTreeItemPtr ItemPtr = Outliner->GetTreeItem(InID))
+		{
+			Outliner->ClearHighlightedItem(ItemPtr);
+		}
+	}
+}
+
+void FVariablesOutlinerMode::SaveAsset() const
+{
+	const FSceneOutlinerItemSelection& Selection = SceneOutliner->GetSelection();
+	Selection.ForEachItem<FVariablesOutlinerAssetItem>([](const FVariablesOutlinerAssetItem& AssetItem)
+	{
+		if (const UAnimNextRigVMAsset* Asset = AssetItem.SoftAsset.Get())
+		{
+			FEditorFileUtils::PromptForCheckoutAndSave({Asset->GetPackage()}, false, /*bPromptToSave=*/ false);
+		}
+	});	
+}
+
+bool FVariablesOutlinerMode::CanSaveAsset() const
+{
+	const FSceneOutlinerItemSelection& Selection = SceneOutliner->GetSelection();	
+	return Selection.SelectedItems.ContainsByPredicate([](const TWeakPtr<ISceneOutlinerTreeItem>& WeakItem)
+	{
+		if (WeakItem.IsValid())
+		{
+			if (const FVariablesOutlinerAssetItem* AssetItem = WeakItem.Pin()->CastTo<FVariablesOutlinerAssetItem>())
+			{
+				if (const UAnimNextRigVMAsset* Asset = AssetItem->SoftAsset.Get())
+				{
+					return Asset->GetPackage()->IsDirty();
+				}
+			}
+		}
+
+		return false;
+	});
+}
+
+bool FVariablesOutlinerMode::IsAsset() const
+{
+	const FSceneOutlinerItemSelection& Selection = SceneOutliner->GetSelection();	
+	return Selection.SelectedItems.ContainsByPredicate([](const TWeakPtr<ISceneOutlinerTreeItem>& WeakItem) { return WeakItem.IsValid() && WeakItem.Pin()->IsA<FVariablesOutlinerAssetItem>(); });	
+}
+
+TSharedPtr<FDragDropOperation> FVariablesOutlinerMode::CreateDragDropOperation(const FPointerEvent& MouseEvent, const TArray<FSceneOutlinerTreeItemPtr>& InTreeItems) const
+{
+	if(InTreeItems.Num() != 1)
+	{
+		return nullptr;
+	}
+
+	const FSceneOutlinerTreeItemPtr& SelectedItem = InTreeItems[0];
+	if (FVariablesOutlinerEntryItem* VariableItem = SelectedItem->CastTo<FVariablesOutlinerEntryItem>())
+	{
+		if(UAnimNextVariableEntry* Entry = Cast<UAnimNextVariableEntry>(VariableItem->WeakEntry.Get()))
+		{
+			const UAnimNextSharedVariables* Asset = nullptr;
+			if (VariableItem->WeakSharedVariablesEntry.Get())
+			{
+				Asset = VariableItem->WeakSharedVariablesEntry->GetAsset();
+			}
+			else
+			{
+				Asset = Entry->GetTypedOuter<UAnimNextSharedVariables>();
+			}
+		
+			if (Asset)
+			{
+				TSharedPtr<FAnimNextSchemaAction_Variable> Action = MakeShared<FAnimNextSchemaAction_Variable>(Entry->GetVariableName(), Asset, Entry->GetType(), FAnimNextSchemaAction_Variable::EVariableAccessorChoice::Deferred);
+
+				return FVariableDragDropOp::New(Action, StaticCastSharedRef<FVariablesOutlinerEntryItem>(VariableItem->AsShared()));
+			}
+		}
+		else if (const FProperty* Property = VariableItem->PropertyPath.Get())
+		{
+			if (const UScriptStruct* Struct = Property->GetOwner<UScriptStruct>())
+			{
+				TSharedPtr<FAnimNextSchemaAction_Variable> Action = MakeShared<FAnimNextSchemaAction_Variable>(Property->GetFName(), Struct, FAnimNextParamType::FromProperty(Property), FAnimNextSchemaAction_Variable::EVariableAccessorChoice::Deferred);
+				
+				TSharedPtr<FVariableDragDropOp> Op = FVariableDragDropOp::New(Action, StaticCastSharedRef<FVariablesOutlinerEntryItem>(VariableItem->AsShared()));
+				return Op;
+			}
+		}
+	}
+	else if (FVariablesOutlinerCategoryItem* CategoryItem = SelectedItem->CastTo<FVariablesOutlinerCategoryItem>())
+	{
+		TSharedPtr<FCategoryDragDropOp> Op = FCategoryDragDropOp::New(StaticCastSharedRef<FVariablesOutlinerCategoryItem>(CategoryItem->AsShared()));
+		return Op;
+	}
+
+	return nullptr;
+
+}
+
+bool FVariablesOutlinerMode::ParseDragDrop(FSceneOutlinerDragDropPayload& OutPayload, const FDragDropOperation& Operation) const
+{
+	return ISceneOutlinerMode::ParseDragDrop(OutPayload, Operation);
+}
+
+FSceneOutlinerDragValidationInfo FVariablesOutlinerMode::ValidateDrop(const ISceneOutlinerTreeItem& DropTarget,
+	const FSceneOutlinerDragDropPayload& Payload) const
+{
+	return ISceneOutlinerMode::ValidateDrop(DropTarget, Payload);
+}
+
+void FVariablesOutlinerMode::OnDrop(ISceneOutlinerTreeItem& DropTarget, const FSceneOutlinerDragDropPayload& Payload,
+	const FSceneOutlinerDragValidationInfo& ValidationInfo) const
+{
+	ISceneOutlinerMode::OnDrop(DropTarget, Payload, ValidationInfo);
+}
+
+FReply FVariablesOutlinerMode::OnDragOverItem(const FDragDropEvent& Event, const ISceneOutlinerTreeItem& Item) const
+{
+	return ISceneOutlinerMode::OnDragOverItem(Event, Item);
+}
+
+int32 FVariablesOutlinerMode::GetTypeSortPriority(const ISceneOutlinerTreeItem& Item) const
+{
+	if(const FVariablesOutlinerAssetItem* AssetItem = Item.CastTo<FVariablesOutlinerAssetItem>())
+	{
+		return AssetItem->SortValue;
+	}
+
+	if(const FVariablesOutlinerEntryItem* EntryItem = Item.CastTo<FVariablesOutlinerEntryItem>())
+	{
+		return EntryItem->SortValue;
+	}
+
+	if(const FVariablesOutlinerCategoryItem* CategoryItem = Item.CastTo<FVariablesOutlinerCategoryItem>())
+	{
+		return CategoryItem->SortValue;
+	}
+	
+	return ISceneOutlinerMode::GetTypeSortPriority(Item);
+}
+
+void FVariablesOutlinerMode::OnItemDoubleClick(FSceneOutlinerTreeItemPtr Item)
+{
+	if (Item.IsValid())
+	{
+		if(FVariablesOutlinerAssetItem* AssetItem = Item->CastTo<FVariablesOutlinerAssetItem>())
+		{
+			// [TODO] Open asset from Variables Outliner
+		}
+	}	
+	
+	ISceneOutlinerMode::OnItemDoubleClick(Item);
+}
+}
+
+#undef LOCTEXT_NAMESPACE // "FVariablesOutlinerMode"
\ No newline at end of file
diff --git a/Engine/Plugins/Experimental/UAF/UAF/Source/UAFEditor/Private/Variables/Outliner/VariablesOutlinerMode.h b/Engine/Plugins/Experimental/UAF/UAF/Source/UAFEditor/Private/Variables/Outliner/VariablesOutlinerMode.h
new file mode 100644
--- /dev/null
+++ b/Engine/Plugins/Experimental/UAF/UAF/Source/UAFEditor/Private/Variables/Outliner/VariablesOutlinerMode.h
@@ -0,0 +1,94 @@
+// Copyright Epic Games, Inc. All Rights Reserved.
+
+#pragma once
+
+#include "AnimNextRigVMAssetEditorData.h"
+#include "ISceneOutlinerMode.h"
+#include "Common/AnimNextAssetFindReplaceVariables.h"
+
+class UAnimNextRigVMAssetEditorData;
+class FVariablesOutlinerMode;
+namespace UE::Workspace
+{
+	class IWorkspaceEditor;
+}
+
+namespace UE::UAF::Editor
+{
+	class SVariablesOutliner;
+}
+
+namespace UE::UAF::Editor
+{
+class FVariablesOutlinerMode : public ISceneOutlinerMode
+{
+public:
+	FVariablesOutlinerMode(SVariablesOutliner* InVariablesOutliner, const TSharedRef<UE::Workspace::IWorkspaceEditor>& InWorkspaceEditor);
+
+	// Begin ISceneOutlinerMode overrides
+	virtual void Rebuild() override;
+	virtual TSharedPtr<SWidget> CreateContextMenu() override;
+	virtual FReply OnKeyDown(const FKeyEvent& InKeyEvent) override;
+	virtual void OnItemClicked(FSceneOutlinerTreeItemPtr Item) override;
+	void HandleItemSelection(const FSceneOutlinerItemSelection& Selection) const;
+	virtual void OnItemSelectionChanged(FSceneOutlinerTreeItemPtr Item, ESelectInfo::Type SelectionType, const FSceneOutlinerItemSelection& Selection) override;
+	virtual bool CanCustomizeToolbar() const override { return true; }
+	virtual bool CanRenameItem(const ISceneOutlinerTreeItem& Item) const override;
+	virtual ESelectionMode::Type GetSelectionMode() const override { return ESelectionMode::Multi; }
+	virtual void BindCommands(const TSharedRef<FUICommandList>& OutCommandList) override;
+	virtual int32 GetTypeSortPriority(const ISceneOutlinerTreeItem& Item) const override;
+	virtual TSharedPtr<FDragDropOperation> CreateDragDropOperation(const FPointerEvent& MouseEvent, const TArray<FSceneOutlinerTreeItemPtr>& InTreeItems) const override;
+	virtual bool ParseDragDrop(FSceneOutlinerDragDropPayload& OutPayload, const FDragDropOperation& Operation) const override;
+	virtual FSceneOutlinerDragValidationInfo ValidateDrop(const ISceneOutlinerTreeItem& DropTarget, const FSceneOutlinerDragDropPayload& Payload) const override;
+	virtual void OnDrop(ISceneOutlinerTreeItem& DropTarget, const FSceneOutlinerDragDropPayload& Payload, const FSceneOutlinerDragValidationInfo& ValidationInfo) const override;
+	virtual FReply OnDragOverItem(const FDragDropEvent& Event, const ISceneOutlinerTreeItem& Item) const override;
+	virtual bool CanSupportDragAndDrop() const override { return true; }
+	virtual void OnItemDoubleClick(FSceneOutlinerTreeItemPtr Item) override;
+	
+	void SetHighlightedItem(FSceneOutlinerTreeItemID InID) const;
+	void ClearHighlightedItem(FSceneOutlinerTreeItemID InID) const;
+
+	static void PopulateNewVariableToolMenuEntries(UToolMenu* InMenu, bool bAddSeparator);
+protected:
+	virtual TUniquePtr<ISceneOutlinerHierarchy> CreateHierarchy() override;
+	// End ISceneOutlinerMode overrides
+
+	void ResetOutlinerSelection() const;
+
+	SVariablesOutliner* GetOutliner() const;
+
+	void Rename() const;
+	virtual bool CanRename() const override;
+
+	void Delete() const;
+	virtual bool CanDelete() const override;
+
+	void Copy() const;
+	virtual bool CanCopy() const override;
+
+	void Paste() const;
+	virtual bool CanPaste() const override;
+
+	void ToggleVariableExport() const;
+	bool CanToggleVariableExport() const;
+
+	void Duplicate() const;
+	bool CanDuplicate() const;
+
+	void FindReferences(ESearchScope InSearchScope) const;
+	bool CanFindReferences() const;
+	bool IsFindReferencesVisible(ESearchScope InSearchScope) const;
+
+	void SaveAsset() const;
+	bool CanSaveAsset() const;
+	bool IsAsset() const;
+
+private:
+	friend class FVariablesOutlinerHierarchy;
+
+	TWeakPtr<UE::Workspace::IWorkspaceEditor> WeakWorkspaceEditor;
+
+	TSharedPtr<FUICommandList> CommandList;
+};
+
+}
diff --git a/Engine/Plugins/Experimental/UAF/UAF/Source/UAFEditor/Private/Variables/Outliner/VariablesOutlinerStructSharedVariablesItem.cpp b/Engine/Plugins/Experimental/UAF/UAF/Source/UAFEditor/Private/Variables/Outliner/VariablesOutlinerStructSharedVariablesItem.cpp
new file mode 100644
--- /dev/null
+++ b/Engine/Plugins/Experimental/UAF/UAF/Source/UAFEditor/Private/Variables/Outliner/VariablesOutlinerStructSharedVariablesItem.cpp
@@ -0,0 +1,159 @@
+﻿// Copyright Epic Games, Inc. All Rights Reserved.
+
+#include "VariablesOutlinerStructSharedVariablesItem.h"
+
+#include "ISceneOutliner.h"
+#include "Styling/SlateColor.h"
+#include "Misc/Attribute.h"
+#include "Misc/Optional.h"
+#include "SceneOutlinerStandaloneTypes.h"
+#include "AnimNextEditorModule.h"
+#include "AnimNextRigVMAssetEditorData.h"
+#include "VariablesOutlinerDragDrop.h"
+#include "Entries/AnimNextSharedVariablesEntry.h"
+#include "Styling/SlateIconFinder.h"
+#include "Styling/StarshipCoreStyle.h"
+#include "UObject/Package.h"
+#include "Widgets/SCompoundWidget.h"
+#include "Widgets/DeclarativeSyntaxSupport.h"
+#include "Widgets/Images/SImage.h"
+#include "Widgets/Text/SInlineEditableTextBlock.h"
+#include "Widgets/Text/STextBlock.h"
+#include "Widgets/Views/STableRow.h"
+#include "Widgets/Views/SListView.h"
+
+#define LOCTEXT_NAMESPACE "VariablesOutlinerTreeItem"
+
+namespace UE::UAF::Editor
+{
+
+const FSceneOutlinerTreeItemType FVariablesOutlinerStructSharedVariablesItem::Type;
+
+class SVariablesOutlinerStructSharedVariablesLabel : FSceneOutlinerCommonLabelData, public SCompoundWidget
+{
+	SLATE_BEGIN_ARGS(SVariablesOutlinerStructSharedVariablesLabel) {}
+	SLATE_END_ARGS()
+
+	void Construct(const FArguments& InArgs, FVariablesOutlinerStructSharedVariablesItem& InTreeItem, ISceneOutliner& SceneOutliner, const STableRow<FSceneOutlinerTreeItemPtr>& InRow)
+	{
+		WeakSceneOutliner = StaticCastSharedRef<ISceneOutliner>(SceneOutliner.AsShared());
+		TreeItem = StaticCastSharedRef<FVariablesOutlinerStructSharedVariablesItem>(InTreeItem.AsShared());
+
+		FText AssetName;
+		if(const UAnimNextSharedVariablesEntry* AssetEntry = InTreeItem.WeakEntry.Get())
+		{
+			AssetName = FText::FromString(AssetEntry->GetObjectPath().ToString());
+		}
+		else
+		{
+			AssetName = LOCTEXT("UnknownAssetName", "Unknown Asset");
+		}
+		
+		ChildSlot
+		[
+			SNew(SHorizontalBox)
+			.ToolTipText(FText::Format(LOCTEXT("ImportedVariablesFormat", "Shared variables from '{0}'"), AssetName))
+			+SHorizontalBox::Slot()
+			.AutoWidth()
+			.VAlign(VAlign_Center)
+			.Padding(2.0f, 2.0f)
+			[
+				SNew(SImage)
+				.Image(FSlateIconFinder::FindIconBrushForClass(UUserDefinedStruct::StaticClass()))
+			]
+			+SHorizontalBox::Slot()
+			.AutoWidth()
+			.VAlign(VAlign_Center)
+			.Padding(2.0f, 2.0f)
+			[
+				SNew(STextBlock)
+				.Font(FStyleFonts::Get().NormalBold)
+				.Text(this, &SVariablesOutlinerStructSharedVariablesLabel::GetDisplayText)
+				.HighlightText(SceneOutliner.GetFilterHighlightText())
+				.ColorAndOpacity(this, &SVariablesOutlinerStructSharedVariablesLabel::GetForegroundColor)
+			]
+		];
+	}
+	
+	FText GetDisplayText() const
+	{
+		if (const TSharedPtr<FVariablesOutlinerStructSharedVariablesItem> Item = TreeItem.Pin())
+		{
+			return FText::FromString(Item->GetDisplayString());
+		}
+		return FText();
+	}
+
+	virtual void OnDragEnter( const FGeometry& MyGeometry, const FDragDropEvent& DragDropEvent ) override
+	{
+		TSharedPtr<FVariableDragDropOp> GraphDropOp = DragDropEvent.GetOperationAs<FVariableDragDropOp>();
+		if (GraphDropOp.IsValid())
+		{
+			const FSlateBrush* ErrorSymbol = FAppStyle::GetBrush(TEXT("Graph.ConnectorFeedback.Error"));
+			GraphDropOp->SetSimpleFeedbackMessage(ErrorSymbol, FSlateColor(GetForegroundColor()), LOCTEXT("StructSharedVariablesFeedback", "Cannot make changes to Struct SharedVariables"));
+		}
+	}
+
+	virtual FSlateColor GetForegroundColor() const override
+	{
+		const TOptional<FLinearColor> BaseColor = FSceneOutlinerCommonLabelData::GetForegroundColor(*TreeItem.Pin());
+		return BaseColor.IsSet() ? BaseColor.GetValue() : FSlateColor::UseForeground();
+	}
+
+	TWeakPtr<FVariablesOutlinerStructSharedVariablesItem> TreeItem;
+	TSharedPtr<SInlineEditableTextBlock> TextBlock;
+};
+
+FVariablesOutlinerStructSharedVariablesItem::FVariablesOutlinerStructSharedVariablesItem(UAnimNextSharedVariablesEntry* InEntry)
+	: ISceneOutlinerTreeItem(FVariablesOutlinerStructSharedVariablesItem::Type)
+	, WeakEntry(InEntry)
+{
+}
+
+bool FVariablesOutlinerStructSharedVariablesItem::IsValid() const
+{
+	return WeakEntry.Get() != nullptr;
+}
+
+FSceneOutlinerTreeItemID FVariablesOutlinerStructSharedVariablesItem::GetID() const
+{
+	if (const UAnimNextSharedVariablesEntry* SharedVariableEntry = WeakEntry.Get())
+	{
+		check(SharedVariableEntry->GetStruct());
+		const FSoftObjectPath Path = SharedVariableEntry->GetStruct();
+		return GetTypeHash(Path);
+	}
+	
+	return GetTypeHash(WeakEntry);
+}
+
+FString FVariablesOutlinerStructSharedVariablesItem::GetDisplayString() const
+{
+	const UAnimNextRigVMAssetEntry* Entry = WeakEntry.Get();
+	if(Entry == nullptr)
+	{
+		return FString();
+	}
+
+	return Entry->GetDisplayName().ToString();
+}
+
+TSharedRef<SWidget> FVariablesOutlinerStructSharedVariablesItem::GenerateLabelWidget(ISceneOutliner& Outliner, const STableRow<FSceneOutlinerTreeItemPtr>& InRow)
+{
+	return SNew(SVariablesOutlinerStructSharedVariablesLabel, *this, Outliner, InRow);
+}
+
+FString FVariablesOutlinerStructSharedVariablesItem::GetPackageName() const
+{
+	const UAnimNextRigVMAssetEntry* Entry = WeakEntry.Get();
+	if(Entry == nullptr)
+	{
+		return ISceneOutlinerTreeItem::GetPackageName();
+	}
+
+	return Entry->GetPackage()->GetName();
+}
+
+}
+
+#undef LOCTEXT_NAMESPACE // "VariablesOutlinerTreeItem"
\ No newline at end of file
diff --git a/Engine/Plugins/Experimental/UAF/UAF/Source/UAFEditor/Private/Variables/Outliner/VariablesOutlinerStructSharedVariablesItem.h b/Engine/Plugins/Experimental/UAF/UAF/Source/UAFEditor/Private/Variables/Outliner/VariablesOutlinerStructSharedVariablesItem.h
new file mode 100644
--- /dev/null
+++ b/Engine/Plugins/Experimental/UAF/UAF/Source/UAFEditor/Private/Variables/Outliner/VariablesOutlinerStructSharedVariablesItem.h
@@ -0,0 +1,33 @@
+﻿// Copyright Epic Games, Inc. All Rights Reserved.
+
+#pragma once
+
+#include "ISceneOutlinerTreeItem.h"
+
+class UAnimNextSharedVariablesEntry;
+
+namespace UE::UAF::Editor
+{
+
+class IWorkspaceOutlinerItemDetails;
+
+struct FVariablesOutlinerStructSharedVariablesItem : ISceneOutlinerTreeItem
+{
+	static const FSceneOutlinerTreeItemType Type;
+
+	FVariablesOutlinerStructSharedVariablesItem(UAnimNextSharedVariablesEntry* InEntry);
+
+	// Begin ISceneOutlinerTreeItem overrides
+	virtual bool IsValid() const override;
+	virtual FSceneOutlinerTreeItemID GetID() const override;
+	virtual FString GetDisplayString() const override;
+	virtual bool CanInteract() const override { return true; }
+	virtual TSharedRef<SWidget> GenerateLabelWidget(ISceneOutliner& Outliner, const STableRow<FSceneOutlinerTreeItemPtr>& InRow) override;
+	virtual FString GetPackageName() const override;
+	// End ISceneOutlinerTreeItem overrides
+
+	// Weak-ptr to the represented entry
+	TWeakObjectPtr<UAnimNextSharedVariablesEntry> WeakEntry;
+};
+
+}
diff --git a/Engine/Plugins/Experimental/UAF/UAF/Source/UAFEditor/Private/Variables/SVariablesView.cpp b/Engine/Plugins/Experimental/UAF/UAF/Source/UAFEditor/Private/Variables/SVariablesView.cpp
new file mode 100644
--- /dev/null
+++ b/Engine/Plugins/Experimental/UAF/UAF/Source/UAFEditor/Private/Variables/SVariablesView.cpp
@@ -0,0 +1,300 @@
+﻿// Copyright Epic Games, Inc. All Rights Reserved.
+
+#include "SVariablesView.h"
+
+#include "AnimNextAssetWorkspaceAssetUserData.h"
+#include "AnimNextRigVMAsset.h"
+#include "ContentBrowserModule.h"
+#include "IContentBrowserSingleton.h"
+#include "SAssetDropTarget.h"
+#include "SceneOutlinerPublicTypes.h"
+#include "SceneOutlinerSourceControlColumn.h"
+#include "ScopedTransaction.h"
+#include "SSceneOutliner.h"
+#include "Outliner/VariablesOutlinerColumns.h"
+#include "Outliner/VariablesOutlinerMode.h"
+#include "Outliner/VariablesOutlinerEntryItem.h"
+#include "IWorkspaceEditor.h"
+#include "IWorkspaceEditorModule.h"
+#include "WorkspaceAssetRegistryInfo.h"
+#include "Outliner/VariablesOutlinerAssetItem.h"
+#include "Entries/AnimNextSharedVariablesEntry.h"
+
+#define LOCTEXT_NAMESPACE "SVariablesView"
+
+namespace UE::UAF::Editor
+{
+
+const FLazyName VariablesTabName("VariablesTab");
+
+void SVariablesOutliner::RegisterAssetDelegates(const UAnimNextRigVMAsset* InAsset)
+{	
+	// Bind for any modification callbacks
+	UAnimNextRigVMAssetEditorData* EditorData = UncookedOnly::FUtils::GetEditorData<UAnimNextRigVMAssetEditorData>(InAsset);
+	if(EditorData != nullptr)
+	{
+		EditorData->ModifiedDelegate.AddSP(this, &SVariablesOutliner::OnEditorDataModified);
+
+		for (const TObjectPtr<UAnimNextRigVMAssetEntry>& Entry : EditorData->Entries)
+		{
+			if (UAnimNextSharedVariablesEntry* SharedVariablesEntry = Cast<UAnimNextSharedVariablesEntry>(Entry.Get()))
+			{
+				if (TObjectPtr<const UAnimNextSharedVariables> SharedVariablesAsset = SharedVariablesEntry->GetAsset())
+				{
+					RegisterAssetDelegates(SharedVariablesAsset);
+					Assets.AddUnique(Cast<UAnimNextRigVMAsset>(const_cast<UAnimNextSharedVariables*>(SharedVariablesAsset.Get())));					
+				}
+			}
+		}
+	}
+}
+
+void SVariablesOutliner::UnregisterAssetDelegates(const UAnimNextRigVMAsset* InAsset)
+{
+	UAnimNextRigVMAssetEditorData* EditorData = UncookedOnly::FUtils::GetEditorData<UAnimNextRigVMAssetEditorData>(InAsset);
+	if(EditorData != nullptr)
+	{
+		EditorData->ModifiedDelegate.RemoveAll(this);
+
+		for (const TObjectPtr<UAnimNextRigVMAssetEntry>& Entry : EditorData->Entries)
+		{
+			if (UAnimNextSharedVariablesEntry* SharedVariablesEntry = Cast<UAnimNextSharedVariablesEntry>(Entry.Get()))
+			{
+				if (TObjectPtr<const UAnimNextSharedVariables> SharedVariablesAsset = SharedVariablesEntry->GetAsset())
+				{
+					UnregisterAssetDelegates(SharedVariablesAsset);					
+				}
+			}
+		}
+	}
+}
+
+void SVariablesOutliner::SetAssets(TConstArrayView<TSoftObjectPtr<UAnimNextRigVMAsset>> InAssets)
+{
+	for(const TSoftObjectPtr<UAnimNextRigVMAsset>& CurrentSoftAsset : Assets)
+	{
+		if(UAnimNextRigVMAsset* CurrentAsset = CurrentSoftAsset.Get())
+		{
+			UnregisterAssetDelegates(CurrentAsset);
+		}
+	}
+
+	Assets = InAssets;
+
+	for(const TSoftObjectPtr<UAnimNextRigVMAsset>& NewSoftAsset : InAssets)
+	{
+		if(UAnimNextRigVMAsset* NewAsset = NewSoftAsset.Get())
+		{
+			RegisterAssetDelegates(NewAsset);
+		}
+	}
+
+	FullRefresh();
+}
+
+void SVariablesOutliner::UpdateAssets()
+{
+	TArray<FSoftObjectPath> AssetsToAsyncLoad;
+	TArray<TSoftObjectPtr<UAnimNextRigVMAsset>> ExportAssets;
+
+	auto HandleAssetPath = [this, &AssetsToAsyncLoad, &ExportAssets](const FSoftObjectPath& Path)
+	{
+		ExportAssets.Add(TSoftObjectPtr<UAnimNextRigVMAsset>(Path));
+
+		if (Path.ResolveObject() == nullptr)
+		{
+			AssetsToAsyncLoad.Add(Path);
+		}
+	};
+	
+	if (FWorkspaceOutlinerAssetReferenceItemData::IsAssetReference(Export))
+	{
+		// Reference item
+		TArray<FSoftObjectPath> AssetPaths;
+		Export.GetAssetPaths(AssetPaths);
+
+		for (const FSoftObjectPath& Path : AssetPaths)
+		{
+			HandleAssetPath(Path);
+		}
+	}
+	else
+	{
+		const FSoftObjectPath FirstAssetPath = Export.GetFirstAssetPath();
+		if(FirstAssetPath.IsAsset())
+		{
+			HandleAssetPath(FirstAssetPath);
+		}
+		else if(Export.HasData() && Export.GetData().GetScriptStruct()->IsChildOf(FAnimNextAssetEntryOutlinerData::StaticStruct()))
+		{
+			const FAnimNextAssetEntryOutlinerData& EntryData = Export.GetData().Get<FAnimNextAssetEntryOutlinerData>();
+			if(EntryData.SoftEntryPtr.IsValid())
+			{
+				if(UAnimNextRigVMAsset* Asset = EntryData.GetEntry()->GetTypedOuter<UAnimNextRigVMAsset>())
+				{
+					HandleAssetPath(Asset);
+				}
+			}
+		}
+		else if(Export.HasData() && Export.GetData().GetScriptStruct()->IsChildOf(FAnimNextGraphFunctionOutlinerData::StaticStruct()))
+		{
+			const FAnimNextGraphFunctionOutlinerData& EntryData = Export.GetData().Get<FAnimNextGraphFunctionOutlinerData>();
+			if(EntryData.SoftEditorObject.IsValid())
+			{
+				if(URigVMEdGraph* EdGraph = EntryData.SoftEditorObject.Get())
+				{
+					if(UAnimNextRigVMAsset* Asset = EdGraph->GetTypedOuter<UAnimNextRigVMAsset>())
+					{
+						HandleAssetPath(Asset);
+					}
+				}
+			}
+		}
+	}
+
+	SetAssets(ExportAssets);
+
+	// Try async load any missing assets
+	for(const FSoftObjectPath& AssetPath : AssetsToAsyncLoad)
+	{
+		TWeakPtr<SVariablesOutliner> WeakOutliner = StaticCastSharedRef<SVariablesOutliner>(AsShared());
+		
+		AssetPath.LoadAsync(FLoadSoftObjectPathAsyncDelegate::CreateLambda([WeakOutliner](const FSoftObjectPath& InSoftObjectPath, UObject* InObject)
+		{
+			UAnimNextRigVMAsset* Asset = Cast<UAnimNextRigVMAsset>(InObject);
+			if(Asset == nullptr)
+			{
+				return;
+			}
+
+			TSharedPtr<SVariablesOutliner> PinnedVariablesOutliner = WeakOutliner.Pin();
+			if(!PinnedVariablesOutliner.IsValid())
+			{
+				return;
+			}
+
+			PinnedVariablesOutliner->HandleAssetLoaded(InSoftObjectPath, Asset);
+		}));
+	}
+}
+
+void SVariablesOutliner::SetExport(const FWorkspaceOutlinerItemExport& InExport)
+{
+	Export = InExport;
+
+	UpdateAssets();
+}
+
+void SVariablesOutliner::HandleAssetLoaded(const FSoftObjectPath& InSoftObjectPath, UAnimNextRigVMAsset* InAsset)
+{	
+	if(InAsset && Assets.Contains(TSoftObjectPtr<UAnimNextRigVMAsset>(InAsset)))
+	{
+		RegisterAssetDelegates(InAsset);
+
+		FullRefresh();
+	}
+}
+
+void SVariablesOutliner::OnEditorDataModified(UAnimNextRigVMAssetEditorData* InEditorData, EAnimNextEditorDataNotifType InType, UObject* InSubject)
+{	
+	ensure(Assets.Contains(TSoftObjectPtr<UAnimNextRigVMAsset>(UncookedOnly::FUtils::GetAsset<UAnimNextRigVMAsset>(InEditorData))));
+
+	switch(InType)
+	{
+	case EAnimNextEditorDataNotifType::UndoRedo:
+	case EAnimNextEditorDataNotifType::EntryAdded:
+	case EAnimNextEditorDataNotifType::VariableCategoryChanged:
+	case EAnimNextEditorDataNotifType::CategoryAdded:
+	case EAnimNextEditorDataNotifType::CategoryChanged:
+		FullRefresh();
+		break;
+	case EAnimNextEditorDataNotifType::EntryRemoved:
+		UpdateAssets();
+		break;
+	default:
+		break;
+	}
+}
+
+void SVariablesOutliner::SetHighlightedItem(FSceneOutlinerTreeItemPtr Item) const
+{
+	GetTreeView()->Private_SetItemHighlighted(Item, true);
+}
+
+void SVariablesOutliner::ClearHighlightedItem(FSceneOutlinerTreeItemPtr Item) const
+{
+	GetTreeView()->Private_SetItemHighlighted(Item, false);
+}
+
+bool SVariablesOutliner::HasAssets() const
+{
+	return Assets.Num() > 0;
+}
+
+void SVariablesView::Construct(const FArguments& InArgs, TSharedRef<UE::Workspace::IWorkspaceEditor> InWorkspaceEditor)
+{
+	FSceneOutlinerInitializationOptions InitOptions;
+	InitOptions.OutlinerIdentifier = TEXT("AnimNextVariablesOutliner");
+	InitOptions.bShowHeaderRow = true;
+	InitOptions.ColumnMap.Add(FSceneOutlinerBuiltInColumnTypes::Label(), FSceneOutlinerColumnInfo(ESceneOutlinerColumnVisibility::Visible, 0, FCreateSceneOutlinerColumn(), false, 0.5f));
+	InitOptions.ColumnMap.Add(FVariablesOutlinerTypeColumn::GetID(), FSceneOutlinerColumnInfo(ESceneOutlinerColumnVisibility::Visible, 10, FCreateSceneOutlinerColumn::CreateLambda([](ISceneOutliner& InSceneOutliner) { return MakeShared<FVariablesOutlinerTypeColumn>(InSceneOutliner); }), false));
+	InitOptions.ColumnMap.Add(FVariablesOutlinerValueColumn::GetID(), FSceneOutlinerColumnInfo(ESceneOutlinerColumnVisibility::Visible, 20, FCreateSceneOutlinerColumn::CreateLambda([](ISceneOutliner& InSceneOutliner) { return MakeShared<FVariablesOutlinerValueColumn>(InSceneOutliner); }), true, 0.5f));
+	InitOptions.ColumnMap.Add(FSceneOutlinerSourceControlColumn::GetID(), FSceneOutlinerColumnInfo(ESceneOutlinerColumnVisibility::Visible, 30, FCreateSceneOutlinerColumn::CreateLambda([](ISceneOutliner& InSceneOutliner) { return MakeShared<FSceneOutlinerSourceControlColumn>(InSceneOutliner); }), true));
+	InitOptions.ColumnMap.Add(FVariablesOutlinerAccessSpecifierColumn::GetID(), FSceneOutlinerColumnInfo(ESceneOutlinerColumnVisibility::Visible, 40, FCreateSceneOutlinerColumn::CreateLambda([](ISceneOutliner& InSceneOutliner) { return MakeShared<FVariablesOutlinerAccessSpecifierColumn>(InSceneOutliner); }), false));
+	InitOptions.ModeFactory = FCreateSceneOutlinerMode::CreateLambda([this, WeakWorkspaceEditor=InWorkspaceEditor.ToWeakPtr()](SSceneOutliner* InOutliner) { return new FVariablesOutlinerMode(static_cast<SVariablesOutliner*>(InOutliner), WeakWorkspaceEditor.Pin().ToSharedRef()); });
+
+	VariablesOutliner = SNew(SVariablesOutliner, InitOptions);
+	VariablesOutliner->SetEnabled(MakeAttributeSP(VariablesOutliner.Get(), &SVariablesOutliner::HasAssets));
+
+	InWorkspaceEditor->OnFocusedDocumentChanged().AddSP(this, &SVariablesView::HandleFocusedDocumentChanged);
+	const Workspace::FWorkspaceDocument& Document = InWorkspaceEditor->GetFocusedWorkspaceDocument();
+	HandleFocusedDocumentChanged(Document);
+
+	ChildSlot
+	[
+		VariablesOutliner.ToSharedRef()
+	];
+}
+
+void SVariablesView::SetExportDirectly(const FWorkspaceOutlinerItemExport& InExport) const
+{
+	if (VariablesOutliner.IsValid())
+	{
+		VariablesOutliner->SetExport(InExport);
+	}
+}
+
+void SVariablesView::HandleFocusedDocumentChanged(const UE::Workspace::FWorkspaceDocument& InDocument) const
+{
+	const FWorkspaceOutlinerItemExport& Export = InDocument.Export;
+	VariablesOutliner->SetExport(Export);
+}
+
+FAnimNextVariablesTabSummoner::FAnimNextVariablesTabSummoner(TSharedPtr<UE::Workspace::IWorkspaceEditor> InHostingApp)
+	: FWorkflowTabFactory(VariablesTabName, StaticCastSharedPtr<FAssetEditorToolkit>(InHostingApp))
+{
+	TabLabel = LOCTEXT("AnimNextVariablesTabLabel", "Variables");
+	TabIcon = FSlateIcon("EditorStyle", "LevelEditor.Tabs.Outliner");
+	ViewMenuDescription = LOCTEXT("AnimNextVariablesTabMenuDescription", "Variables");
+	ViewMenuTooltip = LOCTEXT("AnimNextVariablesTabToolTip", "Shows the Variables tab.");
+	bIsSingleton = true;
+
+	VariablesView = SNew(SVariablesView, InHostingApp.ToSharedRef());
+
+	const Workspace::FWorkspaceDocument& Document = InHostingApp->GetFocusedWorkspaceDocument();
+	VariablesView->VariablesOutliner->SetExport(Document.Export);
+}
+
+TSharedRef<SWidget> FAnimNextVariablesTabSummoner::CreateTabBody(const FWorkflowTabSpawnInfo& Info) const
+{
+	return VariablesView.ToSharedRef();
+}
+
+FText FAnimNextVariablesTabSummoner::GetTabToolTipText(const FWorkflowTabSpawnInfo& Info) const
+{
+	return ViewMenuTooltip;
+}
+
+}
+
+#undef LOCTEXT_NAMESPACE
\ No newline at end of file
diff --git a/Engine/Plugins/Experimental/UAF/UAF/Source/UAFEditor/Private/Variables/SVariablesView.h b/Engine/Plugins/Experimental/UAF/UAF/Source/UAFEditor/Private/Variables/SVariablesView.h
new file mode 100644
--- /dev/null
+++ b/Engine/Plugins/Experimental/UAF/UAF/Source/UAFEditor/Private/Variables/SVariablesView.h
@@ -0,0 +1,92 @@
+﻿// Copyright Epic Games, Inc. All Rights Reserved.
+
+#pragma once
+
+#include "ContentBrowserDelegates.h"
+#include "SSceneOutliner.h"
+#include "WorkspaceAssetRegistryInfo.h"
+#include "Editor/Persona/Private/SAnimAssetFindReplace.h"
+#include "Widgets/SCompoundWidget.h"
+#include "WorkflowOrientedApp/WorkflowTabFactory.h"
+
+enum class EAnimNextEditorDataNotifType : uint8;
+class UAnimNextRigVMAsset;
+struct FWorkspaceOutlinerItemExport;
+class SPositiveActionButton;
+class UAnimNextRigVMAssetEditorData;
+
+namespace UE::Workspace
+{
+	class IWorkspaceEditor;
+	struct FWorkspaceDocument;
+}
+
+namespace UE::UAF::Editor
+{
+	class SVariablesOutliner;
+}
+
+namespace UE::UAF::Editor
+{
+
+extern const FLazyName VariablesTabName;
+
+class SVariablesOutliner : public SSceneOutliner
+{
+public:
+	void SetExport(const FWorkspaceOutlinerItemExport& InExport);
+	void OnEditorDataModified(UAnimNextRigVMAssetEditorData* InEditorData, EAnimNextEditorDataNotifType InType, UObject* InSubject);
+
+	void SetHighlightedItem(FSceneOutlinerTreeItemPtr Item) const;
+	void ClearHighlightedItem(FSceneOutlinerTreeItemPtr Item) const; 
+	
+	bool HasAssets() const;
+protected:
+	
+	void SetAssets(TConstArrayView<TSoftObjectPtr<UAnimNextRigVMAsset>> InAssets);
+	void UpdateAssets();
+
+	void HandleAssetLoaded(const FSoftObjectPath& InSoftObjectPath, UAnimNextRigVMAsset* InAsset);
+	
+	void RegisterAssetDelegates(const UAnimNextRigVMAsset* InAsset);
+	void UnregisterAssetDelegates(const UAnimNextRigVMAsset* InAsset);
+private:
+	friend class FVariablesOutlinerMode;
+	friend class FVariablesOutlinerHierarchy;
+
+	TArray<TSoftObjectPtr<UAnimNextRigVMAsset>> Assets;
+	FWorkspaceOutlinerItemExport Export;
+};
+
+class SVariablesView : public SCompoundWidget
+{
+	SLATE_BEGIN_ARGS(SVariablesView) {}
+	SLATE_END_ARGS()
+
+	void Construct(const FArguments& InArgs, TSharedRef<UE::Workspace::IWorkspaceEditor> InWorkspaceEditor);
+
+	void SetExportDirectly(const FWorkspaceOutlinerItemExport& InExport) const;
+private:
+	void HandleFocusedDocumentChanged(const UE::Workspace::FWorkspaceDocument& InDocument) const;
+
+private:
+	friend struct FAnimNextVariablesTabSummoner;
+	TSharedPtr<SVariablesOutliner> VariablesOutliner;
+};
+
+struct FAnimNextVariablesTabSummoner : public FWorkflowTabFactory
+{
+public:
+	FAnimNextVariablesTabSummoner(TSharedPtr<UE::Workspace::IWorkspaceEditor> InHostingApp);
+
+private:
+	// FWorkflowTabFactory interface
+	virtual TSharedRef<SWidget> CreateTabBody(const FWorkflowTabSpawnInfo& Info) const override;
+	virtual FText GetTabToolTipText(const FWorkflowTabSpawnInfo& Info) const override;
+
+	// The widget this tab spawner wraps
+	TSharedPtr<SVariablesView> VariablesView;
+};
+
+
+};
\ No newline at end of file
diff --git a/Engine/Plugins/Experimental/UAF/UAF/Source/UAFEditor/Private/Workspace/AnimNextWorkspaceEditorMode.cpp b/Engine/Plugins/Experimental/UAF/UAF/Source/UAFEditor/Private/Workspace/AnimNextWorkspaceEditorMode.cpp
new file mode 100644
--- /dev/null
+++ b/Engine/Plugins/Experimental/UAF/UAF/Source/UAFEditor/Private/Workspace/AnimNextWorkspaceEditorMode.cpp
@@ -0,0 +1,393 @@
+// Copyright Epic Games, Inc. All Rights Reserved.
+
+#include "AnimNextWorkspaceEditorMode.h"
+
+#include "AnimNextEditorContext.h"
+#include "AnimNextEditorModule.h"
+#include "AnimNextRigVMAsset.h"
+#include "AnimNextRigVMAssetEditorData.h"
+#include "AnimNextScopedCompileJob.h"
+#include "ContextObjectStore.h"
+#include "IAssetCompilationHandler.h"
+#include "InteractiveToolManager.h"
+#include "IWorkspaceEditor.h"
+#include "IWorkspaceEditorModule.h"
+#include "RigVMCommands.h"
+#include "UncookedOnlyUtils.h"
+#include "Logging/MessageLog.h"
+#include "Modules/ModuleManager.h"
+#include "Toolkits/BaseToolkit.h"
+
+#define LOCTEXT_NAMESPACE "AnimNextWorkspaceEditorMode"
+
+const FEditorModeID UAnimNextWorkspaceEditorMode::EM_AnimNextWorkspace("AnimNextWorkspace");
+
+UAnimNextWorkspaceEditorMode::UAnimNextWorkspaceEditorMode()
+{
+	Info = FEditorModeInfo(UAnimNextWorkspaceEditorMode::EM_AnimNextWorkspace,
+		LOCTEXT("AnimNextWorkspaceEditorModeName", "AnimNextWorkspaceEditorMode"),
+		FSlateIcon(),
+		false);
+}
+
+void UAnimNextWorkspaceEditorMode::Enter()
+{
+	Super::Enter();
+
+	TSharedPtr<UE::Workspace::IWorkspaceEditor> WorkspaceEditor = GetWorkspaceEditor();
+	if(!WorkspaceEditor.IsValid())
+	{
+		return;
+	}
+	
+	WorkspaceEditor->OnFocusedDocumentChanged().AddUObject(this, &UAnimNextWorkspaceEditorMode::HandleFocusedDocumentChanged);
+
+	TArray<UObject*> Assets;
+	WorkspaceEditor->GetOpenedAssets<UAnimNextRigVMAsset>(Assets);
+	for(UObject* Asset : Assets)
+	{
+		UAnimNextRigVMAsset* AnimNextRigVMAsset = static_cast<UAnimNextRigVMAsset*>(Asset);
+		UAnimNextRigVMAssetEditorData* EditorData = UE::UAF::UncookedOnly::FUtils::GetEditorData<UAnimNextRigVMAssetEditorData>(AnimNextRigVMAsset);
+
+		EditorData->RigVMCompiledEvent.RemoveAll(this);
+		EditorData->RigVMCompiledEvent.AddUObject(this, &UAnimNextWorkspaceEditorMode::HandleRigVMCompiledEvent);
+		EditorData->RigVMGraphModifiedEvent.RemoveAll(this);
+		EditorData->RigVMGraphModifiedEvent.AddUObject(this, &UAnimNextWorkspaceEditorMode::HandleRigVMModifiedEvent);
+
+		WeakAssets.Add(Asset);
+	}
+
+	UpdateCompileStatus();
+}
+
+void UAnimNextWorkspaceEditorMode::Exit()
+{
+	Super::Exit();
+
+	for (TPair<FObjectKey, TSharedRef<UE::UAF::Editor::IAssetCompilationHandler>>& AssetCompiler : AssetCompilers)
+	{
+		TSharedRef<UE::UAF::Editor::IAssetCompilationHandler> AssetCompilerRef = AssetCompiler.Value;
+		AssetCompilerRef->OnCompileStatusChanged().Unbind();
+	}
+
+	AssetCompilers.Reset();
+
+	for(TObjectKey<UObject> Asset : WeakAssets)
+	{
+		UAnimNextRigVMAsset* AnimNextRigVMAsset = Cast<UAnimNextRigVMAsset>(Asset.ResolveObjectPtr());
+		if(AnimNextRigVMAsset == nullptr)
+		{
+			continue;;
+		}
+		
+		UAnimNextRigVMAssetEditorData* EditorData = UE::UAF::UncookedOnly::FUtils::GetEditorData<UAnimNextRigVMAssetEditorData>(AnimNextRigVMAsset);
+
+		EditorData->RigVMCompiledEvent.RemoveAll(this);
+		EditorData->RigVMGraphModifiedEvent.RemoveAll(this);
+	}
+
+	WeakAssets.Reset();
+}
+
+void UAnimNextWorkspaceEditorMode::BindCommands()
+{
+	Super::BindCommands();
+
+	TSharedPtr<UE::Workspace::IWorkspaceEditor> WorkspaceEditor = GetWorkspaceEditor();
+	if(!WorkspaceEditor.IsValid())
+	{
+		return;
+	}
+
+	TSharedRef<FUICommandList> ToolkitCommands = WorkspaceEditor->GetToolkitCommands();
+
+	const UE::UAF::FRigVMCommands& RigVMCommands = UE::UAF::FRigVMCommands::Get();
+	ToolkitCommands->MapAction(RigVMCommands.Compile,
+		FExecuteAction::CreateUObject(this, &UAnimNextWorkspaceEditorMode::HandleCompile));
+
+	ToolkitCommands->MapAction(RigVMCommands.AutoCompile,
+		FExecuteAction::CreateUObject(this, &UAnimNextWorkspaceEditorMode::HandleAutoCompile),
+		FCanExecuteAction(),
+		FIsActionChecked::CreateUObject(this, &UAnimNextWorkspaceEditorMode::IsAutoCompileChecked));
+
+	ToolkitCommands->MapAction(RigVMCommands.CompileWholeWorkspace,
+		FExecuteAction::CreateUObject(this, &UAnimNextWorkspaceEditorMode::HandleCompileWholeWorkspace),
+		FCanExecuteAction(),
+		FIsActionChecked::CreateUObject(this, &UAnimNextWorkspaceEditorMode::IsCompileWholeWorkspaceChecked));
+
+	ToolkitCommands->MapAction(RigVMCommands.CompileDirtyFiles,
+		FExecuteAction::CreateUObject(this, &UAnimNextWorkspaceEditorMode::HandleCompileDirtyFiles),
+		FCanExecuteAction(),
+		FIsActionChecked::CreateUObject(this, &UAnimNextWorkspaceEditorMode::IsCompileDirtyFilesChecked));
+}
+
+void UAnimNextWorkspaceEditorMode::HandleCompile()
+{
+	TSharedPtr<UE::Workspace::IWorkspaceEditor> WorkspaceEditor = GetWorkspaceEditor();
+	if(!WorkspaceEditor.IsValid())
+	{
+		return;
+	}
+
+	TArray<UObject*> Assets;
+
+	UObject* Asset = nullptr;
+	FText JobName = LOCTEXT("DefaultJobName", "Job");
+	if(IsCompileWholeWorkspaceChecked())
+	{
+		WorkspaceEditor->GetOpenedAssets<UObject>(Assets);
+		if(Assets.Num() == 1)
+		{
+			Asset = Assets[0];
+			JobName = FText::FromName(Asset->GetFName());
+		}
+		else
+		{
+			Asset = WorkspaceEditor->GetWorkspaceAsset();
+			JobName = FText::FromName(Asset->GetFName());
+		}
+	}
+	else if(UObject* FocusedDocument = WorkspaceEditor->GetFocusedWorkspaceDocument().GetObject())
+	{
+		// Find the outer asset for this document
+		Asset = FocusedDocument;
+		
+		while(Asset && (!Asset->IsAsset() || Asset->IsPackageExternal()))
+		{
+			Asset = Asset->GetOuter();
+		}
+
+		if(Asset)
+		{
+			JobName = FText::FromName(Asset->GetFName());
+			Assets.Add(Asset);
+		}
+	}
+
+	if(IsCompileDirtyFilesChecked())
+	{
+		for(int32 AssetIndex = 0; AssetIndex < Assets.Num(); ++AssetIndex)
+		{
+			UObject* AssetToCheck = Assets[AssetIndex];
+			if(TSharedPtr<UE::UAF::Editor::IAssetCompilationHandler> FoundCompiler = GetAssetCompiler(AssetToCheck))
+			{
+				UE::UAF::Editor::ECompileStatus AssetCompileStatus = FoundCompiler->GetCompileStatus(WorkspaceEditor.ToSharedRef(), AssetToCheck);
+				if( AssetCompileStatus != UE::UAF::Editor::ECompileStatus::Dirty &&
+					AssetCompileStatus != UE::UAF::Editor::ECompileStatus::Error)
+				{
+					Assets.RemoveAtSwap(AssetIndex);
+					AssetIndex--;
+				}
+			}
+		}
+	}
+
+	if(Assets.Num() == 0)
+	{
+		return;
+	}
+
+	// Start a batch compile scope
+	UE::UAF::UncookedOnly::FScopedCompileJob CompileJob(JobName, Assets);
+
+	UE::UAF::Editor::FAnimNextEditorModule& AnimNextEditorModule = FModuleManager::LoadModuleChecked<UE::UAF::Editor::FAnimNextEditorModule>("UAFEditor");
+	for(UObject* AssetToCompile : Assets)
+	{
+		if(TSharedPtr<UE::UAF::Editor::IAssetCompilationHandler> FoundCompiler = GetAssetCompiler(AssetToCompile))
+		{
+			FoundCompiler->Compile(WorkspaceEditor.ToSharedRef(), AssetToCompile);
+		}
+	}
+}
+
+void UAnimNextWorkspaceEditorMode::HandleAutoCompile()
+{
+	State.bAutoCompile = !State.bAutoCompile;
+
+	TSharedPtr<UE::Workspace::IWorkspaceEditor> WorkspaceEditor = GetWorkspaceEditor();
+	if(WorkspaceEditor.IsValid())
+	{
+		PropagateAutoCompile(WorkspaceEditor.ToSharedRef(), State.bAutoCompile);
+	}
+}
+
+void UAnimNextWorkspaceEditorMode::PropagateAutoCompile(TSharedRef<UE::Workspace::IWorkspaceEditor> InWorkspaceEditor, bool bInAutoCompile)
+{
+	TArray<UObject*> Assets;
+	InWorkspaceEditor->GetOpenedAssets(Assets);
+	UE::UAF::Editor::FAnimNextEditorModule& AnimNextEditorModule = FModuleManager::LoadModuleChecked<UE::UAF::Editor::FAnimNextEditorModule>("UAFEditor");
+	for(UObject* Asset : Assets)
+	{
+		if(TSharedPtr<UE::UAF::Editor::IAssetCompilationHandler> FoundCompiler = GetAssetCompiler(Asset))
+		{
+			FoundCompiler->SetAutoCompile(InWorkspaceEditor, Asset, bInAutoCompile);
+		}
+	}
+}
+
+bool UAnimNextWorkspaceEditorMode::IsAutoCompileChecked() const
+{
+	return State.bAutoCompile;
+}
+
+void UAnimNextWorkspaceEditorMode::HandleCompileWholeWorkspace()
+{
+	State.bCompileWholeWorkspace = !State.bCompileWholeWorkspace;
+}
+
+bool UAnimNextWorkspaceEditorMode::IsCompileWholeWorkspaceChecked() const
+{
+	return State.bCompileWholeWorkspace;
+}
+
+void UAnimNextWorkspaceEditorMode::HandleCompileDirtyFiles()
+{
+	State.bCompileDirtyFiles = !State.bCompileDirtyFiles;
+}
+
+bool UAnimNextWorkspaceEditorMode::IsCompileDirtyFilesChecked() const
+{
+	return State.bCompileDirtyFiles;
+}
+
+void UAnimNextWorkspaceEditorMode::HandleFocusedDocumentChanged(const UE::Workspace::FWorkspaceDocument& InDocument)
+{
+	UObject* InObject = InDocument.GetObject();
+	
+	if(InObject == nullptr)
+	{
+		return;
+	}
+
+	UAnimNextRigVMAsset* AnimNextRigVMAsset = Cast<UAnimNextRigVMAsset>(InObject);
+	if(AnimNextRigVMAsset == nullptr)
+	{
+		AnimNextRigVMAsset = InObject->GetTypedOuter<UAnimNextRigVMAsset>();
+	}
+	
+	if(AnimNextRigVMAsset == nullptr)
+	{
+		return;
+	}
+
+	// Ensure auto-compilation is propagated for newly opened asset 
+	UAnimNextRigVMAssetEditorData* EditorData = UE::UAF::UncookedOnly::FUtils::GetEditorData<UAnimNextRigVMAssetEditorData>(AnimNextRigVMAsset);
+	EditorData->SetAutoVMRecompile(State.bAutoCompile);
+
+	// Subscribe to asset compilation/modification
+	EditorData->RigVMCompiledEvent.RemoveAll(this);
+	EditorData->RigVMCompiledEvent.AddUObject(this, &UAnimNextWorkspaceEditorMode::HandleRigVMCompiledEvent);
+	EditorData->RigVMGraphModifiedEvent.RemoveAll(this);
+	EditorData->RigVMGraphModifiedEvent.AddUObject(this, &UAnimNextWorkspaceEditorMode::HandleRigVMModifiedEvent);
+
+	WeakAssets.Add(AnimNextRigVMAsset);
+}
+
+void UAnimNextWorkspaceEditorMode::UpdateCompileStatus()
+{
+	bUpdateCompileStatus = false;
+
+	TSharedPtr<UE::Workspace::IWorkspaceEditor> WorkspaceEditor = GetWorkspaceEditor();
+	if(!WorkspaceEditor.IsValid())
+	{
+		return;
+	}
+
+	int32 NumAssetsWithWarnings = 0;
+	int32 NumAssetsWithErrors = 0;
+	int32 NumAssetsDirty = 0;
+
+	TArray<UObject*> Assets;
+	WorkspaceEditor->GetOpenedAssets<UAnimNextRigVMAsset>(Assets);
+
+	for(UObject* Asset : Assets)
+	{
+		if(TSharedPtr<UE::UAF::Editor::IAssetCompilationHandler> FoundCompiler = GetAssetCompiler(Asset))
+		{
+			UE::UAF::Editor::ECompileStatus AssetCompileStatus = FoundCompiler->GetCompileStatus(WorkspaceEditor.ToSharedRef(), Asset);
+			if(AssetCompileStatus == UE::UAF::Editor::ECompileStatus::Dirty)
+			{
+				NumAssetsDirty++;
+			}
+
+			if(AssetCompileStatus == UE::UAF::Editor::ECompileStatus::Error)
+			{
+				NumAssetsWithErrors++;
+			}
+
+			if(AssetCompileStatus == UE::UAF::Editor::ECompileStatus::Warning)
+			{
+				NumAssetsWithWarnings++;
+			}
+		}
+	}
+
+	if(NumAssetsWithErrors > 0)
+	{
+		CompileStatus = UE::UAF::Editor::ECompileStatus::Error;
+	}
+	else if(NumAssetsWithWarnings > 0)
+	{
+		CompileStatus = UE::UAF::Editor::ECompileStatus::Warning;
+	}
+	else if(NumAssetsDirty > 0)
+	{
+		CompileStatus = UE::UAF::Editor::ECompileStatus::Dirty;
+	}
+	else
+	{
+		CompileStatus = UE::UAF::Editor::ECompileStatus::UpToDate;
+	}
+}
+
+void UAnimNextWorkspaceEditorMode::HandleRigVMCompiledEvent(UObject* InAsset, URigVM* InVM, FRigVMExtendedExecuteContext& InExtendedExecuteContext)
+{
+	UpdateCompileStatus();
+}
+
+void UAnimNextWorkspaceEditorMode::HandleRigVMModifiedEvent(ERigVMGraphNotifType InType, URigVMGraph* InGraph, UObject* InSubject)
+{
+	if(InGraph == nullptr)
+	{
+		return;
+	}
+
+	UAnimNextRigVMAssetEditorData* EditorData = InGraph->GetTypedOuter<UAnimNextRigVMAssetEditorData>();
+	if(EditorData->IsDirtyForRecompilation())
+	{
+		CompileStatus = UE::UAF::Editor::ECompileStatus::Dirty;
+	}
+}
+
+TSharedPtr<UE::Workspace::IWorkspaceEditor> UAnimNextWorkspaceEditorMode::GetWorkspaceEditor() const
+{
+	if (const UContextObjectStore* ContextObjectStore = GetToolManager()->GetContextObjectStore())
+	{
+		if (const UAnimNextEditorContext* Context = ContextObjectStore->FindContext<UAnimNextEditorContext>())
+		{
+			return Context->WeakWorkspaceEditor.Pin();
+		}
+	}
+
+	return nullptr;
+}
+
+TSharedPtr<UE::UAF::Editor::IAssetCompilationHandler> UAnimNextWorkspaceEditorMode::GetAssetCompiler(UObject* InAsset)
+{
+	if(TSharedRef<UE::UAF::Editor::IAssetCompilationHandler>* FoundCompiler = AssetCompilers.Find(InAsset))
+	{
+		return *FoundCompiler;
+	}
+
+	UE::UAF::Editor::FAnimNextEditorModule& AnimNextEditorModule = FModuleManager::LoadModuleChecked<UE::UAF::Editor::FAnimNextEditorModule>("UAFEditor");
+	if(const UE::UAF::Editor::FAssetCompilationHandlerFactoryDelegate* FoundFactory = AnimNextEditorModule.FindAssetCompilationHandlerFactory(InAsset->GetClass()))
+	{
+		TSharedRef<UE::UAF::Editor::IAssetCompilationHandler> NewCompiler = AssetCompilers.Add(InAsset, FoundFactory->Execute(InAsset));
+		NewCompiler->OnCompileStatusChanged().BindUObject(this, &UAnimNextWorkspaceEditorMode::UpdateCompileStatus);
+		return NewCompiler;
+	}
+
+	return nullptr;
+}
+
+#undef LOCTEXT_NAMESPACE
\ No newline at end of file
diff --git a/Engine/Plugins/Experimental/UAF/UAF/Source/UAFUncookedOnly/Internal/AnimNextControllerBase.h b/Engine/Plugins/Experimental/UAF/UAF/Source/UAFUncookedOnly/Internal/AnimNextControllerBase.h
new file mode 100644
--- /dev/null
+++ b/Engine/Plugins/Experimental/UAF/UAF/Source/UAFUncookedOnly/Internal/AnimNextControllerBase.h
@@ -0,0 +1,40 @@
+// Copyright Epic Games, Inc. All Rights Reserved.
+
+#pragma once
+
+#include "CoreMinimal.h"
+#include "RigVMModel/RigVMController.h"
+#include "AnimNextControllerBase.generated.h"
+
+class UAnimNextSharedVariables;
+class UAnimNextSharedVariableNode;
+class UAnimNextRigVMAssetEditorData;
+
+/**
+  * Implements AnimNext RigVM controller extensions
+  */
+UCLASS(MinimalAPI)
+class UAnimNextControllerBase : public URigVMController
+{
+	GENERATED_BODY()
+
+public:
+	// Add a shared variable from the specified asset
+	UFUNCTION(BlueprintCallable, Category = RigVMController)
+	UAFUNCOOKEDONLY_API UAnimNextSharedVariableNode* AddAssetSharedVariableNode(const UAnimNextSharedVariables* InAsset, const FName& InVariableName, const FString& InCPPType, UObject* InCPPTypeObject, bool bIsGetter, const FString& InDefaultValue, const FVector2D& InPosition = FVector2D::ZeroVector, const FString& InNodeName = TEXT(""), bool bSetupUndoRedo = true, bool bPrintPythonCommand = false);
+
+	// Add a shared variable from the specified struct
+	UFUNCTION(BlueprintCallable, Category = RigVMController)
+	UAFUNCOOKEDONLY_API UAnimNextSharedVariableNode* AddStructSharedVariableNode(const UScriptStruct* InStruct, const FName& InVariableName, const FString& InCPPType, UObject* InCPPTypeObject, bool bIsGetter, const FString& InDefaultValue, const FVector2D& InPosition = FVector2D::ZeroVector, const FString& InNodeName = TEXT(""), bool bSetupUndoRedo = true, bool bPrintPythonCommand = false);
+
+	// Refreshes a UAnimNextSharedVariableNode instance with provided data (similar to URigVMController::RefreshVariableNode)
+	UFUNCTION(BlueprintCallable, Category = RigVMController)
+	UAFUNCOOKEDONLY_API void RefreshSharedVariableNode(const FName& InNodeName, const FString& InSourceObjectPath, const FName& InVariableName, const FString& InCPPType, UObject* InCPPTypeObject, bool bSetupUndoRedo, bool bPrintPythonCommand);
+
+	// Add a shared variable from the specified source object
+	UFUNCTION(BlueprintCallable, Category = RigVMController)
+	UAFUNCOOKEDONLY_API UAnimNextSharedVariableNode* AddSharedVariableNode(const FString& InSourceObjectPath, const FName& InVariableName, const FString& InCPPType, const FString& InCPPTypeObjectPath, bool bIsGetter, const FString& InDefaultValue, const FVector2D& InPosition, const FString& InNodeName, bool bSetupUndoRedo, bool bPrintPythonCommand);
+	
+	UAnimNextSharedVariableNode* ReplaceVariableNodeWithSharedVariableNode(URigVMVariableNode* InVariableNode, FName InNewVariableName, const UObject* InAssetOrStruct, bool bSetupUndoRedo, bool bPrintPythonCommand);
+	URigVMVariableNode* ReplaceSharedVariableNodeWithVariableNode(UAnimNextSharedVariableNode* InVariableNode, FName InNewVariableName, bool bSetupUndoRedo, bool bPrintPythonCommand);
+};
diff --git a/Engine/Plugins/Experimental/UAF/UAF/Source/UAFUncookedOnly/Internal/AnimNextEdGraph.h b/Engine/Plugins/Experimental/UAF/UAF/Source/UAFUncookedOnly/Internal/AnimNextEdGraph.h
new file mode 100644
--- /dev/null
+++ b/Engine/Plugins/Experimental/UAF/UAF/Source/UAFUncookedOnly/Internal/AnimNextEdGraph.h
@@ -0,0 +1,34 @@
+﻿// Copyright Epic Games, Inc. All Rights Reserved.
+
+#pragma once
+
+#include "CoreMinimal.h"
+#include "EdGraph/RigVMEdGraph.h"
+#include "AnimNextEdGraph.generated.h"
+
+class UAnimNextRigVMAssetEditorData;
+
+namespace UE::UAF::UncookedOnly
+{
+	struct FUtils;
+}
+
+/**
+  * Wraps UEdGraph which represents the node graph
+  */
+UCLASS(MinimalAPI)
+class UAnimNextEdGraph : public URigVMEdGraph
+{
+	GENERATED_BODY()
+
+	friend class UAnimNextModule_EditorData;
+	friend class UAnimNextRigVMAssetEditorData;
+
+	// UObject interface
+	virtual void PostLoad() override;
+
+	// URigVMEdGraph interface
+	virtual FRigVMClient* GetRigVMClient() const override;
+	
+	void Initialize(UAnimNextRigVMAssetEditorData* InEditorData);
+};
\ No newline at end of file
diff --git a/Engine/Plugins/Experimental/UAF/UAF/Source/UAFUncookedOnly/Internal/UncookedOnlyUtils.h b/Engine/Plugins/Experimental/UAF/UAF/Source/UAFUncookedOnly/Internal/UncookedOnlyUtils.h
new file mode 100644
--- /dev/null
+++ b/Engine/Plugins/Experimental/UAF/UAF/Source/UAFUncookedOnly/Internal/UncookedOnlyUtils.h
@@ -0,0 +1,239 @@
+﻿// Copyright Epic Games, Inc. All Rights Reserved.
+
+#pragma once
+
+#include "CoreMinimal.h"
+#include "AnimNextAssetWorkspaceAssetUserData.h"
+#include "IAnimNextRigVMGraphInterface.h"
+#include "AssetRegistry/AssetData.h"
+#include "Param/ParamType.h"
+#include "RigVMCore/RigVMTemplate.h"
+#include "StructUtils/InstancedStruct.h"
+#include "RigVMCore/RigVMGraphFunctionDefinition.h"
+#include "UObject/SoftObjectPath.h"
+#include "RigVMModel/RigVMGraph.h"
+#include "AnimNextExports.h"
+#include "AssetRegistry/IAssetRegistry.h"
+
+#define UE_API UAFUNCOOKEDONLY_API
+
+class UAnimNextSharedVariables;
+struct FAnimNextSoftVariableReference;
+struct FAnimNextVariableReference;
+struct FAnimNextGetVariableCompileContext;
+struct FAnimNextVariableBindingData;
+struct FEdGraphPinType;
+struct FWorkspaceOutlinerItemExports;
+struct FWorkspaceOutlinerItemExport;
+struct FRigVMGraphFunctionData;
+struct FRigVMCompileSettings;
+struct FRigVMGraphFunctionHeaderArray;
+struct FGraphContextMenuBuilder;
+struct FRigVMFunction;
+class UAnimNextModule;
+class UAnimNextModule_EditorData;
+class UAnimNextEdGraph;
+class URigVMController;
+class URigVMGraph;
+class UAnimNextEdGraph;
+class UAnimNextRigVMAsset;
+class UAnimNextRigVMAssetEditorData;
+class UAnimNextRigVMAssetEntry;
+class URigVMUnitNode;
+class UAnimNextControllerBase;
+
+namespace UE::UAF
+{
+	static const FLazyName AnimNextPublicGraphFunctionsExportsRegistryTag = TEXT("AnimNextPublicGraphFunctions");
+	static const FLazyName ControlRigAssetPublicGraphFunctionsExportsRegistryTag = TEXT("PublicGraphFunctions");
+}
+
+namespace UE::UAF::UncookedOnly
+{
+
+extern UAFUNCOOKEDONLY_API TAutoConsoleVariable<bool> CVarDumpProgrammaticGraphs;
+
+struct FUtils
+{
+	static UE_API void CompileVariables(const FRigVMCompileSettings& InSettings, UAnimNextRigVMAsset* InAsset, FAnimNextGetVariableCompileContext& OutCompileContext);
+
+	static UE_API FInstancedPropertyBag MakePropertyBagForEditorData(const UAnimNextRigVMAssetEditorData* InEditorData, const FAnimNextGetVariableCompileContext& InCompileContext);
+	
+	static UE_API void CompileVariableBindings(const FRigVMCompileSettings& InSettings, UAnimNextRigVMAsset* InAsset, TArray<URigVMGraph*>& OutGraphs);
+
+	static UE_API void RecreateVM(UAnimNextRigVMAsset* InAsset);
+
+	// Get the corresponding asset from an asset's editor data (casts the outer appropriately)
+	static UE_API UAnimNextRigVMAsset* GetAsset(UAnimNextRigVMAssetEditorData* InEditorData);
+
+	template<typename AssetType, typename EditorDataType>
+	static AssetType* GetAsset(EditorDataType* InEditorData)
+	{
+		using NonConstEditorDataType = std::remove_const_t<EditorDataType>;
+		return CastChecked<AssetType>(GetAsset(const_cast<NonConstEditorDataType*>(InEditorData)));
+	}
+
+	// Get the corresponding editor data from an asset (casts the editor data appropriately)
+	static UE_API UAnimNextRigVMAssetEditorData* GetEditorData(UAnimNextRigVMAsset* InAsset);
+
+	template<typename EditorDataType, typename AssetType>
+	static EditorDataType* GetEditorData(AssetType* InAsset)
+	{
+		using NonConstAssetType = std::remove_const_t<AssetType>;
+		return CastChecked<EditorDataType>(GetEditorData(const_cast<NonConstAssetType*>(InAsset)));
+	}
+
+	/**
+	 * Get an AnimNext parameter type from an FEdGraphPinType.
+	 * Note that the returned handle may not be valid, so should be checked using IsValid() before use.
+	 **/
+	static UE_API FAnimNextParamType GetParamTypeFromPinType(const FEdGraphPinType& InPinType);
+
+	/**
+	 * Get an FEdGraphPinType from an AnimNext parameter type/handle.
+	 * Note that the returned pin type may not be valid.
+	 **/
+	static UE_API FEdGraphPinType GetPinTypeFromParamType(const FAnimNextParamType& InParamType);
+
+	/**
+	 * Get an FRigVMTemplateArgumentType from an AnimNext parameter type/handle.
+	 * Note that the returned pin type may not be valid.
+	 **/
+	static UE_API FRigVMTemplateArgumentType GetRigVMArgTypeFromParamType(const FAnimNextParamType& InParamType);
+
+	/** Set up a simple event graph */
+	static UE_API void SetupEventGraph(URigVMController* InController, UScriptStruct* InEventStruct, FName InEventName = NAME_None, bool bPrintPythonCommand = false);
+
+	// Gets the variables that are exported to the asset registry for an asset
+	static UE_API bool GetExportedVariablesForAsset(const FAssetData& InAsset, FAnimNextAssetRegistryExports& OutExports);
+
+	// Gets all the variables that are exported to the asset registry
+	static UE_API bool GetExportedVariablesFromAssetRegistry(TMap<FAssetData, FAnimNextAssetRegistryExports>& OutExports);
+
+	template<typename ExportDataType>
+	static bool GetExportsOfTypeForAsset(const FAssetData& InAsset, FAnimNextAssetRegistryExports& OutExports)
+	{
+		const int32 StartNum = OutExports.Exports.Num();
+		FAnimNextAssetRegistryExports AssetExports;
+		const FString TagValue = InAsset.GetTagValueRef<FString>(UE::UAF::ExportsAnimNextAssetRegistryTag);
+		if (FAnimNextAssetRegistryExports::StaticStruct()->ImportText(*TagValue, &AssetExports, nullptr, PPF_None, nullptr, FAnimNextAssetRegistryExports::StaticStruct()->GetName()) != nullptr)
+		{
+			Algo::TransformIf(AssetExports.Exports, OutExports.Exports,
+				[](const FAnimNextExport& Export)
+				{
+					return Export.Data.GetScriptStruct() == ExportDataType::StaticStruct();
+				},
+				[](const FAnimNextExport& Export)
+				{
+					return Export;
+				});
+		}
+
+		return OutExports.Exports.Num() > StartNum;
+	}
+
+	template<typename ExportDataType>
+	static bool GetExportsOfTypeFromAssetRegistry(TMap<FAssetData, FAnimNextAssetRegistryExports>& OutExports)
+	{
+		TArray<FAssetData> AssetData;
+		IAssetRegistry::GetChecked().GetAssetsByTags({UE::UAF::ExportsAnimNextAssetRegistryTag}, AssetData);
+
+		for (const FAssetData& Asset : AssetData)
+		{
+			const FString TagValue = Asset.GetTagValueRef<FString>(UE::UAF::ExportsAnimNextAssetRegistryTag);
+			FAnimNextAssetRegistryExports AssetExports;
+			if (FAnimNextAssetRegistryExports::StaticStruct()->ImportText(*TagValue, &AssetExports, nullptr, PPF_None, nullptr, FAnimNextAssetRegistryExports::StaticStruct()->GetName()) != nullptr)
+			{
+				FAnimNextAssetRegistryExports& Export = OutExports.Add(Asset);
+				Algo::TransformIf(AssetExports.Exports, Export.Exports,
+					[](const FAnimNextExport& Export)
+					{
+						return Export.Data.GetScriptStruct() == ExportDataType::StaticStruct();
+					},
+					[](const FAnimNextExport& Export)
+					{
+						return Export;
+					});
+			}
+		}
+
+		return OutExports.Num() > 0;
+	}
+
+	static void GetAssetVariableExports(const UAnimNextRigVMAssetEditorData* EditorData, FAnimNextAssetRegistryExports& OutExports, FAssetRegistryTagsContext Context);
+
+	// Gets all the functions that are exported to the asset registry for the specified Tag
+	static UE_API bool GetExportedFunctionsFromAssetRegistry(FName Tag, TMap<FAssetData, FRigVMGraphFunctionHeaderArray>& OutExports);
+
+	// Gets the exported public functions that are used by a RigVM asset
+	static void GetAssetFunctions(const UAnimNextRigVMAssetEditorData* InEditorData, FRigVMGraphFunctionHeaderArray& OutExports);
+
+	// Gets the non-exported private functions that are used by a RigVM asset
+	static UE_API void GetAssetPrivateFunctions(const UAnimNextRigVMAssetEditorData* InEditorData, FRigVMGraphFunctionHeaderArray& OutExports);
+
+	// Gets the asset-registry information needed for representing the contained data into the Workspace Outliner
+	// Note: We pass parent as index to avoid the ref from being invalidated due to realloc as the export array grows recursively
+	static UE_API void GetAssetWorkspaceExports(const UAnimNextRigVMAssetEditorData* EditorData, FWorkspaceOutlinerItemExports& OutExports, FAssetRegistryTagsContext Context);
+
+	// Attempts to determine the type from a variable name & asset/struct path
+	// If the name cannot be found, the returned type will be invalid
+	// Note that this is expensive and can query the asset registry
+	static UE_API FAnimNextParamType FindVariableType(const FAnimNextVariableReference& InVariableReference);
+	static UE_API FAnimNextParamType FindVariableType(const FAnimNextSoftVariableReference& InSoftVariableReference);
+
+	// Returns an user friendly name for the Function Library
+	static UE_API const FText& GetFunctionLibraryDisplayName();
+
+#if WITH_EDITOR
+	static UE_API void OpenProgrammaticGraphs(UAnimNextRigVMAssetEditorData* InEditorData, const TArray<URigVMGraph*>& ProgrammaticGraphs);
+#endif // WITH_EDITOR
+
+	// Make a variable name that we use as a wrapper for a function param or return
+	static UE_API FString MakeFunctionWrapperVariableName(FName InFunctionName, FName InVariableName);
+
+	// Make an event name that we use as a wrapper to call RigVM functions 
+	static UE_API FString MakeFunctionWrapperEventName(FName InFunctionName);
+
+	// Returns all unique variable names found within provided editor data, optionally recursing into SharedVariable entries
+	static UE_API void GetVariableNames(UAnimNextRigVMAssetEditorData* InEditorData, TArray<FName>& OutVariableNames, bool bRecursive = true);
+
+	// Return valid variable FName within the context of the provided EditorData (recurses into SharedVariable entries as well) 
+	static UE_API FName GetValidVariableName(UAnimNextRigVMAssetEditorData* InEditorData, FName InBaseName);
+
+	// Return valid variable FName within the context of the provided set of existing variable names
+	static UE_API FName GetValidVariableName(FName InBaseName, const TArrayView<FName> ExistingNames);
+
+	// Deletes the provided variable entry from its editor data, removing any references across the project its assets 
+	static UE_API void DeleteVariable(UAnimNextVariableEntry* VariableEntry, bool bSetupUndoRedo = true, bool bPrintPythonCommands = true);
+
+	// Deletes the provided variable entry from its editor data, renaming any references across the project its assets to reference it by the NewName
+	static UE_API void RenameVariable(UAnimNextVariableEntry* VariableEntry, const FName NewName, bool bSetupUndoRedo = true, bool bPrintPythonCommands = true);
+
+	// Moves the provided variable entry from its editor data to the NewOuter, replacing any references across the project its assets to reference it by its new outer asset
+	static UE_API void MoveVariableToAsset(UAnimNextVariableEntry* VariableEntry, UAnimNextRigVMAsset* NewOuter, bool bSetupUndoRedo = true, bool bPrintPythonCommands = true);
+
+	// Changes the provided variable entry its data type, replacing any references across the project to reference using its new data type
+	static UE_API void SetVariableType(UAnimNextVariableEntry* VariableEntry, const FAnimNextParamType NewType, bool bSetupUndoRedo = true, bool bPrintPythonCommands = true);
+
+	// Replaces the provided variable reference with ReferenceToReplaceWith within the context of InEditorData
+	static UE_API void ReplaceVariableReferences(UAnimNextRigVMAssetEditorData* InEditorData, const FAnimNextSoftVariableReference& ReferenceToFind, const FAnimNextSoftVariableReference& ReferenceToReplaceWith, bool bSetupUndoRedo = true, bool bPrintPythonCommands = true);
+	// Replaces the provided variable reference with ReferenceToReplaceWith within the context of the project
+	static UE_API void ReplaceVariableReferencesAcrossProject(const FAnimNextSoftVariableReference& ReferenceToFind, const FAnimNextSoftVariableReference& ReferenceToReplaceWith, bool bSetupUndoRedo = true, bool bPrintPythonCommands = true);
+private:
+	static UE_API void CompileVariableBindingsInternal(const FRigVMCompileSettings& InSettings, UAnimNextRigVMAsset* InAsset, TArray<URigVMGraph*>& OutGraphs, bool bInThreadSafe);
+	
+	static void GetSubGraphWorkspaceExportsRecursive(const UAnimNextRigVMAssetEditorData* EditorData, FWorkspaceOutlinerItemExports& OutExports, const FWorkspaceOutlinerItemExport& RootExport, int32 ParentExportIndex, URigVMEdGraph* RigVMEdGraph, FAssetRegistryTagsContext Context);
+	static void GetFunctionLibraryWorkspaceExportsRecursive(const UAnimNextRigVMAssetEditorData* EditorData, FWorkspaceOutlinerItemExports& OutExports, const FWorkspaceOutlinerItemExport& RootExport, int32 ParentExportIndex, const TArray<FRigVMGraphFunctionData>& PublicFunctions, const TArray<FRigVMGraphFunctionData>& PrivateFunctions);
+	static void GetFunctionsWorkspaceExportsRecursive(const UAnimNextRigVMAssetEditorData* EditorData, FWorkspaceOutlinerItemExports& OutExports, const FWorkspaceOutlinerItemExport& RootExport, int32 ParentExportIndex, const TArray<FRigVMGraphFunctionData>& Functions, bool bPublicFunctions);
+	
+	// Gets the exported variables that are used by a RigVM asset	
+	static void GetAssetVariableExports(const UAnimNextRigVMAssetEditorData* EditorData, TSet<FAnimNextExport>& OutExports, FAssetRegistryTagsContext Context, EAnimNextExportedVariableFlags InFlags = EAnimNextExportedVariableFlags::NoFlags);
+
+	static void GetStructVariableExports(const UScriptStruct* Struct, TSet<FAnimNextExport>& OutExports);
+	static void GetSubGraphVariableExportsRecursive(const URigVMEdGraph* RigVMEdGraph, TSet<FAnimNextExport>& OutExports);
+	static void ProcessPinAssetReferences(const URigVMPin* InPin, FWorkspaceOutlinerItemExports& OutExports, const FWorkspaceOutlinerItemExport& RootExport, int32 ParentExportIndex);
+};
+
+}
+
+#undef UE_API
diff --git a/Engine/Plugins/Experimental/UAF/UAF/Source/UAFUncookedOnly/Internal/Variables/AnimNextSharedVariables_EditorData.h b/Engine/Plugins/Experimental/UAF/UAF/Source/UAFUncookedOnly/Internal/Variables/AnimNextSharedVariables_EditorData.h
new file mode 100644
--- /dev/null
+++ b/Engine/Plugins/Experimental/UAF/UAF/Source/UAFUncookedOnly/Internal/Variables/AnimNextSharedVariables_EditorData.h
@@ -0,0 +1,31 @@
+// Copyright Epic Games, Inc. All Rights Reserved.
+
+#pragma once
+
+#include "CoreMinimal.h"
+#include "AnimNextExecuteContext.h"
+#include "AnimNextRigVMAssetEditorData.h"
+#include "AnimNextSharedVariables_EditorData.generated.h"
+
+class UAnimNextSharedVariablesFactory;
+
+namespace UE::UAF::UncookedOnly
+{
+	struct FUtils;
+}
+
+/** Editor data for AnimNext data interfaces */
+UCLASS()
+class UAFUNCOOKEDONLY_API UAnimNextSharedVariables_EditorData : public UAnimNextRigVMAssetEditorData
+{
+	GENERATED_BODY()
+
+protected:
+	// UAnimNextRigVMAssetEditorData interface
+	virtual UScriptStruct* GetExecuteContextStruct() const override { return FAnimNextExecuteContext::StaticStruct(); }
+	virtual TConstArrayView<TSubclassOf<UAnimNextRigVMAssetEntry>> GetEntryClasses() const override;
+	virtual void CustomizeNewAssetEntry(UAnimNextRigVMAssetEntry* InNewEntry) const override;
+
+	friend class UAnimNextSharedVariablesFactory;
+	friend UE::UAF::UncookedOnly::FUtils;
+};
