UAF: Moved/renamed AnimNext and AnimNextAnimGraph plugins

diff --git a/Engine/Plugins/Experimental/UAF/UAF/Source/UAFUncookedOnly/Private/AnimNextControllerBase.cpp b/Engine/Plugins/Experimental/UAF/UAF/Source/UAFUncookedOnly/Private/AnimNextControllerBase.cpp
new file mode 100644
--- /dev/null
+++ b/Engine/Plugins/Experimental/UAF/UAF/Source/UAFUncookedOnly/Private/AnimNextControllerBase.cpp
@@ -0,0 +1,193 @@
+// Copyright Epic Games, Inc. All Rights Reserved.
+
+#include "AnimNextControllerBase.h"
+#include "AnimNextSharedVariableNode.h"
+#include "RigVMPythonUtils.h"
+#include "Variables/AnimNextSharedVariables.h"
+
+UAnimNextSharedVariableNode* UAnimNextControllerBase::AddAssetSharedVariableNode(const UAnimNextSharedVariables* InAsset, const FName& InVariableName, const FString& InCPPType, UObject* InCPPTypeObject, bool bIsGetter, const FString& InDefaultValue, const FVector2D& InPosition, const FString& InNodeName, bool bSetupUndoRedo, bool bPrintPythonCommand)
+{
+	if (InAsset == nullptr)
+	{
+		ReportError(TEXT("Invalid asset supplied to AddAssetSharedVariableNode."));
+		return nullptr;
+	}
+
+	return AddSharedVariableNode(InAsset->GetPathName(), InVariableName, InCPPType, InCPPTypeObject->GetPathName(), bIsGetter, InDefaultValue, InPosition, InNodeName, bSetupUndoRedo, bPrintPythonCommand);
+}
+
+UAnimNextSharedVariableNode* UAnimNextControllerBase::AddStructSharedVariableNode(const UScriptStruct* InStruct, const FName& InVariableName, const FString& InCPPType, UObject* InCPPTypeObject, bool bIsGetter, const FString& InDefaultValue, const FVector2D& InPosition, const FString& InNodeName, bool bSetupUndoRedo, bool bPrintPythonCommand)
+{
+	if (InStruct == nullptr)
+	{
+		ReportError(TEXT("Invalid struct supplied to AddStructSharedVariableNode."));
+		return nullptr;
+	}
+
+	return AddSharedVariableNode(InStruct->GetPathName(), InVariableName, InCPPType, InCPPTypeObject->GetPathName(), bIsGetter, InDefaultValue, InPosition, InNodeName, bSetupUndoRedo, bPrintPythonCommand);
+}
+
+UAnimNextSharedVariableNode* UAnimNextControllerBase::AddSharedVariableNode(const FString& InSourceObjectPath, const FName& InVariableName, const FString& InCPPType, const FString& InCPPTypeObjectPath, bool bIsGetter, const FString& InDefaultValue, const FVector2D& InPosition, const FString& InNodeName, bool bSetupUndoRedo, bool bPrintPythonCommand)
+{
+	UObject* SourceObject = nullptr;
+	if (!InSourceObjectPath.IsEmpty())
+	{
+		SourceObject = RigVMTypeUtils::FindObjectFromCPPTypeObjectPath<UObject>(InSourceObjectPath);
+		if (SourceObject == nullptr)
+		{
+			ReportErrorf(TEXT("Cannot find cpp type object for path '%s'."), *InSourceObjectPath);
+			return nullptr;
+		}
+	}
+
+	UObject* CPPTypeObject = nullptr;
+	if (!InCPPTypeObjectPath.IsEmpty())
+	{
+		CPPTypeObject = RigVMTypeUtils::FindObjectFromCPPTypeObjectPath<UObject>(InCPPTypeObjectPath);
+		if (CPPTypeObject == nullptr)
+		{
+			ReportErrorf(TEXT("Cannot find cpp type object for path '%s'."), *InCPPTypeObjectPath);
+			return nullptr;
+		}
+	}
+
+	UAnimNextSharedVariableNode* VariableNode = CastChecked<UAnimNextSharedVariableNode>(AddVariableNode(InVariableName, UAnimNextSharedVariableNode::StaticClass(), InCPPType, CPPTypeObject, bIsGetter, InDefaultValue, InPosition, InNodeName, bSetupUndoRedo, false), ECastCheckedType::Type::NullAllowed);
+	if (VariableNode == nullptr)
+	{
+		return nullptr;
+	}
+
+	if (const UAnimNextSharedVariables* Asset = Cast<UAnimNextSharedVariables>(SourceObject))
+	{
+		VariableNode->Type = EAnimNextSharedVariablesType::Asset;
+		VariableNode->Asset = Asset;
+	}
+	else if (const UScriptStruct* Struct = Cast<UScriptStruct>(SourceObject))
+	{
+		VariableNode->Type = EAnimNextSharedVariablesType::Struct;
+		VariableNode->Struct = Struct;
+	}
+	else
+	{
+		return nullptr;
+	}
+
+	if (bPrintPythonCommand)
+	{
+		const FString GraphName = GetSchema()->GetSanitizedGraphName(GetGraph()->GetGraphName());
+
+		RigVMPythonUtils::Print(GetSchema()->GetGraphOuterName(GetGraph()), 
+			FString::Printf(TEXT("asset.get_controller_by_name('%s').add_shared_variable_node('%s', '%s', '%s', '%s', '%s', '%s', %s, %s)"),
+				*GraphName,
+				*InSourceObjectPath,
+				*InVariableName.ToString(),
+				*InCPPType,
+				*InCPPTypeObjectPath,
+				bIsGetter ? TEXT("True") : TEXT("False"),
+				*InDefaultValue,
+				*RigVMPythonUtils::Vector2DToPythonString(InPosition),
+				*InNodeName
+			)
+		);
+	}
+
+	return VariableNode;
+}
+
+void UAnimNextControllerBase::RefreshSharedVariableNode(const FName& InNodeName, const FString& InSourceObjectPath, const FName& InVariableName, const FString& InCPPType, UObject* InCPPTypeObject, bool bSetupUndoRedo, bool bPrintPythonCommand)
+{
+	UObject* SourceObject = nullptr;
+	if (!InSourceObjectPath.IsEmpty())
+	{
+		SourceObject = RigVMTypeUtils::FindObjectFromCPPTypeObjectPath<UObject>(InSourceObjectPath);
+	}
+
+	URigVMGraph* Graph = GetGraph();
+	check(Graph);
+	if (UAnimNextSharedVariableNode* VariableNode = Cast<UAnimNextSharedVariableNode>(Graph->FindNodeByName(InNodeName)))
+	{
+		if (bSetupUndoRedo)
+		{
+			VariableNode->Modify();
+		}
+		
+		constexpr bool bSetupOrphanPins = false;
+		RefreshVariableNode(InNodeName, InVariableName, InCPPType, InCPPTypeObject, bSetupUndoRedo, bSetupOrphanPins);
+
+		if (SourceObject)
+		{
+			if (const UAnimNextSharedVariables* Asset = Cast<UAnimNextSharedVariables>(SourceObject))
+			{
+				VariableNode->Type = EAnimNextSharedVariablesType::Asset;
+				VariableNode->Asset = Asset;
+			}
+			else if (const UScriptStruct* Struct = Cast<UScriptStruct>(SourceObject))
+			{
+				VariableNode->Type = EAnimNextSharedVariablesType::Struct;
+				VariableNode->Struct = Struct;
+			}
+		}
+		else
+		{
+			VariableNode->Type = EAnimNextSharedVariablesType::Asset;
+			VariableNode->Asset = nullptr;
+		}
+
+		Notify(ERigVMGraphNotifType::NodeTitleChanged, VariableNode);
+
+		if (bPrintPythonCommand)
+		{
+			
+		}
+	}
+}
+
+UAnimNextSharedVariableNode* UAnimNextControllerBase::ReplaceVariableNodeWithSharedVariableNode(URigVMVariableNode* InVariableNode, FName InNewVariableName, const UObject* InAssetOrStruct, bool bSetupUndoRedo, bool bPrintPythonCommand)
+{
+	check(InVariableNode);
+	check(InAssetOrStruct);
+	check(InNewVariableName != NAME_None);
+
+	UAnimNextSharedVariableNode* SharedVariableNode = AddSharedVariableNode(InAssetOrStruct->GetPathName(), InNewVariableName, InVariableNode->GetCPPType(), InVariableNode->GetCPPTypeObject() ? InVariableNode->GetCPPTypeObject()->GetPathName() : TEXT(""), InVariableNode->IsGetter(), InVariableNode->GetDefaultValue(), InVariableNode->GetPosition(), FString(), bSetupUndoRedo, bPrintPythonCommand);
+	if (SharedVariableNode)
+	{
+		if (!InVariableNode->IsGetter())
+		{
+			RewireLinks(InVariableNode->GetValuePin(), SharedVariableNode->GetValuePin(), true, bSetupUndoRedo);
+			RewireLinks(InVariableNode->FindExecutePin(), SharedVariableNode->FindExecutePin(), true, bSetupUndoRedo);
+			RewireLinks(InVariableNode->FindExecutePin(), SharedVariableNode->FindExecutePin(), false, bSetupUndoRedo);
+		}
+		else
+		{
+			RewireLinks(InVariableNode->GetValuePin(), SharedVariableNode->GetValuePin(), false, bSetupUndoRedo);
+		}
+		RemoveNode(InVariableNode, bSetupUndoRedo, bPrintPythonCommand);
+	}
+
+	return SharedVariableNode;
+}
+
+URigVMVariableNode* UAnimNextControllerBase::ReplaceSharedVariableNodeWithVariableNode(UAnimNextSharedVariableNode* InVariableNode, FName InNewVariableName, bool bSetupUndoRedo, bool bPrintPythonCommand)
+{
+	check(InVariableNode);
+	check(InNewVariableName != NAME_None);
+
+
+	URigVMVariableNode* VariableNode = AddVariableNode(InNewVariableName, InVariableNode->GetCPPType(), InVariableNode->GetCPPTypeObject(), InVariableNode->IsGetter(), InVariableNode->GetDefaultValue(), InVariableNode->GetPosition(), FString(), bSetupUndoRedo, bPrintPythonCommand);
+	if (VariableNode)
+	{
+		if (!InVariableNode->IsGetter())
+		{
+			RewireLinks(InVariableNode->GetValuePin(), VariableNode->GetValuePin(), true, bSetupUndoRedo);
+			RewireLinks(InVariableNode->FindExecutePin(), VariableNode->FindExecutePin(), true, bSetupUndoRedo);
+			RewireLinks(InVariableNode->FindExecutePin(), VariableNode->FindExecutePin(), false, bSetupUndoRedo);
+		}
+		else
+		{
+			RewireLinks(InVariableNode->GetValuePin(), VariableNode->GetValuePin(), false, bSetupUndoRedo);
+		}
+		RemoveNode(InVariableNode, bSetupUndoRedo, bPrintPythonCommand);
+	}
+
+	return VariableNode;
+}
\ No newline at end of file
diff --git a/Engine/Plugins/Experimental/UAF/UAF/Source/UAFUncookedOnly/Private/AnimNextEdGraph.cpp b/Engine/Plugins/Experimental/UAF/UAF/Source/UAFUncookedOnly/Private/AnimNextEdGraph.cpp
new file mode 100644
--- /dev/null
+++ b/Engine/Plugins/Experimental/UAF/UAF/Source/UAFUncookedOnly/Private/AnimNextEdGraph.cpp
@@ -0,0 +1,30 @@
+﻿// Copyright Epic Games, Inc. All Rights Reserved.
+
+#include "AnimNextEdGraph.h"
+#include "Module/AnimNextModule_EditorData.h"
+
+void UAnimNextEdGraph::PostLoad()
+{
+	Super::PostLoad();
+
+	UAnimNextRigVMAssetEditorData* EditorData = GetTypedOuter<UAnimNextRigVMAssetEditorData>();
+	check(EditorData);
+	Initialize(EditorData);
+}
+
+void UAnimNextEdGraph::Initialize(UAnimNextRigVMAssetEditorData* InEditorData)
+{
+	InEditorData->RigVMGraphModifiedEvent.RemoveAll(this);
+	InEditorData->RigVMGraphModifiedEvent.AddUObject(this, &UAnimNextEdGraph::HandleModifiedEvent);
+	InEditorData->RigVMCompiledEvent.RemoveAll(this);
+	InEditorData->RigVMCompiledEvent.AddUObject(this, &UAnimNextEdGraph::HandleVMCompiledEvent);
+}
+
+FRigVMClient* UAnimNextEdGraph::GetRigVMClient() const
+{
+	if (const UAnimNextRigVMAssetEditorData* EditorData = GetTypedOuter<UAnimNextRigVMAssetEditorData>())
+	{
+		return const_cast<FRigVMClient*>(EditorData->GetRigVMClient());
+	}
+	return nullptr;
+}
diff --git a/Engine/Plugins/Experimental/UAF/UAF/Source/UAFUncookedOnly/Private/AnimNextEdGraphNode.cpp b/Engine/Plugins/Experimental/UAF/UAF/Source/UAFUncookedOnly/Private/AnimNextEdGraphNode.cpp
new file mode 100644
--- /dev/null
+++ b/Engine/Plugins/Experimental/UAF/UAF/Source/UAFUncookedOnly/Private/AnimNextEdGraphNode.cpp
@@ -0,0 +1,3 @@
+﻿// Copyright Epic Games, Inc. All Rights Reserved.
+
+#include "AnimNextEdGraphNode.h"
diff --git a/Engine/Plugins/Experimental/UAF/UAF/Source/UAFUncookedOnly/Private/AnimNextEdGraphSchema.cpp b/Engine/Plugins/Experimental/UAF/UAF/Source/UAFUncookedOnly/Private/AnimNextEdGraphSchema.cpp
new file mode 100644
--- /dev/null
+++ b/Engine/Plugins/Experimental/UAF/UAF/Source/UAFUncookedOnly/Private/AnimNextEdGraphSchema.cpp
@@ -0,0 +1,21 @@
+﻿// Copyright Epic Games, Inc. All Rights Reserved.
+
+#include "AnimNextEdGraphSchema.h"
+
+#include "AnimNextRigVMAsset.h"
+#include "Entries/AnimNextRigVMAssetEntry.h"
+
+#define LOCTEXT_NAMESPACE "AnimNextEdGraphSchema"
+
+void UAnimNextEdGraphSchema::GetGraphDisplayInformation(const UEdGraph& Graph, /*out*/ FGraphDisplayInfo& DisplayInfo) const
+{
+	Super::GetGraphDisplayInformation(Graph, DisplayInfo);
+	
+	if(UAnimNextRigVMAssetEntry* AssetEntry = Cast<UAnimNextRigVMAssetEntry>(Graph.GetOuter()))
+	{
+		DisplayInfo.DisplayName = FText::Format(LOCTEXT("GraphTabTitleFormat", "{0}: {1}"), FText::FromName(AssetEntry->GetEntryName()), FText::FromName(AssetEntry->GetTypedOuter<UAnimNextRigVMAsset>()->GetFName()));
+		DisplayInfo.Tooltip = FText::Format(LOCTEXT("GraphTabTooltipFormat", "{0} in:\n{1}"), FText::FromName(AssetEntry->GetEntryName()), FText::FromString(AssetEntry->GetTypedOuter<UAnimNextRigVMAsset>()->GetPathName()));
+	}
+}
+
+#undef LOCTEXT_NAMESPACE
\ No newline at end of file
