UAF: Moved/renamed AnimNext and AnimNextAnimGraph plugins

diff --git a/Engine/Plugins/Experimental/UAF/UAF/Source/UAFUncookedOnly/Private/AnimNextRigVMAssetEditorData.cpp b/Engine/Plugins/Experimental/UAF/UAF/Source/UAFUncookedOnly/Private/AnimNextRigVMAssetEditorData.cpp
new file mode 100644
--- /dev/null
+++ b/Engine/Plugins/Experimental/UAF/UAF/Source/UAFUncookedOnly/Private/AnimNextRigVMAssetEditorData.cpp
@@ -0,0 +1,3094 @@
+// Copyright Epic Games, Inc. All Rights Reserved.
+
+#include "AnimNextRigVMAssetEditorData.h"
+#include "AnimNextEdGraph.h"
+#include "AnimNextEdGraphSchema.h"
+#include "AnimNextRigVMAsset.h"
+#include "Compilation/AnimNextRigVMAssetCompileContext.h"
+#include "Compilation/AnimNextGetFunctionHeaderCompileContext.h"
+#include "Compilation/AnimNextGetVariableCompileContext.h"
+#include "Compilation/AnimNextGetGraphCompileContext.h"
+#include "Compilation/AnimNextProcessGraphCompileContext.h"
+#include "Entries/AnimNextRigVMAssetEntry.h"
+#include "AnimNextRigVMAssetSchema.h"
+#include "AnimNextAssetWorkspaceAssetUserData.h"
+#include "AnimNextControllerBase.h"
+#include "AnimNextScopedCompileJob.h"
+#include "RigVMPythonUtils.h"
+#include "ExternalPackageHelper.h"
+#include "IAnimNextRigVMGraphInterface.h"
+#include "IWorkspaceEditor.h"
+#include "IWorkspaceEditorModule.h"
+#include "ObjectTools.h"
+#include "UncookedOnlyUtils.h"
+#include "Animation/Skeleton.h"
+#include "AnimNextRigVMFunctionData.h"
+#include "Variables/StructDataCache.h"
+#include "DataInterface/AnimNextNativeDataInterface.h"
+#include "Variables/AnimNextSharedVariables_EditorData.h"
+#include "Entries/AnimNextSharedVariablesEntry.h"
+#include "Entries/AnimNextEventGraphEntry.h"
+#include "Entries/AnimNextVariableEntry.h"
+#include "Graph/RigUnit_AnimNextBeginExecution.h"
+#include "Misc/TransactionObjectEvent.h"
+#include "Module/AnimNextEventGraphSchema.h"
+#include "RigVMModel/RigVMFunctionLibrary.h"
+#include "RigVMModel/RigVMNotifications.h"
+#include "RigVMModel/Nodes/RigVMAggregateNode.h"
+#include "RigVMModel/Nodes/RigVMCollapseNode.h"
+#include "UObject/AssetRegistryTagsContext.h"
+#include "RigVMFunctions/Execution/RigVMFunction_UserDefinedEvent.h"
+#include "Misc/UObjectToken.h"
+#include "UObject/SavePackage.h"
+#include "UObject/FortniteMainBranchObjectVersion.h"
+
+#if WITH_EDITOR
+#include "FileHelpers.h"
+#include "PackageSourceControlHelper.h"
+#endif
+
+#define LOCTEXT_NAMESPACE "AnimNextRigVMAssetEditorData"
+
+void UAnimNextRigVMAssetEditorData::BroadcastModified(EAnimNextEditorDataNotifType InType, UObject* InSubject)
+{
+	using namespace UE::UAF::UncookedOnly;
+
+	// Modifications here can trigger compilation, so add a scope to catch compile batches
+	UAnimNextRigVMAsset* Asset = FUtils::GetAsset<UAnimNextRigVMAsset>(this);
+	FScopedCompileJob CompilerResults(LOCTEXT("ModifiedAssetsJobName", "Modified Assets"), { Asset });
+
+	if(!bSuspendEditorDataNotifications)
+	{
+		ModifiedDelegate.Broadcast(this, InType, InSubject);
+	}
+
+	RequestAutoVMRecompilation();
+}
+
+void UAnimNextRigVMAssetEditorData::ReportError(const TCHAR* InMessage)
+{
+	FScriptExceptionHandler::Get().HandleException(ELogVerbosity::Error, InMessage, TEXT(""));
+}
+
+void UAnimNextRigVMAssetEditorData::ReconstructAllNodes()
+{
+	// Avoid refreshing EdGraph nodes during cook
+	if (GIsCookerLoadingPackage)
+	{
+		return;
+	}
+	
+	if (GetRigVMClient()->GetDefaultModel() == nullptr)
+	{
+		return;
+	}
+
+	TArray<URigVMEdGraphNode*> AllNodes;
+	GetAllNodesOfClass(AllNodes);
+
+	for (URigVMEdGraphNode* Node : AllNodes)
+	{
+		Node->SetFlags(RF_Transient);
+	}
+
+	for(URigVMEdGraphNode* Node : AllNodes)
+	{
+		Node->ReconstructNode();
+	}
+
+	for (URigVMEdGraphNode* Node : AllNodes)
+	{
+		Node->ClearFlags(RF_Transient);
+	}
+}
+
+void UAnimNextRigVMAssetEditorData::Serialize(FArchive& Ar)
+{
+	Ar.UsingCustomVersion(FFortniteMainBranchObjectVersion::GUID);
+
+	RigVMClient.SetDefaultSchemaClass(UAnimNextRigVMAssetSchema::StaticClass());
+	RigVMClient.SetOuterClientHost(this, GET_MEMBER_NAME_CHECKED(UAnimNextRigVMAssetEditorData, RigVMClient));
+
+	Super::Serialize(Ar);
+
+	if (Ar.IsLoading() && Ar.CustomVer(FFortniteMainBranchObjectVersion::GUID) < FFortniteMainBranchObjectVersion::AnimNextVariableReferences)
+	{
+		PRAGMA_DISABLE_DEPRECATION_WARNINGS
+		DefaultInjectionSiteReference = FAnimNextVariableReference(DefaultInjectionSite_DEPRECATED);
+		PRAGMA_ENABLE_DEPRECATION_WARNINGS
+
+		PRAGMA_DISABLE_DEPRECATION_WARNINGS
+		if (NativeInterface_DEPRECATED)
+		{
+			NativeInterfaces_DEPRECATED.Add(NativeInterface_DEPRECATED);
+			NativeInterface_DEPRECATED = nullptr;
+		}
+		PRAGMA_ENABLE_DEPRECATION_WARNINGS
+
+		bUpgradeDataInterfacesOnLoad = true;
+	}
+}
+
+void UAnimNextRigVMAssetEditorData::Initialize(bool bRecompileVM)
+{
+	RigVMClient.bDefaultModelCanBeRemoved = true;
+	RigVMClient.SetDefaultSchemaClass(UAnimNextRigVMAssetSchema::StaticClass());
+	RigVMClient.SetControllerClass(GetControllerClass());
+	RigVMClient.SetOuterClientHost(this, GET_MEMBER_NAME_CHECKED(UAnimNextRigVMAssetEditorData, RigVMClient));
+	RigVMClient.SetExternalModelHost(this);
+
+	URigVMFunctionLibrary* RigVMFunctionLibrary = nullptr;
+	{
+		TGuardValue<bool> DisableClientNotifs(RigVMClient.bSuspendNotifications, true);
+		RigVMFunctionLibrary = RigVMClient.GetOrCreateFunctionLibrary(false);
+	}
+
+	ensure(RigVMFunctionLibrary->GetFunctionHostObjectPathDelegate.IsBound());
+
+	if (RigVMClient.GetController(0) == nullptr)
+	{
+		if(RigVMClient.GetDefaultModel())
+		{
+			RigVMClient.GetOrCreateController(RigVMClient.GetDefaultModel());
+		}
+
+		check(RigVMFunctionLibrary);
+		RigVMClient.GetOrCreateController(RigVMFunctionLibrary);
+
+		if (!FunctionLibraryEdGraph)
+		{
+			FunctionLibraryEdGraph = NewObject<UAnimNextEdGraph>(CastChecked<UObject>(this), NAME_None, RF_Transactional);
+
+			FunctionLibraryEdGraph->Schema = UAnimNextEdGraphSchema::StaticClass();
+			FunctionLibraryEdGraph->bAllowRenaming = 0;
+			FunctionLibraryEdGraph->bEditable = 0;
+			FunctionLibraryEdGraph->bAllowDeletion = 0;
+			FunctionLibraryEdGraph->bIsFunctionDefinition = false;
+			FunctionLibraryEdGraph->ModelNodePath = RigVMClient.GetFunctionLibrary()->GetNodePath();
+			FunctionLibraryEdGraph->Initialize(this);
+		}
+
+		// Init function library controllers
+		for(URigVMLibraryNode* LibraryNode : RigVMClient.GetFunctionLibrary()->GetFunctions())
+		{
+			RigVMClient.GetOrCreateController(LibraryNode->GetContainedGraph());
+		}
+		
+		if(bRecompileVM)
+		{
+			RecompileVM();
+		}
+	}
+
+	for(UAnimNextRigVMAssetEntry* Entry : Entries)
+	{
+		Entry->Initialize(this);
+	}
+
+	InitializeAssetUserData();
+}
+
+void UAnimNextRigVMAssetEditorData::InitializeAssetUserData()
+{
+	if (IInterface_AssetUserData* OuterUserData = Cast<IInterface_AssetUserData>(GetOuter()))
+	{
+		if(!OuterUserData->HasAssetUserDataOfClass(GetAssetUserDataClass()))
+		{
+			OuterUserData->AddAssetUserDataOfClass(GetAssetUserDataClass());
+		}
+	}
+}
+
+void UAnimNextRigVMAssetEditorData::PostLoad()
+{
+	Super::PostLoad();
+
+	GraphModels.Reset();
+	
+	PostLoadExternalPackages();
+
+	// Preload our entries, as we need them for RefreshExternalModels
+	for (UAnimNextRigVMAssetEntry* Entry : Entries)
+	{
+		Entry->ConditionalPreload();
+		Entry->ConditionalPostLoad();
+	}
+
+	RefreshExternalModels();
+
+	Initialize(/*bRecompileVM*/false);
+	
+	GetRigVMClient()->RefreshAllModels(ERigVMLoadType::PostLoad, false, bIsCompiling);
+
+	GetRigVMClient()->PatchFunctionReferencesOnLoad();
+	TMap<URigVMLibraryNode*, FRigVMGraphFunctionHeader> OldHeaders;
+	TArray<FName> BackwardsCompatiblePublicFunctions;
+	GetRigVMClient()->PatchFunctionsOnLoad(this, BackwardsCompatiblePublicFunctions, OldHeaders);
+
+	// Register function references at RigVMBuildData
+	if (URigVMBuildData* BuildData = URigVMBuildData::Get())
+	{
+		TArray<FRigVMReferenceNodeData> ReferenceNodeDatas;
+		const TArray<URigVMGraph*> AllModels = GetAllModels();
+		for (URigVMGraph* ModelToVisit : AllModels)
+		{
+			for (URigVMNode* Node : ModelToVisit->GetNodes())
+			{
+				if (URigVMFunctionReferenceNode* ReferenceNode = Cast<URigVMFunctionReferenceNode>(Node))
+				{
+					ReferenceNodeDatas.Add(FRigVMReferenceNodeData(ReferenceNode));
+				}
+			}
+		}
+
+		// update the build data from the current function references
+		for (const FRigVMReferenceNodeData& ReferenceNodeData : ReferenceNodeDatas)
+		{
+			BuildData->RegisterFunctionReference(ReferenceNodeData);
+		}
+
+		BuildData->ClearInvalidReferences();
+	}
+
+	// Mark this as being dirty so that we recompile when needed
+	bVMRecompilationRequired = true;
+
+	// Queue compilation once the package has been fully loaded
+	// This is necessary in case we have external packages that haven't post-loaded yet
+	// However, if we are duplicating the asset OnEndLoadPackage won't be called
+	FCoreUObjectDelegates::OnEndLoadPackage.AddUObject(this, &UAnimNextRigVMAssetEditorData::HandlePackageDone);
+}
+
+void UAnimNextRigVMAssetEditorData::PostLoadExternalPackages()
+{
+	if(bUsesExternalPackages)
+	{
+		FExternalPackageHelper::LoadObjectsFromExternalPackages<UAnimNextRigVMAssetEntry>(this, [this](UAnimNextRigVMAssetEntry* InLoadedEntry)
+		{
+			check(IsValid(InLoadedEntry));
+			InLoadedEntry->Initialize(this);
+			Entries.Add(InLoadedEntry);
+		});
+	}
+
+	// Internal entries should be empty if we are externally packaged
+	ensure(!bUsesExternalPackages || InternalEntries.IsEmpty());
+
+	// Copy any internal entries to the main entries array
+	Entries.Append(InternalEntries);
+}
+
+void UAnimNextRigVMAssetEditorData::PostEditChangeProperty(FPropertyChangedEvent& PropertyChangedEvent)
+{
+	Super::PostEditChangeProperty(PropertyChangedEvent);
+
+	BroadcastModified(EAnimNextEditorDataNotifType::PropertyChanged, this);
+}
+
+void UAnimNextRigVMAssetEditorData::PostTransacted(const FTransactionObjectEvent& TransactionEvent)
+{
+	Super::PostTransacted(TransactionEvent);
+
+	if (TransactionEvent.GetEventType() == ETransactionObjectEventType::UndoRedo)
+	{
+		BroadcastModified(EAnimNextEditorDataNotifType::UndoRedo, this);
+	}
+}
+
+void UAnimNextRigVMAssetEditorData::GetAssetRegistryTags(FAssetRegistryTagsContext Context) const
+{
+	Super::GetAssetRegistryTags(Context);
+
+	{
+		// We may not have compiled yet, so cache exports if we havent already
+		if (!CachedExports.IsSet())
+		{
+			CachedExports = FAnimNextAssetRegistryExports();
+			FAnimNextAssetRegistryExports& OutExports = CachedExports.GetValue();
+			GetAnimNextAssetRegistryTags(Context, OutExports);
+		}
+
+		FString TagValue;
+		FAnimNextAssetRegistryExports::StaticStruct()->ExportText(TagValue, &CachedExports.GetValue(), nullptr, nullptr, PPF_None, nullptr);
+		Context.AddTag(FAssetRegistryTag(UE::UAF::ExportsAnimNextAssetRegistryTag, TagValue, FAssetRegistryTag::TT_Hidden));
+	}
+
+	{
+		FRigVMGraphFunctionHeaderArray FunctionExports;
+		UE::UAF::UncookedOnly::FUtils::GetAssetFunctions(this, FunctionExports);
+
+		FString TagValue;
+		const FArrayProperty* HeadersProperty = CastField<FArrayProperty>(FRigVMGraphFunctionHeaderArray::StaticStruct()->FindPropertyByName(TEXT("Headers")));
+		HeadersProperty->ExportText_Direct(TagValue, &(FunctionExports.Headers), &(FunctionExports.Headers), nullptr, PPF_None, nullptr);
+		Context.AddTag(FAssetRegistryTag(UE::UAF::AnimNextPublicGraphFunctionsExportsRegistryTag, TagValue, FAssetRegistryTag::TT_Hidden));
+	}
+
+	{
+		// Export user defined events as notifies
+		FString NotifyList = USkeleton::AnimNotifyTagDelimiter;
+		for(FName EventName : RigVMClient.GetEntryNames(FRigVMFunction_UserDefinedEvent::StaticStruct()))
+		{
+			NotifyList += FString::Printf(TEXT("%s%s"), *EventName.ToString(), *USkeleton::AnimNotifyTagDelimiter);
+		}
+		Context.AddTag(FAssetRegistryTag(USkeleton::AnimNotifyTag, NotifyList, FAssetRegistryTag::TT_Hidden));
+	}
+}
+
+bool UAnimNextRigVMAssetEditorData::Rename(const TCHAR* NewName, UObject* NewOuter, ERenameFlags Flags)
+{
+	FExternalPackageHelper::FRenameExternalObjectsHelperContext Context(this, Flags);
+	return Super::Rename(NewName, NewOuter, Flags);
+}
+
+void UAnimNextRigVMAssetEditorData::PreDuplicate(FObjectDuplicationParameters& DupParams)
+{
+	UObject::PreDuplicate(DupParams);
+	FExternalPackageHelper::DuplicateExternalPackages(this, DupParams);
+}
+
+void UAnimNextRigVMAssetEditorData::HandlePackageDone(const FEndLoadPackageContext& Context)
+{
+	if (!Context.LoadedPackages.Contains(GetPackage()))
+	{
+		return;
+	}
+	HandlePackageDone();
+}
+
+void UAnimNextRigVMAssetEditorData::HandlePackageDone()
+{
+	FCoreUObjectDelegates::OnEndLoadPackage.RemoveAll(this);
+
+	ReconstructAllNodes(); // If this is not executed on a node for whatever reason, it will appear transparent in the editor
+
+	TGuardValue<bool> DisableCompilationNotifications(bSuspendCompilationNotifications, true);
+
+	if (bUpgradeDataInterfacesOnLoad)
+	{
+		UpgradeDataInterfaces();
+	}
+
+	RecompileVM();
+}
+
+void UAnimNextRigVMAssetEditorData::GetAnimNextAssetRegistryTags(FAssetRegistryTagsContext& Context, FAnimNextAssetRegistryExports& OutExports) const
+{	
+	UE::UAF::UncookedOnly::FUtils::GetAssetVariableExports(this, OutExports, Context);
+}
+
+void UAnimNextRigVMAssetEditorData::RefreshAllModels(ERigVMLoadType InLoadType)
+{
+}
+
+void UAnimNextRigVMAssetEditorData::OnRigVMRegistryChanged()
+{
+	GetRigVMClient()->RefreshAllModels(ERigVMLoadType::PostLoad, false, bIsCompiling);
+	//RebuildGraphFromModel(); // TODO zzz : Move from blueprint to client
+}
+
+void UAnimNextRigVMAssetEditorData::RequestRigVMInit()
+{
+	// TODO zzz : How we do this on AnimNext ?
+}
+
+URigVMGraph* UAnimNextRigVMAssetEditorData::GetModel(const UEdGraph* InEdGraph) const
+{
+	return RigVMClient.GetModel(InEdGraph);
+}
+
+URigVMGraph* UAnimNextRigVMAssetEditorData::GetModel(const FString& InNodePath) const
+{
+	return RigVMClient.GetModel(InNodePath);
+}
+
+URigVMGraph* UAnimNextRigVMAssetEditorData::GetDefaultModel() const 
+{
+	return RigVMClient.GetDefaultModel();
+}
+
+TArray<URigVMGraph*> UAnimNextRigVMAssetEditorData::GetAllModels() const
+{
+	return RigVMClient.GetAllModels(true, true);
+}
+
+URigVMFunctionLibrary* UAnimNextRigVMAssetEditorData::GetLocalFunctionLibrary() const
+{
+	return RigVMClient.GetFunctionLibrary();
+}
+
+URigVMFunctionLibrary* UAnimNextRigVMAssetEditorData::GetOrCreateLocalFunctionLibrary(bool bSetupUndoRedo)
+{
+	return RigVMClient.GetOrCreateFunctionLibrary(bSetupUndoRedo);
+}
+
+URigVMGraph* UAnimNextRigVMAssetEditorData::AddModel(FString InName, bool bSetupUndoRedo, bool bPrintPythonCommand)
+{
+	TGuardValue<bool> EnablePythonPrint(bSuspendPythonMessagesForRigVMClient, !bPrintPythonCommand);
+	return RigVMClient.AddModel(InName, bSetupUndoRedo, bPrintPythonCommand);
+}
+
+bool UAnimNextRigVMAssetEditorData::RemoveModel(FString InName, bool bSetupUndoRedo, bool bPrintPythonCommand)
+{
+	TGuardValue<bool> EnablePythonPrint(bSuspendPythonMessagesForRigVMClient, !bPrintPythonCommand);
+	return RigVMClient.RemoveModel(InName, bSetupUndoRedo, bPrintPythonCommand);
+}
+
+FRigVMGetFocusedGraph& UAnimNextRigVMAssetEditorData::OnGetFocusedGraph()
+{
+	return RigVMClient.OnGetFocusedGraph();
+}
+
+const FRigVMGetFocusedGraph& UAnimNextRigVMAssetEditorData::OnGetFocusedGraph() const
+{
+	return RigVMClient.OnGetFocusedGraph();
+}
+
+URigVMGraph* UAnimNextRigVMAssetEditorData::GetFocusedModel() const
+{
+	return RigVMClient.GetFocusedModel();
+}
+
+URigVMController* UAnimNextRigVMAssetEditorData::GetController(const URigVMGraph* InGraph) const
+{
+	return RigVMClient.GetController(InGraph);
+};
+
+URigVMController* UAnimNextRigVMAssetEditorData::GetControllerByName(const FString InGraphName) const
+{
+	return RigVMClient.GetControllerByName(InGraphName);
+};
+
+URigVMController* UAnimNextRigVMAssetEditorData::GetOrCreateController(URigVMGraph* InGraph)
+{
+	return RigVMClient.GetOrCreateController(InGraph);
+};
+
+URigVMController* UAnimNextRigVMAssetEditorData::GetController(const UEdGraph* InEdGraph) const
+{
+	return RigVMClient.GetController(InEdGraph);
+};
+
+URigVMController* UAnimNextRigVMAssetEditorData::GetOrCreateController(const UEdGraph* InEdGraph)
+{
+	return RigVMClient.GetOrCreateController(InEdGraph);
+};
+
+TArray<FString> UAnimNextRigVMAssetEditorData::GeneratePythonCommands(const FString InNewBlueprintName)
+{
+	return TArray<FString>();
+}
+
+void UAnimNextRigVMAssetEditorData::SetupPinRedirectorsForBackwardsCompatibility()
+{
+}
+
+FRigVMGraphModifiedEvent& UAnimNextRigVMAssetEditorData::OnModified()
+{
+	return RigVMGraphModifiedEvent;
+}
+
+bool UAnimNextRigVMAssetEditorData::IsFunctionPublic(const FName& InFunctionName) const
+{
+	return GetLocalFunctionLibrary()->IsFunctionPublic(InFunctionName);
+}
+
+void UAnimNextRigVMAssetEditorData::MarkFunctionPublic(const FName& InFunctionName, bool bIsPublic)
+{
+	if (IsFunctionPublic(InFunctionName) == bIsPublic)
+	{
+		return;
+	}
+
+	URigVMController* Controller = RigVMClient.GetOrCreateController(GetLocalFunctionLibrary());
+	Controller->MarkFunctionAsPublic(InFunctionName, bIsPublic);
+}
+
+void UAnimNextRigVMAssetEditorData::RenameGraph(const FString& InNodePath, const FName& InNewName)
+{
+	if (URigVMGraph* ModelForNodePath = GetModel(InNodePath))
+	{
+		if (UEdGraph* EdGraph = Cast<UEdGraph>(GetEditorObjectForRigVMGraph(ModelForNodePath)))
+		{
+			RigVMClient.RenameModel(InNodePath, InNewName, true);
+		}
+	}
+}
+
+UClass* UAnimNextRigVMAssetEditorData::GetRigVMSchemaClass() const
+{
+	return UAnimNextRigVMAssetSchema::StaticClass();
+}
+
+UScriptStruct* UAnimNextRigVMAssetEditorData::GetRigVMExecuteContextStruct() const 
+{
+	return FAnimNextExecuteContext::StaticStruct();
+}
+
+UClass* UAnimNextRigVMAssetEditorData::GetRigVMEdGraphClass() const 
+{
+	return UAnimNextEdGraph::StaticClass();
+}
+
+UClass* UAnimNextRigVMAssetEditorData::GetRigVMEdGraphNodeClass() const
+{
+	return UAnimNextEdGraphNode::StaticClass();
+}
+
+UClass* UAnimNextRigVMAssetEditorData::GetRigVMEdGraphSchemaClass() const
+{
+	return UAnimNextEdGraphSchema::StaticClass();
+}
+
+UClass* UAnimNextRigVMAssetEditorData::GetRigVMEditorSettingsClass() const
+{
+	return URigVMEditorSettings::StaticClass();
+}
+
+FRigVMClient* UAnimNextRigVMAssetEditorData::GetRigVMClient()
+{
+	return &RigVMClient;
+}
+
+const FRigVMClient* UAnimNextRigVMAssetEditorData::GetRigVMClient() const
+{
+	return &RigVMClient;
+}
+
+IRigVMGraphFunctionHost* UAnimNextRigVMAssetEditorData::GetRigVMGraphFunctionHost()
+{
+	return this;
+}
+
+const IRigVMGraphFunctionHost* UAnimNextRigVMAssetEditorData::GetRigVMGraphFunctionHost() const
+{
+	return this;
+}
+
+void UAnimNextRigVMAssetEditorData::HandleRigVMGraphAdded(const FRigVMClient* InClient, const FString& InNodePath)
+{
+	if(URigVMGraph* RigVMGraph = InClient->GetModel(InNodePath))
+	{
+		RigVMGraph->SetExecuteContextStruct(GetExecuteContextStruct());
+
+		if(!HasAnyFlags(RF_ClassDefaultObject | RF_NeedInitialization | RF_NeedLoad | RF_NeedPostLoad) &&
+			GetOuter() != GetTransientPackage())
+		{
+			CreateEdGraph(RigVMGraph, true);
+			RequestAutoVMRecompilation();
+		}
+		
+#if WITH_EDITOR
+		if(!bSuspendPythonMessagesForRigVMClient)
+		{
+			const FString AssetName = RigVMGraph->GetSchema()->GetSanitizedName(GetName(), true, false);
+			RigVMPythonUtils::Print(AssetName, FString::Printf(TEXT("asset.add_graph('%s')"), *RigVMGraph->GetName()));
+		}
+#endif
+	}
+}
+
+void UAnimNextRigVMAssetEditorData::HandleRigVMGraphRemoved(const FRigVMClient* InClient, const FString& InNodePath)
+{
+	if(URigVMGraph* RigVMGraph = InClient->GetModel(InNodePath))
+	{
+		if (UAnimNextRigVMAssetEntry* Entry = FindEntryForRigVMGraph(RigVMGraph))
+		{
+			if (IAnimNextRigVMGraphInterface* GraphInterface = Cast<IAnimNextRigVMGraphInterface>(Entry))
+			{
+				GraphInterface->SetRigVMGraph(nullptr);
+			}
+		}
+		GraphModels.Remove(RigVMGraph);
+
+		RemoveEdGraph(RigVMGraph);
+		RequestAutoVMRecompilation();
+
+#if WITH_EDITOR
+		if(!bSuspendPythonMessagesForRigVMClient)
+		{
+			const FString AssetName = RigVMGraph->GetSchema()->GetSanitizedName(GetName(), true, false);
+			RigVMPythonUtils::Print(AssetName, FString::Printf(TEXT("asset.add_graph('%s')"), *RigVMGraph->GetName()));
+		}
+#endif
+	}
+}
+
+void UAnimNextRigVMAssetEditorData::HandleRigVMGraphRenamed(const FRigVMClient* InClient, const FString& InOldNodePath, const FString& InNewNodePath)
+{
+	if (InClient->GetModel(InNewNodePath))
+	{
+		TArray<UEdGraph*> EdGraphs = GetAllEdGraphs();
+		for (UEdGraph* EdGraph : EdGraphs)
+		{
+			if (URigVMEdGraph* RigGraph = Cast<URigVMEdGraph>(EdGraph))
+			{
+				RigGraph->HandleRigVMGraphRenamed(InOldNodePath, InNewNodePath);
+			}
+		}
+	}
+}
+
+
+void UAnimNextRigVMAssetEditorData::HandleConfigureRigVMController(const FRigVMClient* InClient, URigVMController* InControllerToConfigure)
+{
+	InControllerToConfigure->OnModified().AddUObject(this, &UAnimNextRigVMAssetEditorData::HandleModifiedEvent);
+
+	TWeakObjectPtr<UAnimNextRigVMAssetEditorData> WeakThis(this);
+
+	InControllerToConfigure->GetExternalVariablesDelegate.BindLambda([](URigVMGraph* InGraph) -> TArray<FRigVMExternalVariable> {
+		if (InGraph)
+		{
+			if(URigVMHost* RigVMHost = InGraph->GetTypedOuter<URigVMHost>())
+			{
+				return RigVMHost->GetExternalVariables();
+			}
+		}
+		return TArray<FRigVMExternalVariable>();
+	});
+	
+	// this delegate is used by the controller
+	// to retrieve the current bytecode of the VM
+	InControllerToConfigure->GetCurrentByteCodeDelegate.BindLambda([WeakThis]() -> const FRigVMByteCode*
+	{
+		if (WeakThis.IsValid())
+		{
+			if(UAnimNextRigVMAsset* Asset = WeakThis->GetTypedOuter<UAnimNextRigVMAsset>())
+			{
+				if (Asset->VM)
+				{
+					return &Asset->VM->GetByteCode();
+				}
+			}
+		}
+		return nullptr;
+
+	});
+
+#if WITH_EDITOR
+	InControllerToConfigure->SetupDefaultUnitNodeDelegates(TDelegate<FName(FRigVMExternalVariable, FString)>::CreateLambda(
+		[](FRigVMExternalVariable InVariableToCreate, FString InDefaultValue) -> FName
+		{
+			return NAME_None;
+		}
+	));
+#endif
+}
+
+UObject* UAnimNextRigVMAssetEditorData::GetEditorObjectForRigVMGraph(const URigVMGraph* InVMGraph) const
+{
+	if(InVMGraph)
+	{
+		if (InVMGraph->IsA<URigVMFunctionLibrary>())
+		{
+			return Cast<UObject>(FunctionLibraryEdGraph.Get());
+		}
+
+		const auto FindSubgraph = ([](const FString SearchGraphNodePath, URigVMEdGraph* EdGraph) -> URigVMEdGraph*
+		{
+			TArray<UEdGraph*> SubGraphs;
+			EdGraph->GetAllChildrenGraphs(SubGraphs);
+			for (UEdGraph* SubGraph : SubGraphs)
+			{
+				if (URigVMEdGraph* RigVMEdGraph = Cast<URigVMEdGraph>(SubGraph))
+				{
+					if (RigVMEdGraph->GetRigVMNodePath() == SearchGraphNodePath)
+					{
+						return RigVMEdGraph;
+					}
+				}
+			}
+			return nullptr;
+		});
+
+		const FString GraphNodePath = InVMGraph->GetNodePath();
+		for(UAnimNextRigVMAssetEntry* Entry : Entries)
+		{
+			if(IAnimNextRigVMGraphInterface* GraphInterface = Cast<IAnimNextRigVMGraphInterface>(Entry))
+			{
+				URigVMEdGraph* EdGraph = GraphInterface->GetEdGraph();
+
+				if (const URigVMGraph* RigVMGraph = GraphInterface->GetRigVMGraph())
+				{
+					if (RigVMGraph == InVMGraph)
+					{
+						return EdGraph;
+					}
+				}
+
+				if (URigVMEdGraph* RigVMEdGraph = FindSubgraph(GraphNodePath, EdGraph))
+				{
+					return RigVMEdGraph;
+				}
+			}
+		}
+
+		for (const TObjectPtr<URigVMEdGraph>& FunctionEdGraph : FunctionEdGraphs)
+		{
+			if (FunctionEdGraph->ModelNodePath == GraphNodePath)
+			{
+				return FunctionEdGraph;
+			}
+
+			if (URigVMEdGraph* RigVMEdGraph = FindSubgraph(GraphNodePath, FunctionEdGraph))
+			{
+				return RigVMEdGraph;
+			}
+		}
+	}
+	return nullptr;
+}
+
+URigVMGraph* UAnimNextRigVMAssetEditorData::GetRigVMGraphForEditorObject(UObject* InObject) const
+{
+	if(const URigVMEdGraph* Graph = Cast<URigVMEdGraph>(InObject))
+	{
+		if (Graph->bIsFunctionDefinition)
+		{
+			if (URigVMLibraryNode* LibraryNode = RigVMClient.GetFunctionLibrary()->FindFunction(*Graph->ModelNodePath))
+			{
+				return LibraryNode->GetContainedGraph();
+			}
+		}
+		else
+		{
+			return RigVMClient.GetModel(Graph->ModelNodePath);
+		}
+	}
+
+	return nullptr;
+}
+
+FRigVMGraphFunctionStore* UAnimNextRigVMAssetEditorData::GetRigVMGraphFunctionStore()
+{
+	return &GraphFunctionStore;
+}
+
+const FRigVMGraphFunctionStore* UAnimNextRigVMAssetEditorData::GetRigVMGraphFunctionStore() const
+{
+	return &GraphFunctionStore;
+}
+
+TObjectPtr<URigVMGraph> UAnimNextRigVMAssetEditorData::CreateContainedGraphModel(URigVMCollapseNode* CollapseNode, const FName& Name)
+{
+	check(CollapseNode);
+
+	TObjectPtr<URigVMGraph> Model = NewObject<URigVMGraph>(CollapseNode, Name);
+
+	check(CollapseNode->GetGraph());
+	if (CollapseNode->GetGraph()->GetSchema() != nullptr)
+	{
+		Model->SetSchemaClass(CollapseNode->GetGraph()->GetSchema()->GetClass());
+	}
+	else
+	{
+		Model->SetSchemaClass(RigVMClient.GetDefaultSchemaClass());
+	}
+
+	URigVMGraph* CollapseNodeModelRootGraph = CollapseNode->GetRootGraph();
+	check(CollapseNodeModelRootGraph);
+
+	// If we are a transient asset, or not using external packages dont use external packages
+	if (bUsesExternalPackages && !CollapseNodeModelRootGraph->HasAnyFlags(RF_Transient))
+	{
+		Model->SetExternalPackage(CollapseNodeModelRootGraph->GetExternalPackage());
+	}
+
+	return Model;
+}
+
+void UAnimNextRigVMAssetEditorData::BuildFunctionWrapperEventVariables(FAnimNextRigVMAssetCompileContext& InContext)
+{
+	using namespace UE::UAF::UncookedOnly;
+	
+	for (const FAnimNextProgrammaticFunctionHeader& ProgrammaticFunctionHeader : InContext.FunctionHeaders)
+	{
+		const FRigVMGraphFunctionHeader& FunctionHeader = ProgrammaticFunctionHeader.FunctionHeader;
+		for (const FRigVMGraphFunctionArgument& Argument : FunctionHeader.Arguments)
+		{
+			// Don't create internal variables for execution IO, visible only, or hidden pins. Those types do not need variable nodes in the graph.
+			if (Argument.Direction == ERigVMPinDirection::Input || Argument.Direction == ERigVMPinDirection::Output)
+			{
+				FRigVMGraphFunctionArgument InternallyNamedArgument = Argument;
+				InternallyNamedArgument.Name = FName(FUtils::MakeFunctionWrapperVariableName(FunctionHeader.Name, Argument.Name));
+				InContext.ProgrammaticVariables.Add(FAnimNextProgrammaticVariable::FromRigVMGraphFunctionArgument(InternallyNamedArgument));
+			}
+		}
+	}
+}
+
+void UAnimNextRigVMAssetEditorData::BuildFunctionWrapperEvents(FAnimNextRigVMAssetCompileContext& InContext, const FRigVMCompileSettings& InSettings)
+{
+	using namespace UE::UAF::UncookedOnly;
+
+	UAnimNextRigVMAsset* Asset = FUtils::GetAsset<UAnimNextRigVMAsset>(this);
+	Asset->FunctionData.Reset();
+
+	FRigVMClient* VMClient = GetRigVMClient();
+
+	// Create all shim events for our traits to call
+	const bool bSetupUndoRedo = false;
+	URigVMGraph* WrapperGraph = NewObject<URigVMGraph>(this, NAME_None, RF_Transient);
+	URigVMController* Controller = VMClient->GetOrCreateController(WrapperGraph);
+	FRigVMControllerNotifGuard NotifGuard(Controller);
+	bool bAddedWrapperEvent = true;
+	TArray<FRigVMExternalVariable> ExternalVariables = Asset->GetExternalVariables();
+
+	for(const FAnimNextProgrammaticFunctionHeader& AnimNextFunctionHeader : InContext.FunctionHeaders)
+	{
+		const FRigVMGraphFunctionHeader& FunctionHeader = AnimNextFunctionHeader.FunctionHeader;
+
+		URigVMLibraryNode* LibraryNode = Cast<URigVMLibraryNode>(FunctionHeader.LibraryPointer.GetNodeSoftPath().TryLoad());
+		if(LibraryNode == nullptr)
+		{
+			InSettings.ReportError(FString::Printf(TEXT("Could not find function '%s'"), *FunctionHeader.Name.ToString()));
+			continue;
+		}
+
+		// Create user-defined entry point
+		FString WrapperEventName = FUtils::MakeFunctionWrapperEventName(FunctionHeader.Name);
+		URigVMUnitNode* EventNode = Controller->AddUnitNode(FRigVMFunction_UserDefinedEvent::StaticStruct(), TEXT("Execute"), FVector2D::ZeroVector, FunctionHeader.Name.ToString(), bSetupUndoRedo);
+		if(EventNode == nullptr)
+		{
+			InSettings.ReportError(FString::Printf(TEXT("Could not spawn event node for function '%s'"), *FunctionHeader.Name.ToString()));
+			continue;
+		}
+		URigVMPin* EventNamePin = EventNode->FindPin(GET_MEMBER_NAME_STRING_CHECKED(FRigVMFunction_UserDefinedEvent, EventName));
+		if(EventNamePin == nullptr)
+		{
+			InSettings.ReportError(FString::Printf(TEXT("Could not find custom event name pin")));
+			continue;
+		}
+		Controller->SetPinDefaultValue(EventNamePin->GetPinPath(), WrapperEventName, true, bSetupUndoRedo);
+
+		// Call function
+		URigVMFunctionReferenceNode* FunctionNode = Controller->AddFunctionReferenceNode(LibraryNode, FVector2D::ZeroVector, FunctionHeader.Name.ToString(), bSetupUndoRedo);
+		if(FunctionNode == nullptr)
+		{
+			InSettings.ReportError(FString::Printf(TEXT("Could not spawn function node for function '%s'"), *FunctionHeader.Name.ToString()));
+			continue;
+		}
+
+		// Link up Execute nodes if needed, function may be pure & lack an input pin
+		URigVMPin* CurrentExecuteOutputPin = EventNode->FindPin(FRigVMStruct::ExecuteContextName.ToString());
+		URigVMPin* ExecuteInputPin = FunctionNode->FindPin(FRigVMStruct::ExecuteContextName.ToString());
+		if (ExecuteInputPin && !Controller->AddLink(CurrentExecuteOutputPin, ExecuteInputPin, bSetupUndoRedo))
+		{
+			InSettings.ReportError(FString::Printf(TEXT("Could not link execute pins for function '%s'"), *FunctionHeader.Name.ToString()));
+			continue;
+		}
+
+		// Update current execute pin, RigVM doesn't have a concept of input / output execute pins, just one execute content pin used for both
+		CurrentExecuteOutputPin = ExecuteInputPin ? ExecuteInputPin : CurrentExecuteOutputPin;
+
+			// Controller needs to notify the AST of variable changes to make new links
+		bool bSuspendNotificationForInternalVariables = false;
+		FRigVMControllerNotifGuard VarNotifGuard(Controller, bSuspendNotificationForInternalVariables);
+
+		// Generate & link input arguments, also generate result variable node but link later
+		TArray<int32> ArgIndices;
+		for (const FRigVMGraphFunctionArgument& Argument : FunctionHeader.Arguments)
+		{
+			// Execution context is captured as arg pins, skip those for internal variable gen
+			if (Argument.Direction == ERigVMPinDirection::IO)
+			{
+				continue;
+			}
+
+			bool bIsGetter = Argument.Direction == ERigVMPinDirection::Input;
+			FName InternalArgName = FName(FUtils::MakeFunctionWrapperVariableName(FunctionHeader.Name, Argument.Name));
+			int32 VariableIndex = ExternalVariables.IndexOfByPredicate([&InternalArgName](const FRigVMExternalVariable& InVariable)
+			{
+				return InVariable.Name == InternalArgName;
+			});
+			check(VariableIndex != INDEX_NONE);
+			ArgIndices.Add(VariableIndex);
+
+			if (bIsGetter)
+			{
+				URigVMVariableNode* FunctionParamVariableNode = Controller->AddVariableNode(InternalArgName
+					, Argument.CPPType.ToString()
+					, Argument.CPPTypeObject.Get()
+					, bIsGetter
+					, Argument.DefaultValue
+					, FVector2D::ZeroVector
+					, InternalArgName.ToString()
+					, bSetupUndoRedo);
+
+				if (!FunctionParamVariableNode)
+				{
+					InSettings.ReportError(FString::Printf(TEXT("Failed to add internal variable node for param: %s, var: %s"), *FunctionHeader.Name.ToString(), *InternalArgName.ToString()));
+					return;
+				}
+
+				// Link Param Pins
+				URigVMPin* ParamValuePin = FunctionParamVariableNode->GetValuePin();
+				URigVMPin* FunctionArgumentPin = FunctionNode->FindPin(Argument.Name.ToString());
+				if (!Controller->AddLink(ParamValuePin, FunctionArgumentPin, bSetupUndoRedo))
+				{
+					InSettings.ReportError(FString::Printf(TEXT("Failed to link internal variable param node to function: %s -> %s"), *GetNameSafe(ParamValuePin), *GetNameSafe(FunctionArgumentPin)));
+					return;
+				}
+			}
+
+			if (!bIsGetter)
+			{
+				
+				URigVMVariableNode* FunctionResultVariableNode = Controller->AddVariableNode(InternalArgName
+					, Argument.CPPType.ToString()
+					, Argument.CPPTypeObject.Get()
+					, bIsGetter
+					, Argument.DefaultValue
+					, FVector2D::ZeroVector
+					, InternalArgName.ToString()
+					, bSetupUndoRedo);
+
+				if (!FunctionResultVariableNode)
+				{
+					InSettings.ReportError(FString::Printf(TEXT("Failed to add internal variable node for result: %s, var: %s"), *FunctionHeader.Name.ToString(), *InternalArgName.ToString()));
+					return;
+				}
+
+				// Link Result pins
+				URigVMPin* FunctionResultPin = FunctionNode->FindPin(Argument.Name.ToString());
+				URigVMPin* ResultValuePin = FunctionResultVariableNode->GetValuePin();
+				if (!Controller->AddLink(FunctionResultPin, ResultValuePin, bSetupUndoRedo))
+				{
+					InSettings.ReportError(FString::Printf(TEXT("Failed to link internal variable result node to function: %s -> %s"), *GetNameSafe(FunctionResultPin), *GetNameSafe(ResultValuePin)));
+					return;
+				}
+
+				// Link Result Execute pins
+				URigVMPin* ResultExecuteInputPin = FunctionResultVariableNode->FindPin(FRigVMStruct::ExecuteContextName.ToString());
+				if (!Controller->AddLink(CurrentExecuteOutputPin, ResultExecuteInputPin, bSetupUndoRedo))
+				{
+					InSettings.ReportError(FString::Printf(TEXT("Failed to link execute pins for variable result node: %s -> %s"), *GetNameSafe(CurrentExecuteOutputPin), *GetNameSafe(ResultExecuteInputPin)));
+					return;
+				}
+
+				// Update current execute pin, RigVM doesn't have a concept of input / output execute pins, just one execute content pin used for both
+				CurrentExecuteOutputPin = ResultExecuteInputPin;
+			}
+		}
+
+		bAddedWrapperEvent = true;
+
+		FAnimNextRigVMFunctionData FunctionData;
+		FunctionData.Name = FunctionHeader.Name;
+		FunctionData.EventName = FName(*WrapperEventName);
+		FunctionData.ArgIndices = MoveTemp(ArgIndices);
+		Asset->FunctionData.Add(MoveTemp(FunctionData));
+	}
+
+	if(bAddedWrapperEvent)
+	{
+		InContext.ProgrammaticGraphs.Add(WrapperGraph);
+	}
+}
+
+void UAnimNextRigVMAssetEditorData::RecompileVM()
+{
+	using namespace UE::UAF::UncookedOnly;
+
+	if (bIsCompiling)
+	{
+		return;
+	}
+
+	TGuardValue<bool> CompilingGuard(bIsCompiling, true);
+
+	UAnimNextRigVMAsset* Asset = FUtils::GetAsset<UAnimNextRigVMAsset>(this);
+	FScopedCompileJob CompilerResults(Asset);
+
+	VMCompileSettings.SetExecuteContextStruct(FAnimNextExecuteContext::StaticStruct());
+	FRigVMCompileSettings Settings = (bCompileInDebugMode) ? FRigVMCompileSettings::Fast(VMCompileSettings.GetExecuteContextStruct()) : VMCompileSettings;
+	Settings.SurpressInfoMessages = false;
+	Settings.bWarnAboutDuplicateEvents = true;
+	Settings.ASTSettings.ReportDelegate.BindUObject(this, &UAnimNextRigVMAssetEditorData::HandleReportFromCompiler);
+
+	Asset->VMRuntimeSettings = VMRuntimeSettings;
+
+	OnPreCompileAsset(Settings);
+
+	CachedExports.Reset();  // asset variables and other tags will be updated at the end by AssetRegistry->AssetUpdateTags
+
+	bWarningsDuringCompilation = false;
+	bErrorsDuringCompilation = false;
+
+	RigGraphDisplaySettings.MinMicroSeconds = RigGraphDisplaySettings.LastMinMicroSeconds = DBL_MAX;
+	RigGraphDisplaySettings.MaxMicroSeconds = RigGraphDisplaySettings.LastMaxMicroSeconds = (double)INDEX_NONE;
+
+	FAnimNextRigVMAssetCompileContext CompileContext = {};
+	{
+		TGuardValue<bool> ReentrantGuardSelf(bSuspendModelNotificationsForSelf, true);
+		TGuardValue<bool> ReentrantGuardOthers(RigVMClient.bSuspendModelNotificationsForOthers, true);
+
+		FUtils::RecreateVM(Asset);
+
+		{
+			FAnimNextGetFunctionHeaderCompileContext GetFunctionHeaderCompileContext(CompileContext);
+			OnPreCompileGetProgrammaticFunctionHeaders(Settings, GetFunctionHeaderCompileContext);
+		}
+
+		BuildFunctionWrapperEventVariables(CompileContext);
+
+		{
+			FAnimNextGetVariableCompileContext GetVariableCompileContext(CompileContext);
+			FUtils::CompileVariables(Settings, Asset, GetVariableCompileContext);
+			OnPostCompileVariables(Settings, GetVariableCompileContext);
+		}
+
+		BuildFunctionWrapperEvents(CompileContext, Settings);
+
+		{
+			FAnimNextGetGraphCompileContext GetGraphCompileContext(CompileContext);
+			OnPreCompileGetProgrammaticGraphs(Settings, GetGraphCompileContext);
+		}
+
+		for(URigVMGraph* ProgrammaticGraph : CompileContext.ProgrammaticGraphs)
+		{
+			check(ProgrammaticGraph != nullptr);
+		}
+
+		FRigVMClient* VMClient = GetRigVMClient();
+
+		CompileContext.AllGraphs = VMClient->GetAllModels(false, false);
+		CompileContext.AllGraphs.Append(CompileContext.ProgrammaticGraphs);
+
+		{
+			FAnimNextProcessGraphCompileContext ProcessGraphCompileContext(CompileContext);
+			OnPreCompileProcessGraphs(Settings, ProcessGraphCompileContext);
+		}
+
+		if(CompileContext.AllGraphs.Num() > 0)
+		{
+			URigVMController* Controller = VMClient->GetOrCreateController(CompileContext.AllGraphs[0]);
+
+			URigVMCompiler* Compiler = URigVMCompiler::StaticClass()->GetDefaultObject<URigVMCompiler>();
+			Compiler->Compile(Settings, CompileContext.AllGraphs, Controller, Asset->VM, Asset->ExtendedExecuteContext, Asset->GetExternalVariables(), &PinToOperandMap);
+		}
+
+		// Initialize right away, in packaged builds we initialize during PostLoad
+		Asset->VM->Initialize(Asset->ExtendedExecuteContext);
+		Asset->GenerateUserDefinedDependenciesData(Asset->ExtendedExecuteContext);
+
+		// Notable difference with vanilla RigVM host behavior - we init the VM here at the moment as we only have one 'instance'
+		Asset->InitializeVM(FRigUnit_AnimNextBeginExecution::EventName);
+
+		if (bErrorsDuringCompilation)
+		{
+			if(Settings.SurpressErrors)
+			{
+				Settings.Reportf(EMessageSeverity::Info, Asset, TEXT("Compilation Errors may be suppressed for AnimNext asset: %s. See VM Compile Settings for more Details"), *Asset->GetName());
+			}
+		}
+
+		bVMRecompilationRequired = false;
+
+		if(Asset->VM)
+		{
+			RigVMCompiledEvent.Broadcast(Asset, Asset->VM, Asset->ExtendedExecuteContext);
+		}
+
+#if WITH_EDITOR
+		// Display programmatic graphs
+		if(CVarDumpProgrammaticGraphs.GetValueOnGameThread())
+		{
+			FUtils::OpenProgrammaticGraphs(this, CompileContext.ProgrammaticGraphs);
+		}
+		else
+#endif
+		{
+			RemoveProgrammaticGraphs(CompileContext.ProgrammaticGraphs);
+		}
+
+		RemoveTransientGraphs(CompileContext.AllGraphs);
+
+		OnPostCompileCleanup(Settings);
+
+#if WITH_EDITOR
+		//	RefreshBreakpoints(EditorData);
+#endif
+
+		// Refresh CachedExports
+		if (IAssetRegistry* AssetRegistry = IAssetRegistry::Get())
+		{
+			AssetRegistry->AssetUpdateTags(Asset, EAssetRegistryTagsCaller::Fast);
+		}
+	}
+}
+
+void UAnimNextRigVMAssetEditorData::RemoveProgrammaticGraphs(TArrayView<URigVMGraph*> InGraphs)
+{
+	FRigVMClient* VMClient = GetRigVMClient();
+	
+	for(URigVMGraph* Graph : InGraphs)
+	{
+		VMClient->RemoveController(Graph);
+		Graph->Rename(nullptr, GetTransientPackage(), REN_ForceNoResetLoaders | REN_DoNotDirty | REN_DontCreateRedirectors | REN_NonTransactional);
+	}
+}
+
+void UAnimNextRigVMAssetEditorData::RemoveTransientGraphs(TArrayView<URigVMGraph*> InGraphs)
+{
+	FRigVMClient* VMClient = GetRigVMClient();
+	
+	for(URigVMGraph* Graph : InGraphs)
+	{
+		if(Graph->HasAnyFlags(RF_Transient))
+		{
+			VMClient->RemoveController(Graph);
+			Graph->Rename(nullptr, GetTransientPackage(), REN_ForceNoResetLoaders | REN_DoNotDirty | REN_DontCreateRedirectors | REN_NonTransactional);
+		}
+	}
+}
+
+void UAnimNextRigVMAssetEditorData::HandleRemoveNotify(UObject* InAsset, const FString& InFindString, bool bFindWholeWord, ESearchCase::Type InSearchCase)
+{
+	UAnimNextRigVMAsset* Asset = Cast<UAnimNextRigVMAsset>(InAsset);
+	if(Asset == nullptr)
+	{
+		return;
+	}
+
+	UAnimNextRigVMAssetEditorData* EditorData = UE::UAF::UncookedOnly::FUtils::GetEditorData<UAnimNextRigVMAssetEditorData>(Asset);
+	if(EditorData == nullptr)
+	{
+		return;
+	}
+
+	URigVMController* Controller = EditorData->GetController();
+	Controller->OpenUndoBracket(LOCTEXT("RemoveNotifyEvents", "Remove Notify Events").ToString());
+
+	for(TObjectPtr<URigVMGraph> Model : EditorData->RigVMClient.GetModels())
+	{
+		for(URigVMNode* Node : Model->GetNodes())
+		{
+			if(URigVMUnitNode* UnitNode = Cast<URigVMUnitNode>(Node))
+			{
+				if(UnitNode->GetScriptStruct()->IsChildOf(FRigVMFunction_UserDefinedEvent::StaticStruct()))
+				{
+					URigVMPin* Pin = UnitNode->FindPin(GET_MEMBER_NAME_STRING_CHECKED(FRigVMFunction_UserDefinedEvent, EventName));
+					FString EventNameString = Pin->GetDefaultValue();
+					if( (bFindWholeWord && EventNameString.Equals(InFindString, InSearchCase)) ||
+						(!bFindWholeWord && EventNameString.Contains(InFindString, InSearchCase)))
+					{
+						Controller->RemoveNode(Node, true, true);
+					}
+				}
+			}
+		}
+	}
+
+	Controller->CloseUndoBracket();
+}
+
+void UAnimNextRigVMAssetEditorData::HandleReplaceNotify(UObject* InAsset, const FString& InFindString, const FString& InReplaceString, bool bFindWholeWord, ESearchCase::Type InSearchCase)
+{
+	UAnimNextRigVMAsset* Asset = Cast<UAnimNextRigVMAsset>(InAsset);
+	if(Asset == nullptr)
+	{
+		return;
+	}
+
+	UAnimNextRigVMAssetEditorData* EditorData = UE::UAF::UncookedOnly::FUtils::GetEditorData<UAnimNextRigVMAssetEditorData>(Asset);
+	if(EditorData == nullptr)
+	{
+		return;
+	}
+
+	URigVMController* Controller = EditorData->GetController();
+	Controller->OpenUndoBracket(LOCTEXT("ReplaceNotifyEvents", "Replace Notify Events").ToString());
+
+	for(TObjectPtr<URigVMGraph> Model : EditorData->RigVMClient.GetModels())
+	{
+		for(URigVMNode* Node : Model->GetNodes())
+		{
+			if(URigVMUnitNode* UnitNode = Cast<URigVMUnitNode>(Node))
+			{
+				if(UnitNode->GetScriptStruct()->IsChildOf(FRigVMFunction_UserDefinedEvent::StaticStruct()))
+				{
+					URigVMPin* Pin = UnitNode->FindPin(GET_MEMBER_NAME_STRING_CHECKED(FRigVMFunction_UserDefinedEvent, EventName));
+					FString EventNameString = Pin->GetDefaultValue();
+					if( (bFindWholeWord && EventNameString.Equals(InFindString, InSearchCase)) ||
+						(!bFindWholeWord && EventNameString.Contains(InFindString, InSearchCase)))
+					{
+						const FString NewName = EventNameString.Replace(*InFindString, *InReplaceString, InSearchCase);
+						Controller->SetPinDefaultValue(Pin->GetPinPath(), NewName, true, true, false, true);
+					}
+				}
+			}
+		}
+	}
+
+	Controller->CloseUndoBracket();
+}
+
+bool UAnimNextRigVMAssetEditorData::IsDirtyForRecompilation() const
+{
+	if(bVMRecompilationRequired)
+	{
+		return true;
+	}
+
+	bool bDependencyDirty = false;
+	ForEachEntryOfType<UAnimNextSharedVariablesEntry>([&bDependencyDirty](UAnimNextSharedVariablesEntry* InEntry)
+	{
+		if(InEntry->Asset)
+		{
+			UAnimNextRigVMAssetEditorData* EditorData = UE::UAF::UncookedOnly::FUtils::GetEditorData<UAnimNextRigVMAssetEditorData>(InEntry->Asset.Get());
+			if(EditorData->IsDirtyForRecompilation())
+			{
+				bDependencyDirty = true;
+				return false;
+			}
+		}
+		return true;
+	});
+
+	return bDependencyDirty;
+}
+
+void UAnimNextRigVMAssetEditorData::RecompileVMIfRequired()
+{
+	if (bVMRecompilationRequired)
+	{
+		RecompileVM();
+	}
+}
+
+void UAnimNextRigVMAssetEditorData::RequestAutoVMRecompilation()
+{
+	bVMRecompilationRequired = true;
+	if (bAutoRecompileVM && VMRecompilationBracket == 0)
+	{
+		RecompileVMIfRequired();
+	}
+}
+
+void UAnimNextRigVMAssetEditorData::SetAutoVMRecompile(bool bAutoRecompile)
+{
+	bAutoRecompileVM = bAutoRecompile;
+}
+
+bool UAnimNextRigVMAssetEditorData::GetAutoVMRecompile() const
+{
+	return bAutoRecompileVM;
+}
+
+void UAnimNextRigVMAssetEditorData::IncrementVMRecompileBracket()
+{
+	VMRecompilationBracket++;
+}
+
+void UAnimNextRigVMAssetEditorData::DecrementVMRecompileBracket()
+{
+	if (VMRecompilationBracket == 1)
+	{
+		if (bAutoRecompileVM)
+		{
+			RecompileVMIfRequired();
+		}
+		VMRecompilationBracket = 0;
+
+		if (InteractionBracketFinished.IsBound())
+		{
+			InteractionBracketFinished.Broadcast(this);
+		}
+	}
+	else if (VMRecompilationBracket > 0)
+	{
+		VMRecompilationBracket--;
+	}
+}
+
+void UAnimNextRigVMAssetEditorData::HandleModifiedEvent(ERigVMGraphNotifType InNotifType, URigVMGraph* InGraph, UObject* InSubject)
+{
+	// Skip any notifications we get while compiling (they can come from programmatic graph generation)
+	if(bIsCompiling)
+	{
+		return;
+	}
+	
+	bool bNotifForOthersPending = true;
+
+	switch(InNotifType)
+	{
+	case ERigVMGraphNotifType::InteractionBracketOpened:
+		{
+			IncrementVMRecompileBracket();
+			break;
+		}
+	case ERigVMGraphNotifType::InteractionBracketClosed:
+	case ERigVMGraphNotifType::InteractionBracketCanceled:
+		{
+			DecrementVMRecompileBracket();
+			break;
+		}
+	case ERigVMGraphNotifType::NodeAdded:
+		{
+			if (URigVMCollapseNode* CollapseNode = Cast<URigVMCollapseNode>(InSubject))
+			{
+				CreateEdGraphForCollapseNode(CollapseNode, false);
+				break;
+			}
+			RequestAutoVMRecompilation();
+			break;
+	}
+	case ERigVMGraphNotifType::NodeRemoved:
+	{
+		if (URigVMCollapseNode* CollapseNode = Cast<URigVMCollapseNode>(InSubject))
+		{
+			RemoveEdGraphForCollapseNode(CollapseNode, false);
+			break;
+		}
+		RequestAutoVMRecompilation();
+		break;
+	}
+	case ERigVMGraphNotifType::NodeRenamed:
+	{
+		if (URigVMCollapseNode* CollapseNode = Cast<URigVMCollapseNode>(InSubject))
+		{
+			FString NewNodePath = CollapseNode->GetNodePath(true /* recursive */);
+			FString Left, Right = NewNodePath;
+			URigVMNode::SplitNodePathAtEnd(NewNodePath, Left, Right);
+			FString OldNodePath = CollapseNode->GetPreviousFName().ToString();
+			if (!Left.IsEmpty())
+			{
+				OldNodePath = URigVMNode::JoinNodePath(Left, OldNodePath);
+			}
+
+			HandleRigVMGraphRenamed(GetRigVMClient(), OldNodePath, NewNodePath);
+
+			if (UEdGraph* ContainedEdGraph = Cast<UEdGraph>(GetEditorObjectForRigVMGraph(CollapseNode->GetContainedGraph())))
+			{
+				ContainedEdGraph->Rename(*CollapseNode->GetEditorSubGraphName(), nullptr);
+			}
+		}
+		break;
+	}
+	case ERigVMGraphNotifType::LinkAdded:
+	case ERigVMGraphNotifType::LinkRemoved:
+	case ERigVMGraphNotifType::PinArraySizeChanged:
+	case ERigVMGraphNotifType::PinDirectionChanged:
+		{
+			RequestAutoVMRecompilation();
+			break;
+		}
+
+	case ERigVMGraphNotifType::PinDefaultValueChanged:
+		{
+			if (InGraph->GetRuntimeAST().IsValid())
+			{
+				URigVMPin* RootPin = CastChecked<URigVMPin>(InSubject)->GetRootPin();
+				FRigVMASTProxy RootPinProxy = FRigVMASTProxy::MakeFromUObject(RootPin);
+				const FRigVMExprAST* Expression = InGraph->GetRuntimeAST()->GetExprForSubject(RootPinProxy);
+				if (Expression == nullptr)
+				{
+					InGraph->ClearAST();
+				}
+				else if (Expression->NumParents() > 1)
+				{
+					InGraph->ClearAST();
+				}
+			}
+
+			RequestAutoVMRecompilation();	// We need to rebuild our metadata when a default value changes
+			break;
+		}
+	case ERigVMGraphNotifType::PinAdded:
+		{
+			if (URigVMPin* Pin = Cast<URigVMPin>(InSubject))
+			{
+				if (Pin->IsTraitPin())
+				{
+					RequestAutoVMRecompilation();
+				}
+			}
+			break;
+		}
+	case ERigVMGraphNotifType::PinRemoved:
+		{
+			RequestAutoVMRecompilation(); // can not check if it is a trait pin, as it has been already removed
+			break;
+		}
+	}
+	
+	// if the notification still has to be sent...
+	if (bNotifForOthersPending && !RigVMClient.bSuspendModelNotificationsForOthers)
+	{
+		if (RigVMGraphModifiedEvent.IsBound())
+		{
+			RigVMGraphModifiedEvent.Broadcast(InNotifType, InGraph, InSubject);
+		}
+	}
+}
+
+TSubclassOf<UAssetUserData> UAnimNextRigVMAssetEditorData::GetAssetUserDataClass() const
+{
+	return UAnimNextAssetWorkspaceAssetUserData::StaticClass();
+}
+
+TArray<UEdGraph*> UAnimNextRigVMAssetEditorData::GetAllEdGraphs() const
+{
+	TArray<UEdGraph*> Graphs;
+	for(UAnimNextRigVMAssetEntry* Entry : Entries)
+	{
+		if(IAnimNextRigVMGraphInterface* GraphInterface = Cast<IAnimNextRigVMGraphInterface>(Entry))
+		{
+			UEdGraph* EdGraph = GraphInterface->GetEdGraph();
+			Graphs.Add(EdGraph);
+			EdGraph->GetAllChildrenGraphs(Graphs);
+		}
+	}
+	for (URigVMEdGraph* RigVMEdGraph : FunctionEdGraphs)
+	{
+		Graphs.Add(RigVMEdGraph);
+		RigVMEdGraph->GetAllChildrenGraphs(Graphs);
+	}
+
+	return Graphs;
+}
+
+UAnimNextRigVMAssetEntry* UAnimNextRigVMAssetLibrary::FindEntry(UAnimNextRigVMAsset* InAsset, FName InName)
+{
+	return UE::UAF::UncookedOnly::FUtils::GetEditorData(InAsset)->FindEntry(InName);
+}
+
+UAnimNextRigVMAssetEntry* UAnimNextRigVMAssetEditorData::FindEntry(FName InName) const
+{
+	if(InName == NAME_None)
+	{
+		ReportError(TEXT("UAnimNextRigVMAssetEditorData::FindEntry: Invalid name supplied."));
+		return nullptr;
+	}
+
+	const TObjectPtr<UAnimNextRigVMAssetEntry>* FoundEntry = Entries.FindByPredicate([InName](const UAnimNextRigVMAssetEntry* InEntry)
+	{
+		if (!InEntry)
+		{
+			return false;
+		}
+
+		return InEntry->GetEntryName() == InName;
+	});
+
+	return FoundEntry != nullptr ? *FoundEntry : nullptr;
+}
+
+bool UAnimNextRigVMAssetEditorData::AddCategory(const FString& CategoryName, bool bInSetupUndoRedo, bool bPrintPythonCommand)
+{
+	if (CategoryName.IsEmpty())
+    {
+		ReportError(*FString::Printf(TEXT("UAnimNextRigVMAssetEditorData::AddCategory: Invalid empty category name provided.")));
+		return false;
+    }
+
+
+	const int32 CategoryIndex = VariableAndFunctionCategories.IndexOfByKey(CategoryName);
+	if(CategoryIndex != INDEX_NONE)
+	{
+		ReportError(*FString::Printf(TEXT("UAnimNextRigVMAssetEditorData::AddCategory: Already contains category with name %s."), *CategoryName));
+		return false;
+	}
+
+	if (bInSetupUndoRedo)
+	{
+		Modify();
+	}
+
+	VariableAndFunctionCategories.Add(CategoryName);
+	BroadcastModified(EAnimNextEditorDataNotifType::CategoryAdded, this);
+
+	if (bPrintPythonCommand)
+	{
+		RigVMPythonUtils::Print(GetName(),
+				FString::Printf(TEXT("asset.add_category('%s')"),
+				 *CategoryName));
+	}
+
+	return true;
+}
+
+bool UAnimNextRigVMAssetEditorData::RenameCategory(const FString& OldName, const FString& NewName, bool bInSetupUndoRedo, bool bPrintPythonCommand)
+{
+	const int32 CategoryIndex = VariableAndFunctionCategories.IndexOfByKey(OldName);
+	if (CategoryIndex == INDEX_NONE)
+	{
+		ReportError(*FString::Printf(TEXT("UAnimNextRigVMAssetEditorData::RenameCategory: Invalid existing category name provided (not found).")));
+		return false;
+	}
+	
+	if (bInSetupUndoRedo)
+	{
+		Modify();
+	}
+
+	VariableAndFunctionCategories[CategoryIndex] = NewName;
+
+	for(UAnimNextRigVMAssetEntry* Entry : Entries)
+	{
+		if(UAnimNextVariableEntry* VariableEntry = Cast<UAnimNextVariableEntry>(Entry))
+		{
+			if (VariableEntry->GetVariableCategory() == OldName)
+			{
+				VariableEntry->SetVariableCategory(NewName, bInSetupUndoRedo);
+			}
+		}
+	}
+	
+	BroadcastModified(EAnimNextEditorDataNotifType::CategoryChanged, this);
+
+	if (bPrintPythonCommand)
+	{
+		RigVMPythonUtils::Print(GetName(),
+			FString::Printf(TEXT("asset.rename_category('%s', '%s')"),
+			 *OldName,
+			 *NewName));
+	}
+
+	return true;
+}
+
+void UAnimNextRigVMAssetEditorData::ReorderCategory(const FString& CategoryName, const FString& BeforeCategoryName, bool bInSetupUndoRedo,
+	bool bPrintPythonCommand)
+{
+	const int32 CategoryIndex = VariableAndFunctionCategories.IndexOfByKey(CategoryName);
+	if (CategoryIndex == INDEX_NONE)
+	{
+		return;
+	}
+	
+	const int32 BeforeCategoryIndex = VariableAndFunctionCategories.IndexOfByKey(BeforeCategoryName);
+	if (BeforeCategoryIndex == INDEX_NONE)
+	{
+		return;
+	}
+
+	if (BeforeCategoryIndex == CategoryIndex)
+	{
+		return;
+	}
+
+	if (bInSetupUndoRedo)
+	{
+		Modify();
+	}
+
+	VariableAndFunctionCategories.RemoveAt(CategoryIndex);
+	
+	int32 NewCategoryIndex = 0;
+
+	// Moving up
+	if (CategoryIndex > BeforeCategoryIndex)
+	{
+		// RemoveAt didn't have any impact
+		NewCategoryIndex = FMath::Max(BeforeCategoryIndex, 0);
+	}
+	else
+	{
+		// RemoveAt means insertion index has to be offset by 1
+		NewCategoryIndex = FMath::Max(BeforeCategoryIndex - 1, 0);
+	}
+	
+	VariableAndFunctionCategories.Insert(CategoryName, NewCategoryIndex);
+
+	BroadcastModified(EAnimNextEditorDataNotifType::CategoryChanged, this);
+}
+
+void UAnimNextRigVMAssetEditorData::ReorderVariable(UAnimNextVariableEntry* VariableEntry, const UAnimNextVariableEntry* BeforeVariableEntry, bool bSetupUndoRedo, bool bPrintPythonCommand)
+{
+	const int32 VariableEntryIndex = InternalEntries.IndexOfByKey(VariableEntry);
+	if (VariableEntryIndex == INDEX_NONE)
+	{
+		return;
+	}
+	
+	const int32 BeforeVariableEntryIndex = InternalEntries.IndexOfByKey(BeforeVariableEntry);
+	if (BeforeVariableEntryIndex == INDEX_NONE)
+	{
+		return;
+	}
+
+	if (BeforeVariableEntryIndex == VariableEntryIndex)
+	{
+		return;
+	}
+
+	if (bSetupUndoRedo)
+	{
+		Modify();
+	}
+
+	RemoveEntryInternal(VariableEntry);
+
+	int32 NewVariableEntryIndex = 0;
+
+	// Moving up
+	if (VariableEntryIndex > BeforeVariableEntryIndex)
+	{
+		// RemoveAt didn't have any impact
+		NewVariableEntryIndex = FMath::Max(BeforeVariableEntryIndex, 0);
+	}
+	else
+	{
+		// RemoveAt means insertion index has to be offset by 1
+		NewVariableEntryIndex = FMath::Max(BeforeVariableEntryIndex - 1, 0);
+	}
+	
+	InsertEntryInternal(VariableEntry, NewVariableEntryIndex);
+
+	BroadcastModified(EAnimNextEditorDataNotifType::EntryAdded, VariableEntry);
+}
+
+bool UAnimNextRigVMAssetLibrary::RemoveEntry(UAnimNextRigVMAsset* InAsset, UAnimNextRigVMAssetEntry* InEntry,  bool bSetupUndoRedo, bool bPrintPythonCommand)
+{
+	return UE::UAF::UncookedOnly::FUtils::GetEditorData(InAsset)->RemoveEntry(InEntry, bSetupUndoRedo, bPrintPythonCommand);
+}
+
+bool UAnimNextRigVMAssetEditorData::RemoveEntry(UAnimNextRigVMAssetEntry* InEntry, bool bSetupUndoRedo, bool bPrintPythonCommand)
+{
+	if(InEntry == nullptr)
+	{
+		ReportError(TEXT("UAnimNextRigVMAssetEditorData::RemoveEntry: Invalid entry supplied."));
+		return false;
+	}
+	
+	TObjectPtr<UAnimNextRigVMAssetEntry>* EntryToRemovePtr = Entries.FindByKey(InEntry);
+	if(EntryToRemovePtr == nullptr)
+	{
+		ReportError(TEXT("UAnimNextRigVMAssetEditorData::RemoveEntry: Asset does not contain the supplied entry."));
+		return false;
+	}
+
+	if(bSetupUndoRedo)
+	{
+		Modify();
+	}
+
+	// Remove from internal array
+	UAnimNextRigVMAssetEntry* EntryToRemove = *EntryToRemovePtr;
+
+	bool bResult = true;
+	if(const IAnimNextRigVMGraphInterface* GraphInterface = Cast<IAnimNextRigVMGraphInterface>(EntryToRemove))
+	{
+		// Remove any graphs
+		if(URigVMGraph* RigVMGraph = GraphInterface->GetRigVMGraph())
+		{
+			TGuardValue<bool> EnablePythonPrint(bSuspendPythonMessagesForRigVMClient, !bPrintPythonCommand);
+			TGuardValue<bool> DisableAutoCompile(bAutoRecompileVM, false);
+			bResult = RigVMClient.RemoveModel(RigVMGraph->GetNodePath(), bSetupUndoRedo);
+		}
+	}
+
+	if (bSetupUndoRedo)
+	{
+		EntryToRemove->Modify();
+	}
+	RemoveEntryInternal(EntryToRemove);
+	RefreshExternalModels();
+
+	// This will cause any external package to be removed when saved
+	EntryToRemove->MarkAsGarbage();
+
+	BroadcastModified(EAnimNextEditorDataNotifType::EntryRemoved, this);
+
+	if (bPrintPythonCommand)
+	{
+		RigVMPythonUtils::Print(GetName(),
+				FString::Printf(TEXT("asset.remove_entry(asset.find_entry('%s'))"),
+				*InEntry->GetEntryName().ToString()));
+	}
+
+	return bResult;
+}
+
+bool UAnimNextRigVMAssetLibrary::RemoveEntries(UAnimNextRigVMAsset* InAsset, const TArray<UAnimNextRigVMAssetEntry*>& InEntries,  bool bSetupUndoRedo, bool bPrintPythonCommand)
+{
+	return UE::UAF::UncookedOnly::FUtils::GetEditorData(InAsset)->RemoveEntries(InEntries, bSetupUndoRedo, bPrintPythonCommand);
+}
+
+bool UAnimNextRigVMAssetEditorData::RemoveEntries(TConstArrayView<UAnimNextRigVMAssetEntry*> InEntries, bool bSetupUndoRedo, bool bPrintPythonCommand)
+{
+	bool bResult = false;
+	{
+		TGuardValue<bool> DisableEditorDataNotifications(bSuspendEditorDataNotifications, true);
+		TGuardValue<bool> DisableAutoCompile(bAutoRecompileVM, false);
+		for(UAnimNextRigVMAssetEntry* Entry : InEntries)
+		{
+			bResult |= RemoveEntry(Entry, bSetupUndoRedo, false);
+		}
+	}
+
+	BroadcastModified(EAnimNextEditorDataNotifType::EntryRemoved, this);
+
+	if (bPrintPythonCommand)
+	{
+
+		FString ArrayStr = TEXT("[");
+		for (int32 Index = 0; Index < InEntries.Num(); ++Index)
+		{
+			ArrayStr += TEXT("asset.find_entry('") + InEntries[Index]->GetEntryName().ToString() + TEXT("')");
+			if (Index < InEntries.Num() - 1)
+			{
+				ArrayStr += TEXT(", ");
+			}
+		}
+		ArrayStr += TEXT("]");
+
+
+		RigVMPythonUtils::Print(GetName(), 
+							FString::Printf(TEXT("asset.remove_entries(%s)"),
+											*ArrayStr));
+	}
+
+	return bResult;
+}
+
+bool UAnimNextRigVMAssetLibrary::RemoveAllEntries(UAnimNextRigVMAsset* InAsset, bool bSetupUndoRedo, bool bPrintPythonCommand)
+{
+	return UE::UAF::UncookedOnly::FUtils::GetEditorData(InAsset)->RemoveAllEntries(bSetupUndoRedo, bPrintPythonCommand);
+}
+
+bool UAnimNextRigVMAssetEditorData::RemoveAllEntries(bool bSetupUndoRedo, bool bPrintPythonCommand)
+{
+	bool bResult = false;
+	{
+		TGuardValue<bool> DisableEditorDataNotifications(bSuspendEditorDataNotifications, true);
+		TGuardValue<bool> DisableAutoCompile(bAutoRecompileVM, false);
+		TArray<UAnimNextRigVMAssetEntry*> EntriesCopy = Entries; 
+		for(UAnimNextRigVMAssetEntry* Entry : EntriesCopy)
+		{
+			bResult |= RemoveEntry(Entry, bSetupUndoRedo, false);
+		}
+	}
+
+	BroadcastModified(EAnimNextEditorDataNotifType::EntryRemoved, this);
+
+
+	if (bPrintPythonCommand)
+	{
+		RigVMPythonUtils::Print(GetName(), 
+							FString::Printf(TEXT("asset.remove_all_entries()")));
+	}
+
+	return bResult;
+}
+
+UObject* UAnimNextRigVMAssetEditorData::CreateNewSubEntry(UAnimNextRigVMAssetEditorData* InEditorData, TSubclassOf<UObject> InClass)
+{
+	UObject* NewEntry = NewObject<UObject>(InEditorData, InClass.Get(), NAME_None, RF_Transactional);
+	// If we are a transient asset, dont use external packages
+	UAnimNextRigVMAsset* Asset = UE::UAF::UncookedOnly::FUtils::GetAsset(InEditorData);
+	check(Asset);
+
+	// Additionally check external packaging flag
+	if(!Asset->HasAnyFlags(RF_Transient) && InEditorData->bUsesExternalPackages)
+	{
+		FExternalPackageHelper::SetPackagingMode(NewEntry, InEditorData, true, false, PKG_None);
+	}
+	return NewEntry;
+}
+
+UAnimNextRigVMAssetEntry* UAnimNextRigVMAssetEditorData::FindEntryForRigVMGraph(const URigVMGraph* InRigVMGraph) const
+{
+	for(UAnimNextRigVMAssetEntry* Entry : Entries)
+	{
+		if(IAnimNextRigVMGraphInterface* GraphInterface = Cast<IAnimNextRigVMGraphInterface>(Entry))
+		{
+			if (const URigVMGraph* RigVMGraph = GraphInterface->GetRigVMGraph())
+			{
+				if(RigVMGraph == InRigVMGraph)
+				{
+					return Entry;
+				}
+			}
+		}
+	}
+
+	return nullptr;
+}
+
+UAnimNextRigVMAssetEntry* UAnimNextRigVMAssetEditorData::FindEntryForRigVMEdGraph(const URigVMEdGraph* InRigVMEdGraph) const
+{
+	for (UAnimNextRigVMAssetEntry* Entry : Entries)
+	{
+		if (IAnimNextRigVMGraphInterface* GraphInterface = Cast<IAnimNextRigVMGraphInterface>(Entry))
+		{
+			if (GraphInterface->GetEdGraph() == InRigVMEdGraph)
+			{
+				return Entry;
+			}
+		}
+	}
+
+	return nullptr;
+}
+
+void UAnimNextRigVMAssetEditorData::CreateEdGraphForCollapseNode(URigVMCollapseNode* InNode, bool bForce)
+{
+	check(InNode);
+	URigVMGraph* CollapseNodeGraph = InNode->GetGraph();
+	check(CollapseNodeGraph);
+
+	if (bForce)
+	{
+		RemoveEdGraphForCollapseNode(InNode, false);
+	}
+
+	// For Function node
+	if (InNode->GetGraph()->IsA<URigVMFunctionLibrary>())
+	{
+		if (URigVMGraph* ContainedGraph = InNode->GetContainedGraph())
+		{
+			bool bFunctionGraphExists = false;
+			for (UEdGraph* FunctionGraph : FunctionEdGraphs)
+			{
+				if (URigVMEdGraph* RigFunctionGraph = Cast<URigVMEdGraph>(FunctionGraph))
+				{
+					if (RigFunctionGraph->ModelNodePath == ContainedGraph->GetNodePath())
+					{
+						bFunctionGraphExists = true;
+						break;
+					}
+				}
+			}
+
+			if (!bFunctionGraphExists)
+			{
+				const FName SubGraphName = RigVMClient.GetUniqueName(this, *InNode->GetName());
+				// create a sub graph
+				UAnimNextEdGraph* RigFunctionGraph = NewObject<UAnimNextEdGraph>(this, SubGraphName, RF_Transactional);
+				RigFunctionGraph->Schema = UAnimNextEdGraphSchema::StaticClass();
+				RigFunctionGraph->bAllowRenaming = true;
+				RigFunctionGraph->bEditable = true;
+				RigFunctionGraph->bAllowDeletion = true;
+				RigFunctionGraph->ModelNodePath = ContainedGraph->GetNodePath();
+				RigFunctionGraph->bIsFunctionDefinition = true;
+
+				RigFunctionGraph->Initialize(this);
+
+				FunctionEdGraphs.Add(RigFunctionGraph);
+
+				RigVMClient.GetOrCreateController(ContainedGraph)->ResendAllNotifications();
+			}
+		}
+	}
+	// --- For Collapse nodes ---
+	else if (URigVMEdGraph* RigEdGraph = Cast<URigVMEdGraph>(GetEditorObjectForRigVMGraph(InNode->GetGraph())))
+	{
+		if (URigVMGraph* ContainedGraph = InNode->GetContainedGraph())
+		{
+			bool bSubGraphExists = false;
+
+			const FString ContainedGraphNodePath = ContainedGraph->GetNodePath();
+			for (UEdGraph* SubGraph : RigEdGraph->SubGraphs)
+			{
+				if (UAnimNextEdGraph* SubRigGraph = Cast<UAnimNextEdGraph>(SubGraph))
+				{
+					if (SubRigGraph->ModelNodePath == ContainedGraphNodePath)
+					{
+						bSubGraphExists = true;
+						break;
+					}
+				}
+			}
+
+			if (!bSubGraphExists)
+			{
+				bool bEditable = true;
+				if (InNode->IsA<URigVMAggregateNode>())
+				{
+					bEditable = false;
+				}
+
+				UObject* Outer = FindEntryForRigVMGraph(CollapseNodeGraph->GetRootGraph());
+				if (Outer == nullptr)
+				{
+					Outer = this; // function library graph has no entry
+				}
+
+				const FName SubGraphName = RigVMClient.GetUniqueName(Outer, *InNode->GetEditorSubGraphName());
+				// create a sub graph, no need to set external package if outer is an Entry
+				UAnimNextEdGraph* SubRigGraph = NewObject<UAnimNextEdGraph>(Outer, SubGraphName, RF_Transactional);
+				SubRigGraph->Schema = UAnimNextEdGraphSchema::StaticClass();
+				SubRigGraph->bAllowRenaming = 1;
+				SubRigGraph->bEditable = bEditable;
+				SubRigGraph->bAllowDeletion = 1;
+				SubRigGraph->ModelNodePath = ContainedGraphNodePath;
+				SubRigGraph->bIsFunctionDefinition = false;
+
+				RigEdGraph->SubGraphs.Add(SubRigGraph);
+
+				SubRigGraph->Initialize(this);
+
+				GetOrCreateController(ContainedGraph)->ResendAllNotifications();
+			}
+		}
+	}
+}
+
+void UAnimNextRigVMAssetEditorData::RemoveEdGraphForCollapseNode(URigVMCollapseNode* InNode, bool bNotify)
+{
+	check(InNode);
+
+	if (InNode->GetGraph()->IsA<URigVMFunctionLibrary>())
+	{
+		if (URigVMGraph* ContainedGraph = InNode->GetContainedGraph())
+		{
+			for (UEdGraph* FunctionGraph : FunctionEdGraphs)
+			{
+				if (URigVMEdGraph* RigFunctionGraph = Cast<URigVMEdGraph>(FunctionGraph))
+				{
+					if (RigFunctionGraph->ModelNodePath == ContainedGraph->GetNodePath())
+					{
+						if (URigVMController* SubController = GetController(ContainedGraph))
+						{
+							SubController->OnModified().RemoveAll(RigFunctionGraph);
+						}
+
+						if (RigVMGraphModifiedEvent.IsBound() && bNotify)
+						{
+							RigVMGraphModifiedEvent.Broadcast(ERigVMGraphNotifType::NodeRemoved, InNode->GetGraph(), InNode);
+						}
+
+						FunctionEdGraphs.Remove(RigFunctionGraph);
+						RigFunctionGraph->Rename(nullptr, GetTransientPackage(), REN_ForceNoResetLoaders | REN_DontCreateRedirectors);
+						RigFunctionGraph->MarkAsGarbage();
+						break;
+					}
+				}
+			}
+		}
+	}
+	else if (URigVMEdGraph* RigGraph = Cast<URigVMEdGraph>(GetEditorObjectForRigVMGraph(InNode->GetGraph())))
+	{
+		if (URigVMGraph* ContainedGraph = InNode->GetContainedGraph())
+		{
+			for (UEdGraph* SubGraph : RigGraph->SubGraphs)
+			{
+				if (URigVMEdGraph* SubRigGraph = Cast<URigVMEdGraph>(SubGraph))
+				{
+					if (SubRigGraph->ModelNodePath == ContainedGraph->GetNodePath())
+					{
+						if (URigVMController* SubController = GetController(ContainedGraph))
+						{
+							SubController->OnModified().RemoveAll(SubRigGraph);
+						}
+
+						if (RigVMGraphModifiedEvent.IsBound() && bNotify)
+						{
+							RigVMGraphModifiedEvent.Broadcast(ERigVMGraphNotifType::NodeRemoved, InNode->GetGraph(), InNode);
+						}
+
+						RigGraph->SubGraphs.Remove(SubRigGraph);
+						SubRigGraph->Rename(nullptr, GetTransientPackage(), REN_ForceNoResetLoaders | REN_DontCreateRedirectors);
+						SubRigGraph->MarkAsGarbage();
+						break;
+					}
+				}
+			}
+		}
+	}
+}
+
+UEdGraph* UAnimNextRigVMAssetEditorData::CreateEdGraph(URigVMGraph* InRigVMGraph, bool bForce)
+{
+	check(InRigVMGraph);
+
+	if(InRigVMGraph->IsA<URigVMFunctionLibrary>())
+	{
+		return nullptr;
+	}
+
+	const bool bIsTransient = InRigVMGraph->HasAnyFlags(RF_Transient);
+	IAnimNextRigVMGraphInterface* Entry = Cast<IAnimNextRigVMGraphInterface>(FindEntryForRigVMGraph(InRigVMGraph));
+	if(Entry == nullptr && !bIsTransient)
+	{
+		// Not found, we could be adding a new entry, in which case the graph wont be assigned yet
+		check(Entries.Num() > 0);
+		check(Cast<IAnimNextRigVMGraphInterface>(Entries.Last()) != nullptr);
+		check(Cast<IAnimNextRigVMGraphInterface>(Entries.Last())->GetRigVMGraph() == nullptr);
+		Entry = Cast<IAnimNextRigVMGraphInterface>(FindEntryForRigVMGraph(nullptr));
+	}
+
+	if(Entry == nullptr && !bIsTransient)
+	{
+		return nullptr;
+	}
+	
+	if(bForce)
+	{
+		RemoveEdGraph(InRigVMGraph);
+	}
+
+	UObject* Outer = nullptr;
+	EObjectFlags Flags = RF_NoFlags;
+	if(!bIsTransient)
+	{
+		Outer = CastChecked<UObject>(Entry);
+		Flags = RF_Transactional;
+	}
+	else
+	{
+		// This outer is to allow URigVMEdGraph::GetModel to retrieve the graph in 'preview' scenarios 
+		Outer = InRigVMGraph;
+		Flags = RF_Transient;
+	}
+
+	const FName GraphName = Entry != nullptr ? RigVMClient.GetUniqueName(Outer, Entry->GetGraphName()) : NAME_None;
+	UAnimNextEdGraph* RigFunctionGraph = NewObject<UAnimNextEdGraph>(Outer, GraphName, Flags);
+	RigFunctionGraph->Schema = UAnimNextEdGraphSchema::StaticClass();
+	RigFunctionGraph->bAllowDeletion = true;
+	RigFunctionGraph->bIsFunctionDefinition = false;
+	RigFunctionGraph->ModelNodePath = InRigVMGraph->GetNodePath();
+	RigFunctionGraph->Initialize(this);
+
+	if(!bIsTransient)
+	{
+		Entry->SetEdGraph(RigFunctionGraph);
+		if(Entry->GetRigVMGraph() == nullptr)
+		{
+			Entry->SetRigVMGraph(InRigVMGraph);
+		}
+		else
+		{
+			check(Entry->GetRigVMGraph() == InRigVMGraph);
+		}
+	}
+
+	return RigFunctionGraph;
+}
+
+bool UAnimNextRigVMAssetEditorData::RemoveEdGraph(URigVMGraph* InModel)
+{
+	if(IAnimNextRigVMGraphInterface* Entry = Cast<IAnimNextRigVMGraphInterface>(FindEntryForRigVMGraph(InModel)))
+	{
+		RigVMClient.DestroyObject(Entry->GetEdGraph());
+		Entry->SetEdGraph(nullptr);
+		return true;
+	}
+	return false;
+}
+
+UAnimNextVariableEntry* UAnimNextRigVMAssetLibrary::AddVariable(UAnimNextRigVMAsset* InAsset, FName InName, EPropertyBagPropertyType InValueType,
+	EPropertyBagContainerType InContainerType, const UObject* InValueTypeObject, const FString& InDefaultValue, bool bSetupUndoRedo, bool bPrintPythonCommand)
+{
+	return UE::UAF::UncookedOnly::FUtils::GetEditorData(InAsset)->AddVariable(InName, FAnimNextParamType(InValueType, InContainerType, InValueTypeObject), InDefaultValue, bSetupUndoRedo, bPrintPythonCommand);
+}
+
+UAnimNextVariableEntry* UAnimNextRigVMAssetEditorData::AddVariable(FName InName, FAnimNextParamType InType, const FString& InDefaultValue, bool bSetupUndoRedo, bool bPrintPythonCommand)
+{
+	if(InName == NAME_None)
+	{
+		ReportError(TEXT("UAnimNextRigVMAssetEditorData::AddVariable: Invalid variable name supplied."));
+		return nullptr;
+	}
+
+	if(!GetEntryClasses().Contains(UAnimNextVariableEntry::StaticClass()) || !CanAddNewEntry(UAnimNextVariableEntry::StaticClass()))
+	{
+		ReportError(TEXT("UAnimNextRigVMAssetEditorData::AddVariable: Cannot add a variable to this asset - entry is not allowed."));
+		return nullptr;
+	}
+
+	// Check for duplicate name
+	FName NewParameterName = InName;
+	auto DuplicateNamePredicate = [&NewParameterName](const UAnimNextRigVMAssetEntry* InEntry)
+	{
+		if (!InEntry)
+		{
+			return false;
+		}
+
+		return InEntry->GetEntryName() == NewParameterName;
+	};
+
+	bool bAlreadyExists = Entries.ContainsByPredicate(DuplicateNamePredicate);
+	int32 NameNumber = InName.GetNumber() + 1;
+	while(bAlreadyExists)
+	{
+		NewParameterName = FName(InName, NameNumber++);
+		bAlreadyExists = Entries.ContainsByPredicate(DuplicateNamePredicate);
+	}
+
+	UAnimNextVariableEntry* NewEntry = CreateNewSubEntry<UAnimNextVariableEntry>(this);
+	{
+		TGuardValue<bool> DisableEditorDataNotifications(bSuspendEditorDataNotifications, true);
+		TGuardValue<bool> DisableAutoCompile(bAutoRecompileVM, false);
+
+		NewEntry->SetVariableName(NewParameterName, false);
+		NewEntry->SetType(InType, false);
+		if(InDefaultValue.Len() > 0)
+		{
+			NewEntry->SetDefaultValueFromString(InDefaultValue, false);
+		}
+
+		NewEntry->Initialize(this);
+	}
+
+	if(bSetupUndoRedo)
+	{
+		NewEntry->Modify();
+		Modify();
+	}
+
+	AddEntryInternal(NewEntry);
+
+	CustomizeNewAssetEntry(NewEntry);
+
+	BroadcastModified(EAnimNextEditorDataNotifType::EntryAdded, NewEntry);
+
+	if (bPrintPythonCommand)
+	{
+		const FString ValueTypeString = InType.GetValueTypeObject() ?
+				FString::Printf(TEXT("unreal.%s.static_%s()"), *InType.GetValueTypeObject()->GetName(), InType.GetValueTypeObject()->IsA<UScriptStruct>() ? TEXT("struct") : TEXT("class"))
+				: TEXT("None");
+		RigVMPythonUtils::Print(GetName(), 
+							FString::Printf(TEXT("asset.add_variable('%s', %s, %s, %s, '%s')"),
+											*InName.ToString(),
+											*RigVMPythonUtils::EnumValueToPythonString<EPropertyBagPropertyType>(static_cast<int64>(InType.GetValueType())),
+											*RigVMPythonUtils::EnumValueToPythonString<EPropertyBagContainerType>(static_cast<int64>(InType.GetContainerType())),
+											*ValueTypeString,
+											*InDefaultValue));
+	}
+
+	return NewEntry;
+}
+
+UAnimNextEventGraphEntry* UAnimNextRigVMAssetLibrary::AddEventGraph(UAnimNextRigVMAsset* InAsset, FName InName, UScriptStruct* InEventStruct, bool bSetupUndoRedo, bool bPrintPythonCommand)
+{
+	return UE::UAF::UncookedOnly::FUtils::GetEditorData(InAsset)->AddEventGraph(InName, InEventStruct, bSetupUndoRedo, bPrintPythonCommand);
+}
+
+UAnimNextEventGraphEntry* UAnimNextRigVMAssetEditorData::AddEventGraph(FName InName, UScriptStruct* InEventStruct, bool bSetupUndoRedo, bool bPrintPythonCommand)
+{
+	if(InName == NAME_None)
+	{
+		ReportError(TEXT("UAnimNextRigVMAssetEditorData::AddEventGraph: Invalid graph name supplied."));
+		return nullptr;
+	}
+
+	if(InEventStruct == nullptr || !InEventStruct->IsChildOf(FRigVMStruct::StaticStruct()))
+	{
+		ReportError(TEXT("UAnimNextRigVMAssetEditorData::AddEventGraph: Invalid event struct name supplied."));
+		return nullptr;
+	}
+
+	if(!GetEntryClasses().Contains(UAnimNextEventGraphEntry::StaticClass()) || !CanAddNewEntry(UAnimNextEventGraphEntry::StaticClass()))
+	{
+		ReportError(TEXT("UAnimNextRigVMAssetEditorData::AddEventGraph: Cannot add an event graph to this asset - entry is not allowed."));
+		return nullptr;
+	}
+
+	// Check for duplicate name
+	FName NewGraphName = InName;
+	auto DuplicateNamePredicate = [&NewGraphName](const UAnimNextRigVMAssetEntry* InEntry)
+	{
+		return InEntry->GetEntryName() == NewGraphName;
+	};
+
+	bool bAlreadyExists = Entries.ContainsByPredicate(DuplicateNamePredicate);
+	int32 NameNumber = InName.GetNumber() + 1;
+	while(bAlreadyExists)
+	{
+		NewGraphName = FName(InName, NameNumber++);
+		bAlreadyExists =  Entries.ContainsByPredicate(DuplicateNamePredicate);
+	}
+
+	UAnimNextEventGraphEntry* NewEntry = CreateNewSubEntry<UAnimNextEventGraphEntry>(this);
+	NewEntry->GraphName = NewGraphName;
+	NewEntry->Initialize(this);
+
+	if(bSetupUndoRedo)
+	{
+		NewEntry->Modify();
+		Modify();
+	}
+
+	AddEntryInternal(NewEntry);
+
+	// Add new graph
+	{
+		TGuardValue<bool> EnablePythonPrint(bSuspendPythonMessagesForRigVMClient, !bPrintPythonCommand);
+		TGuardValue<bool> DisableAutoCompile(bAutoRecompileVM, false);
+		// Editor data has to be the graph outer, or RigVM unique name generator will not work
+		URigVMGraph* NewRigVMGraphModel = RigVMClient.CreateModel(URigVMGraph::StaticClass()->GetFName(), UAnimNextEventGraphSchema::StaticClass(), bSetupUndoRedo, this);
+		if (ensure(NewRigVMGraphModel))
+		{
+			// Then, to avoid the graph losing ref due to external package, set the same package as the Entry
+			if (!NewRigVMGraphModel->HasAnyFlags(RF_Transient))
+			{
+				NewRigVMGraphModel->SetExternalPackage(CastChecked<UObject>(NewEntry)->GetExternalPackage());
+			}
+			ensure(NewRigVMGraphModel);
+			NewEntry->Graph = NewRigVMGraphModel;
+
+			RefreshExternalModels();
+			RigVMClient.AddModel(NewRigVMGraphModel, true);
+			URigVMController* Controller = RigVMClient.GetController(NewRigVMGraphModel);
+			UE::UAF::UncookedOnly::FUtils::SetupEventGraph(Controller, InEventStruct, NewGraphName);
+		}
+	}
+
+	CustomizeNewAssetEntry(NewEntry);
+
+	BroadcastModified(EAnimNextEditorDataNotifType::EntryAdded, NewEntry);
+
+	if (bPrintPythonCommand)
+	{
+		RigVMPythonUtils::Print(GetName(),
+				FString::Printf(TEXT("asset.add_event_graph('%s', unreal.%s)"),
+				*InName.ToString(), *InEventStruct->GetName()));
+	}
+
+	return NewEntry;
+}
+
+UAnimNextSharedVariablesEntry* UAnimNextRigVMAssetLibrary::AddSharedVariables(UAnimNextRigVMAsset* InAsset, UAnimNextSharedVariables* InSharedVariables, bool bSetupUndoRedo, bool bPrintPythonCommand)
+{
+	return UE::UAF::UncookedOnly::FUtils::GetEditorData(InAsset)->AddSharedVariables(InSharedVariables, bSetupUndoRedo, bPrintPythonCommand);
+}
+
+UAnimNextSharedVariablesEntry* UAnimNextRigVMAssetEditorData::AddSharedVariables(const UAnimNextSharedVariables* InSharedVariables, bool bSetupUndoRedo, bool bPrintPythonCommand)
+{
+	if(InSharedVariables == nullptr)
+	{
+		ReportError(TEXT("UAnimNextRigVMAssetEditorData::AddSharedVariables: Invalid asset supplied."));
+		return nullptr;
+	}
+
+	if(!GetEntryClasses().Contains(UAnimNextSharedVariablesEntry::StaticClass()) || !CanAddNewEntry(UAnimNextSharedVariablesEntry::StaticClass()))
+	{
+		ReportError(TEXT("UAnimNextRigVMAssetEditorData::AddSharedVariables: Cannot add a shared variables to this asset - entry is not allowed."));
+		return nullptr;
+	}
+	
+	// Check if interface has any public members or if any of its parent interfaces do
+	UAnimNextSharedVariables_EditorData* EditorData = UE::UAF::UncookedOnly::FUtils::GetEditorData<UAnimNextSharedVariables_EditorData>(InSharedVariables);
+	if(EditorData == nullptr)
+	{
+		ReportError(TEXT("UAnimNextRigVMAssetEditorData::AddSharedVariables: Invalid asset supplied - asset has no editor data."));
+		return nullptr;
+	}
+
+	// Check for circularity
+	auto CheckForCircularity = [this](UAnimNextSharedVariables_EditorData* InEditorData, auto& InCheckForCircularity)
+	{
+		if(InEditorData == this)
+		{
+			return true;
+		}
+
+		for(UAnimNextRigVMAssetEntry* Entry : InEditorData->Entries)
+		{
+			if(UAnimNextSharedVariablesEntry* SharedVariablesEntry = Cast<UAnimNextSharedVariablesEntry>(Entry))
+			{
+				const UAnimNextSharedVariables* SharedVariables = SharedVariablesEntry->GetAsset();
+				UAnimNextSharedVariables_EditorData* EditorData = UE::UAF::UncookedOnly::FUtils::GetEditorData<UAnimNextSharedVariables_EditorData>(SharedVariables);
+				if(InCheckForCircularity(EditorData, InCheckForCircularity))
+				{
+					return true;
+				}
+			}
+		}
+
+		return false;
+	};
+
+	if(CheckForCircularity(EditorData, CheckForCircularity))
+	{
+		ReportError(TEXT("UAnimNextRigVMAssetEditorData::AddSharedVariables: Circular reference detected."));
+		return nullptr;
+	}
+
+	auto CheckForPublicMembers = [](UAnimNextSharedVariables_EditorData* InEditorData, auto& InCheckForPublicMembers)
+	{
+		for(UAnimNextRigVMAssetEntry* Entry : InEditorData->Entries)
+		{
+			if(UAnimNextVariableEntry* VariableEntry = Cast<UAnimNextVariableEntry>(Entry))
+			{
+				if(VariableEntry->GetExportAccessSpecifier() == EAnimNextExportAccessSpecifier::Public)
+				{
+					return true;
+				}
+			}
+			else if(UAnimNextSharedVariablesEntry* SharedVariablesEntry = Cast<UAnimNextSharedVariablesEntry>(Entry))
+			{
+				const UAnimNextSharedVariables* SharedVariables = SharedVariablesEntry->GetAsset();
+				UAnimNextSharedVariables_EditorData* EditorData = UE::UAF::UncookedOnly::FUtils::GetEditorData<UAnimNextSharedVariables_EditorData>(SharedVariables);
+				if(InCheckForPublicMembers(EditorData, InCheckForPublicMembers))
+				{
+					return true;
+				}
+			}
+		}
+
+		return false;
+	};
+
+	if(!CheckForPublicMembers(EditorData, CheckForPublicMembers))
+	{
+		ReportError(TEXT("UAnimNextRigVMAssetEditorData::AddSharedVariables: No public variables found."));
+		return nullptr;
+	}
+
+	// Check for duplicate entry
+	auto DuplicatePredicate = [InSharedVariables](const UAnimNextRigVMAssetEntry* InEntry)
+	{
+		if(const UAnimNextSharedVariablesEntry* SharedVariablesEntry = Cast<UAnimNextSharedVariablesEntry>(InEntry))
+		{
+			return SharedVariablesEntry->GetAsset() == InSharedVariables;
+		}
+		return false;
+	};
+
+	if(Entries.ContainsByPredicate(DuplicatePredicate))
+	{
+		ReportError(TEXT("UAnimNextRigVMAssetEditorData::AddSharedVariables: Shared variables already referenced."));
+		return nullptr;
+	}
+
+	UAnimNextSharedVariablesEntry* NewEntry = CreateNewSubEntry<UAnimNextSharedVariablesEntry>(this);
+	NewEntry->SetAsset(InSharedVariables);
+	NewEntry->Initialize(this);
+
+	if(bSetupUndoRedo)
+	{
+		NewEntry->Modify();
+		Modify();
+	}
+
+	AddEntryInternal(NewEntry);
+
+	CustomizeNewAssetEntry(NewEntry);
+
+	BroadcastModified(EAnimNextEditorDataNotifType::EntryAdded, NewEntry);
+
+	if (bPrintPythonCommand)
+	{
+		RigVMPythonUtils::Print(GetName(),
+				FString::Printf(TEXT("asset.add_shared_variables(unreal.find_object(outer=None, name='%s'))"),
+				 *InSharedVariables->GetPathName()));
+	}
+
+	return NewEntry;
+}
+
+
+UAnimNextSharedVariablesEntry* UAnimNextRigVMAssetLibrary::AddSharedVariablesStruct(UAnimNextRigVMAsset* InAsset, UScriptStruct* InStruct, bool bSetupUndoRedo, bool bPrintPythonCommand)
+{
+	return UE::UAF::UncookedOnly::FUtils::GetEditorData(InAsset)->AddSharedVariablesStruct(InStruct, bSetupUndoRedo, bPrintPythonCommand);
+}
+
+UAnimNextSharedVariablesEntry* UAnimNextRigVMAssetEditorData::AddSharedVariablesStruct(const UScriptStruct* InStruct, bool bSetupUndoRedo, bool bPrintPythonCommand)
+{
+	if(InStruct == nullptr)
+	{
+		ReportError(TEXT("UAnimNextRigVMAssetEditorData::AddSharedVariablesStruct: Invalid struct supplied."));
+		return nullptr;
+	}
+
+	if(!GetEntryClasses().Contains(UAnimNextSharedVariablesEntry::StaticClass()) || !CanAddNewEntry(UAnimNextSharedVariablesEntry::StaticClass()))
+	{
+		ReportError(TEXT("UAnimNextRigVMAssetEditorData::AddSharedVariablesStruct: Cannot add a shared variables struct to this asset - entry is not allowed."));
+		return nullptr;
+	}
+
+	// Check for duplicate entry
+	auto DuplicatePredicate = [InStruct](const UAnimNextRigVMAssetEntry* InEntry)
+	{
+		if(const UAnimNextSharedVariablesEntry* SharedVariablesEntry = Cast<UAnimNextSharedVariablesEntry>(InEntry))
+		{
+			return SharedVariablesEntry->GetStruct() == InStruct;
+		}
+		return false;
+	};
+
+	if(Entries.ContainsByPredicate(DuplicatePredicate))
+	{
+		ReportError(TEXT("UAnimNextRigVMAssetEditorData::AddSharedVariables: Shared variables struct already referenced."));
+		return nullptr;
+	}
+
+	UAnimNextSharedVariablesEntry* NewEntry = CreateNewSubEntry<UAnimNextSharedVariablesEntry>(this);
+	NewEntry->SetStruct(InStruct);
+	NewEntry->Initialize(this);
+
+	if(bSetupUndoRedo)
+	{
+		NewEntry->Modify();
+		Modify();
+	}
+
+	AddEntryInternal(NewEntry);
+
+	CustomizeNewAssetEntry(NewEntry);
+
+	BroadcastModified(EAnimNextEditorDataNotifType::EntryAdded, NewEntry);
+
+	if (bPrintPythonCommand)
+	{
+		RigVMPythonUtils::Print(GetName(),
+				FString::Printf(TEXT("asset.add_shared_variables_struct(unreal.find_object(outer=None, name='%s'))"),
+				 *InStruct->GetPathName()));
+	}
+
+	return NewEntry;
+}
+
+URigVMLibraryNode* UAnimNextRigVMAssetLibrary::AddFunction(UAnimNextRigVMAsset* InAsset, FName InFunctionName, bool bInMutable, bool bSetupUndoRedo, bool bPrintPythonCommand)
+{
+	return UE::UAF::UncookedOnly::FUtils::GetEditorData(InAsset)->AddFunction(InFunctionName, bInMutable, bSetupUndoRedo, bPrintPythonCommand);
+}
+
+bool UAnimNextRigVMAssetLibrary::AddCategory(UAnimNextRigVMAsset* InAsset, const FString& CategoryName, bool bSetupUndoRedo, bool bPrintPythonCommand)
+{
+	return UE::UAF::UncookedOnly::FUtils::GetEditorData(InAsset)->AddCategory(CategoryName, bSetupUndoRedo, bPrintPythonCommand);
+}
+
+bool UAnimNextRigVMAssetLibrary::RenameCategory(UAnimNextRigVMAsset* InAsset, const FString& CategoryName, const FString& NewCategoryName, bool bSetupUndoRedo, bool bPrintPythonCommand)
+{
+	return UE::UAF::UncookedOnly::FUtils::GetEditorData(InAsset)->RenameCategory(CategoryName, NewCategoryName, bSetupUndoRedo, bPrintPythonCommand);
+}
+
+URigVMLibraryNode* UAnimNextRigVMAssetEditorData::AddFunction(FName InFunctionName, bool bInMutable, bool bInSetupUndoRedo, bool bPrintPythonCommand)
+{
+	if (bInSetupUndoRedo)
+	{
+		Modify();
+	}
+
+	URigVMController* Controller = RigVMClient.GetOrCreateController(GetLocalFunctionLibrary());
+	URigVMLibraryNode* Node = Controller->AddFunctionToLibrary(InFunctionName, bInMutable, FVector2D::ZeroVector, bInSetupUndoRedo, false);
+
+	if (bPrintPythonCommand)
+	{
+		RigVMPythonUtils::Print(GetName(),
+				FString::Printf(TEXT("asset.add_function('%s', %s)"),
+				 *InFunctionName.ToString(),
+				 bInMutable ? TEXT("True") : TEXT("False")));
+	}
+
+	return Node;
+}
+
+bool UAnimNextRigVMAssetEditorData::HasPublicVariables() const
+{
+	for(UAnimNextRigVMAssetEntry* Entry : Entries)
+	{
+		if(UAnimNextVariableEntry* VariableEntry = Cast<UAnimNextVariableEntry>(Entry))
+		{
+			if(VariableEntry->GetExportAccessSpecifier() == EAnimNextExportAccessSpecifier::Public)
+			{
+				return true;
+			}
+		}
+		else if(UAnimNextSharedVariablesEntry* SharedVariablesEntry = Cast<UAnimNextSharedVariablesEntry>(Entry))
+		{
+			if(SharedVariablesEntry->GetAsset())
+			{
+				UAnimNextSharedVariables_EditorData* EditorData = UE::UAF::UncookedOnly::FUtils::GetEditorData<UAnimNextSharedVariables_EditorData>(SharedVariablesEntry->GetAsset());
+				return EditorData->HasPublicVariables();
+			}
+		}
+	}
+	return false;
+}
+
+void UAnimNextRigVMAssetEditorData::GetAllVariables(TArray<FVariableInfo>& OutVariables, EVariableRecursion InRecursion, EVariableAccessFilter InAccess) const
+{
+	for(UAnimNextRigVMAssetEntry* Entry : Entries)
+	{
+		if(UAnimNextVariableEntry* VariableEntry = Cast<UAnimNextVariableEntry>(Entry))
+		{
+			if (InAccess == EVariableAccessFilter::All ||
+				VariableEntry->GetExportAccessSpecifier() == EAnimNextExportAccessSpecifier::Public)
+			{
+				const UAnimNextRigVMAsset* Asset = UE::UAF::UncookedOnly::FUtils::GetAsset<UAnimNextRigVMAsset>(this);
+
+				const FProperty* Property;
+				TConstArrayView<uint8> Value;
+				VariableEntry->GetDefaultValue(Property, Value);
+				
+				OutVariables.Add(
+					{
+						VariableEntry->GetExportName(),
+						VariableEntry->GetType(),
+						Asset,
+						VariableEntry->GetExportAccessSpecifier(),
+						Property,
+						Value
+					});
+			}
+		}
+		else if(UAnimNextSharedVariablesEntry* SharedVariablesEntry = Cast<UAnimNextSharedVariablesEntry>(Entry))
+		{
+			if (InRecursion == EVariableRecursion::SelfOnly)
+			{
+				continue;
+			}
+
+			switch (SharedVariablesEntry->GetType())
+			{
+			case EAnimNextSharedVariablesType::Asset:
+				if(SharedVariablesEntry->GetAsset())
+				{
+					UAnimNextSharedVariables_EditorData* EditorData = UE::UAF::UncookedOnly::FUtils::GetEditorData<UAnimNextSharedVariables_EditorData>(SharedVariablesEntry->GetAsset());
+					EditorData->GetAllVariables(OutVariables, InRecursion, EVariableAccessFilter::PublicOnly);
+				}
+				break;
+			case EAnimNextSharedVariablesType::Struct:
+				if(SharedVariablesEntry->GetStruct())
+				{
+					for (TFieldIterator<FProperty> It(SharedVariablesEntry->GetStruct()); It; ++It)
+					{
+						const FProperty* Property = *It;
+						if (InAccess == EVariableAccessFilter::All || Property->HasAnyPropertyFlags(CPF_NativeAccessSpecifierPublic))
+						{
+							OutVariables.Add(
+								{
+									Property->GetFName(),
+									FAnimNextParamType::FromProperty(Property),
+									SharedVariablesEntry->GetStruct(),
+									EAnimNextExportAccessSpecifier::Public,
+									Property,
+									TConstArrayView<uint8>()
+								});
+						}
+					}
+				}
+				break;
+			}
+		}
+	}
+}
+
+void UAnimNextRigVMAssetEditorData::HandleReportFromCompiler(EMessageSeverity::Type InSeverity, UObject* InSubject, const FString& InMessage)
+{
+	if (bSuspendCompilerReports)
+	{
+		return;
+	}
+
+	FCompilerResultsLog& Log = UE::UAF::UncookedOnly::FScopedCompileJob::GetLog();
+
+	UObject* SubjectForMessage = InSubject;
+	if(URigVMNode* ModelNode = Cast<URigVMNode>(SubjectForMessage))
+	{
+		if(IRigVMClientHost* RigVMClientHost = ModelNode->GetImplementingOuter<IRigVMClientHost>())
+		{
+			if(URigVMNode* OriginalModelNode = Cast<URigVMNode>(Log.FindSourceObject(ModelNode)))
+			{
+				if(URigVMEdGraph* EdGraph = Cast<URigVMEdGraph>(RigVMClientHost->GetEditorObjectForRigVMGraph(OriginalModelNode->GetGraph())))
+				{
+					if(UEdGraphNode* EdNode = EdGraph->FindNodeForModelNodeName(OriginalModelNode->GetFName()))
+					{
+						SubjectForMessage = EdNode;
+					}
+				}
+			}
+		}
+	}
+
+	TSharedPtr<FTokenizedMessage> Message;
+	if (InSeverity == EMessageSeverity::Error)
+	{
+		// see UnitTest "ControlRig.Basics.OrphanedPins" to learn why errors are suppressed this way
+		if (VMCompileSettings.SurpressErrors)
+		{
+			Log.bSilentMode = true;
+		}
+
+		if(InMessage.Contains(TEXT("@@")))
+		{
+			Message = Log.Error(*InMessage, SubjectForMessage);
+		}
+		else
+		{
+			Message = Log.Error(*InMessage);
+		}
+
+		// see UnitTest "ControlRig.Basics.OrphanedPins" to learn why errors are suppressed this way
+		if (!VMCompileSettings.SurpressErrors)
+		{ 
+			FScriptExceptionHandler::Get().HandleException(ELogVerbosity::Error, *InMessage, *FString());
+		}
+
+		bErrorsDuringCompilation = true;
+	}
+	else if (InSeverity == EMessageSeverity::Warning)
+	{
+		if(InMessage.Contains(TEXT("@@")))
+		{
+			Message = Log.Warning(*InMessage, SubjectForMessage);
+		}
+		else
+		{
+			Message = Log.Warning(*InMessage);
+		}
+
+		FScriptExceptionHandler::Get().HandleException(ELogVerbosity::Warning, *InMessage, *FString());
+
+		bWarningsDuringCompilation = true;
+	}
+	else
+	{
+		if(InMessage.Contains(TEXT("@@")))
+		{
+			Message = Log.Note(*InMessage, SubjectForMessage);
+		}
+		else
+		{
+			Message = Log.Note(*InMessage);
+		}
+
+		UE_LOG(LogAnimation, Display, TEXT("%s"), *InMessage);
+	}
+
+	if (URigVMEdGraphNode* EdGraphNode = Cast<URigVMEdGraphNode>(SubjectForMessage))
+	{
+		if(Message.IsValid())
+		{
+			EdGraphNode->SetErrorInfo(InSeverity, Message->ToText().ToString());
+		}
+		else
+		{
+			EdGraphNode->SetErrorInfo(InSeverity, InMessage);
+		}
+
+		EdGraphNode->bHasCompilerMessage = EdGraphNode->ErrorType <= int32(EMessageSeverity::Info);
+	}
+}
+
+void UAnimNextRigVMAssetEditorData::ClearErrorInfoForAllEdGraphs()
+{
+	for (UEdGraph* Graph : GetAllEdGraphs())
+	{
+		URigVMEdGraph* RigGraph = Cast<URigVMEdGraph>(Graph);
+		if (RigGraph == nullptr)
+		{
+			continue;
+		}
+
+		for (UEdGraphNode* GraphNode : Graph->Nodes)
+		{
+			if (URigVMEdGraphNode* RigVMEdGraphNode = Cast<URigVMEdGraphNode>(GraphNode))
+			{
+				RigVMEdGraphNode->ClearErrorInfo();
+			}
+		}
+	}
+}
+
+void UAnimNextRigVMAssetEditorData::RefreshExternalModels()
+{
+	GraphModels.Reset();
+
+	for (UAnimNextRigVMAssetEntry* Entry : Entries)
+	{
+		if (IAnimNextRigVMGraphInterface* GraphInterface = Cast<IAnimNextRigVMGraphInterface>(Entry))
+		{
+			if(URigVMGraph* Model = GraphInterface->GetRigVMGraph())
+			{
+				GraphModels.Add(Model);
+			}
+		}
+	}
+}
+
+void UAnimNextRigVMAssetEditorData::AddReferencedObjects(UObject* InThis, FReferenceCollector& Collector)
+{
+	Super::AddReferencedObjects(InThis, Collector);
+
+	UAnimNextRigVMAssetEditorData* This = CastChecked<UAnimNextRigVMAssetEditorData>(InThis);
+
+	if (This->CachedExports.IsSet())
+	{
+		// Cached exports may hold references to objects, so make GC aware
+		Collector.AddPropertyReferences(FAnimNextAssetRegistryExports::StaticStruct(), &This->CachedExports.GetValue(), InThis);
+	}
+}
+
+void UAnimNextRigVMAssetEditorData::AddEntryInternal(UAnimNextRigVMAssetEntry* InEntry)
+{
+	InsertEntryInternal(InEntry, Entries.Num());	
+}
+
+void UAnimNextRigVMAssetEditorData::InsertEntryInternal(UAnimNextRigVMAssetEntry* InEntry, int32 InsertionIndex)
+{
+	// If we are using external packages, do not persist this entry
+	if(bUsesExternalPackages)
+	{
+		Entries.Insert(InEntry, InsertionIndex);
+	}
+	else
+	{
+		InternalEntries.Insert(InEntry, InsertionIndex);
+		Entries.Insert(InEntry, InsertionIndex);
+	}
+}
+
+void UAnimNextRigVMAssetEditorData::RemoveEntryInternal(UAnimNextRigVMAssetEntry* InEntry)
+{
+	if(bUsesExternalPackages)
+	{
+		Entries.Remove(InEntry);
+	}
+	else
+	{
+		InternalEntries.Remove(InEntry);
+		Entries.Remove(InEntry);
+	}
+}
+
+#if WITH_EDITOR
+
+void UAnimNextRigVMAssetEditorData::SetUseExternalPackages(TArrayView<UAnimNextRigVMAsset*> InAssets, bool bInUseExternalPackages)
+{
+	TArray<UAnimNextRigVMAssetEditorData*> EditorDatas;
+	for(UAnimNextRigVMAsset* Asset : InAssets)
+	{
+		if(Asset == nullptr)
+		{
+			continue;
+		}
+
+		UAnimNextRigVMAssetEditorData* EditorData = UE::UAF::UncookedOnly::FUtils::GetEditorData<UAnimNextRigVMAssetEditorData>(Asset);
+		if(EditorData == nullptr)
+		{
+			continue;
+		}
+
+		if(bInUseExternalPackages != EditorData->bUsesExternalPackages)
+		{
+			EditorDatas.Add(EditorData);
+		}
+	}
+
+	if(EditorDatas.Num() == 0)
+	{
+		return;
+	}
+
+	if(bInUseExternalPackages)
+	{
+		TArray<UPackage*> PackagesToCheckOut;
+		TArray<UPackage*> PackagesToSave;
+		TArray<UPackage*> PackagesToAdd;
+
+		for(UAnimNextRigVMAssetEditorData* EditorData : EditorDatas)
+		{
+			UPackage* Package = EditorData->GetPackage();
+			PackagesToCheckOut.Add(Package);
+			PackagesToSave.Add(Package);
+		}
+
+		// Prompt the user to check out this package, allowing user to decide against this operation
+		if(!FEditorFileUtils::PromptToCheckoutPackages(false, PackagesToCheckOut))
+		{
+			return;
+		}
+
+		FScopedSlowTask SlowTask(3.0f, LOCTEXT("ConvertingAssets", "Converting Assets"));
+		SlowTask.MakeDialog();
+		SlowTask.EnterProgressFrame(1.0f, LOCTEXT("SettingPackagingStatus", "Setting Packaging Status"));
+
+		for(UAnimNextRigVMAssetEditorData* EditorData : EditorDatas)
+		{
+			EditorData->MarkPackageDirty();
+
+			// Set all internal entries to use external packages
+			TArray<UPackage*> ExternalPackages;
+			for (UAnimNextRigVMAssetEntry* Entry : EditorData->InternalEntries)
+			{
+				FExternalPackageHelper::SetPackagingMode(Entry, EditorData, bInUseExternalPackages, true, PKG_None);
+				UPackage* ExternalPackage = Entry->GetExternalPackage();
+
+				// Switch any graphs to be packaged externally
+				if(IAnimNextRigVMGraphInterface* GraphEntry = Cast<IAnimNextRigVMGraphInterface>(Entry))
+				{
+					GraphEntry->GetRigVMGraph()->SetExternalPackage(ExternalPackage);
+				}
+
+				check(ExternalPackage);
+				PackagesToAdd.Add(ExternalPackage);
+				PackagesToSave.Add(ExternalPackage);
+			}
+
+			// Clear all internal packages, switch to discovery on PostLoad rather than serialized entries
+			EditorData->InternalEntries.Empty();
+			EditorData->bUsesExternalPackages = bInUseExternalPackages;
+		}
+
+		SlowTask.EnterProgressFrame(1.0f, LOCTEXT("AddOrRevertVersionControl", "Adding/Reverting In Version Control"));
+
+		// Add (or revert delete) packages to source control
+		FPackageSourceControlHelper SCCHelper;
+		bool bAdded = SCCHelper.AddToSourceControl(PackagesToAdd);
+
+		SlowTask.EnterProgressFrame(1.0f, LOCTEXT("SavingPackages", "Saving Packages"));
+
+		// Finally save all packages, they need to be consistent on disk after this operation
+		FEditorFileUtils::FPromptForCheckoutAndSaveParams SaveParams;
+		SaveParams.bAlreadyCheckedOut = true;
+		SaveParams.bCanBeDeclined = false;
+		SaveParams.bPromptToSave = false;
+		FEditorFileUtils::PromptForCheckoutAndSave(PackagesToSave, SaveParams);
+	}
+	else
+	{
+		// Gather packages we will modify/delete
+		TArray<UPackage*> PackagesToCheckOut;
+		TArray<UPackage*> PackagesToSave;
+		TArray<UObject*> ObjectsToDelete;
+		for(UAnimNextRigVMAssetEditorData* EditorData : EditorDatas)
+		{
+			UPackage* ThisPackage = EditorData->GetPackage();
+			PackagesToCheckOut.Add(ThisPackage);
+			for (UAnimNextRigVMAssetEntry* Entry : EditorData->Entries)
+			{
+				UPackage* ExternalPackage = Entry->GetExternalPackage();
+				check(ExternalPackage);
+				ObjectsToDelete.Add(ExternalPackage);
+				PackagesToCheckOut.Add(ExternalPackage);
+			}
+		}
+
+		// Prompt the user to check out files, allowing user to decide against this operation
+		if(!FEditorFileUtils::PromptToCheckoutPackages(false, PackagesToCheckOut))
+		{
+			return;
+		}
+
+		FScopedSlowTask SlowTask(3.0f, LOCTEXT("ConvertingAssets", "Converting Assets"));
+		SlowTask.MakeDialog();
+		SlowTask.EnterProgressFrame(1.0f, LOCTEXT("SettingPackagingStatus", "Setting Packaging Status"));
+
+		for(UAnimNextRigVMAssetEditorData* EditorData : EditorDatas)
+		{
+			EditorData->MarkPackageDirty();
+
+			ensure(EditorData->InternalEntries.IsEmpty());
+			EditorData->InternalEntries.Empty();
+
+			// Set all entries to not use external packages
+			for (UAnimNextRigVMAssetEntry* Entry : EditorData->Entries)
+			{
+				FExternalPackageHelper::SetPackagingMode(Entry, EditorData, bInUseExternalPackages, true, PKG_None);
+
+				// Switch any graphs to be packaged internally
+				if(IAnimNextRigVMGraphInterface* GraphEntry = Cast<IAnimNextRigVMGraphInterface>(Entry))
+				{
+					GraphEntry->GetRigVMGraph()->SetExternalPackage(nullptr);
+				}
+			}
+
+			// Ensure we save all of our entries if we are not using external packages
+			EditorData->InternalEntries.Append(EditorData->Entries);
+			EditorData->bUsesExternalPackages = bInUseExternalPackages;
+
+			PackagesToSave.Add(EditorData->GetPackage());
+		}
+
+		SlowTask.EnterProgressFrame(1.0f, LOCTEXT("DeletingOldPackages", "Deleting Old Packages"));
+
+		// Delete the old external packages
+		ObjectTools::DeleteObjectsUnchecked(ObjectsToDelete);
+
+		SlowTask.EnterProgressFrame(1.0f, LOCTEXT("SavingPackages", "Saving Packages"));
+
+		// Finally save our packages, they needs to be consistent on disk after this operation
+		FEditorFileUtils::FPromptForCheckoutAndSaveParams SaveParams;
+		SaveParams.bAlreadyCheckedOut = true;
+		SaveParams.bCanBeDeclined = false;
+		SaveParams.bPromptToSave = false;
+		FEditorFileUtils::PromptForCheckoutAndSave(PackagesToSave, SaveParams);
+	}
+}
+
+#endif
+
+FInstancedPropertyBag UAnimNextRigVMAssetEditorData::GenerateCombinedPropertyBag(const FAnimNextGetVariableCompileContext& InCompileContext) const
+{
+	UAnimNextRigVMAsset* Asset = UE::UAF::UncookedOnly::FUtils::GetAsset<UAnimNextRigVMAsset>(this);
+
+	FInstancedPropertyBag CombinedPropertyBag;
+
+	auto AddInstancedPropertyBag = [this, &CombinedPropertyBag](const FInstancedPropertyBag& InPropertyBag)
+	{
+		if(const UPropertyBag* PropertyBag = InPropertyBag.GetPropertyBagStruct())
+		{
+			TConstArrayView<FPropertyBagPropertyDesc> VariableDescs = PropertyBag->GetPropertyDescs();
+			if(VariableDescs.Num() != 0)
+			{
+				CombinedPropertyBag.AddProperties(VariableDescs);
+			}
+		}
+	};
+
+	// NOTE: Order of external variables is important here!
+	// It must match that of the variables FUAFInstanceVariableData::Initialize
+
+	// First add internal variables
+	FInstancedPropertyBag AssetBag = UE::UAF::UncookedOnly::FUtils::MakePropertyBagForEditorData(this, InCompileContext);
+	AddInstancedPropertyBag(AssetBag);
+
+	// Next add shared variables
+	for (const UAnimNextRigVMAsset* ReferencedVariableAsset : Asset->ReferencedVariableAssets)
+	{
+		if (ReferencedVariableAsset == nullptr)
+		{
+			continue;
+		}
+
+		const UAnimNextRigVMAssetEditorData* ReferencedEditorData = UE::UAF::UncookedOnly::FUtils::GetEditorData<UAnimNextRigVMAssetEditorData>(ReferencedVariableAsset);
+		AssetBag = UE::UAF::UncookedOnly::FUtils::MakePropertyBagForEditorData(ReferencedEditorData, InCompileContext);
+		AddInstancedPropertyBag(AssetBag);
+	}
+
+	// Next add native structs
+	for (const UScriptStruct* ReferencedStruct : Asset->ReferencedVariableStructs)
+	{
+		TSharedRef<UE::UAF::FStructDataCache> StructData = UE::UAF::FStructDataCache::GetStructInfo(ReferencedStruct);
+		for (const UE::UAF::FStructDataCache::FPropertyInfo& PropertyInfo : StructData->GetProperties())
+		{
+			CombinedPropertyBag.AddProperty(PropertyInfo.Property->GetFName(), PropertyInfo.Property);
+		}
+	}
+
+	return CombinedPropertyBag;
+}
+
+void UAnimNextRigVMAssetEditorData::UpgradeDataInterfaces()
+{
+	TGuardValue<bool> DisableEditorDataNotifications(bSuspendEditorDataNotifications, true);
+	TGuardValue<bool> DisableAutoCompile(bAutoRecompileVM, false);
+
+	PRAGMA_DISABLE_DEPRECATION_WARNINGS
+	if (NativeInterfaces_DEPRECATED.Num())
+	{
+		// Iterate over all the native interfaces and add shared variables for them
+		for (const UScriptStruct* Struct : NativeInterfaces_DEPRECATED)
+		{
+			// Translate 'native interface' into the trait data it has become
+			FInstancedStruct InstancedStruct(Struct);
+			const UScriptStruct* UpgradeStruct = InstancedStruct.Get<FAnimNextNativeDataInterface>().GetUpgradeTraitStruct();
+			if(UpgradeStruct)
+			{
+				AddSharedVariablesStruct(UpgradeStruct, false, false);
+
+				// Remove all old variables
+				for (TFieldIterator<FProperty> It(UpgradeStruct); It; ++It)
+				{
+					const FName VariableName = It->GetFName();
+
+					// Try to find a variable both with and without the 'b' prefix 
+					if (It->IsA<FBoolProperty>())
+					{
+						FString BoolName = VariableName.ToString();
+						if (BoolName.RemoveFromStart(TEXT("b")))
+						{
+							if (UAnimNextVariableEntry* VariableEntry = Cast<UAnimNextVariableEntry>(FindEntry(*BoolName)))
+							{
+								RemoveEntry(VariableEntry, false, false);
+								continue;
+							}
+						}
+					}
+
+					if (UAnimNextVariableEntry* VariableEntry = Cast<UAnimNextVariableEntry>(FindEntry(VariableName)))
+					{
+						RemoveEntry(VariableEntry, false, false);
+						continue;
+					}
+				}
+			}
+		}
+
+		NativeInterfaces_DEPRECATED.Empty();
+	}
+	PRAGMA_ENABLE_DEPRECATION_WARNINGS
+
+	{
+		// Recompile to incorporate any new native interface variables (will create errors, so suppress them)
+		TGuardValue<bool> DisableCompilerReports(bSuspendCompilerReports, true);
+		RecompileVM();
+	}
+
+	// Now replace variable nodes with scoped ones in all graphs, if possible
+	for (URigVMGraph* Graph : GetRigVMClient()->GetAllModels(true, true))
+	{
+		UAnimNextControllerBase* Controller = CastChecked<UAnimNextControllerBase>(GetController(Graph));
+		TArray<URigVMNode*> Nodes = Graph->GetNodes();
+		for (URigVMNode* Node : Nodes)
+		{
+			if (URigVMVariableNode* VariableNode = Cast<URigVMVariableNode>(Node))
+			{
+				const FName OriginalVariableName = VariableNode->GetVariableName();
+				FName NewVariableName = OriginalVariableName;
+				TArray<FName> VariableNames;
+				VariableNames.Add(OriginalVariableName);
+				if (VariableNode->GetCPPType() == TEXT("bool"))
+				{
+					TStringBuilder<128> StringBuilder;
+					StringBuilder.Append(TEXT("b"));
+					OriginalVariableName.AppendString(StringBuilder);
+					VariableNames.Add(StringBuilder.ToString());
+				}
+
+				// See if the variable exists in our shared variable entries
+				for(UAnimNextRigVMAssetEntry* Entry : Entries)
+				{
+					const UObject* AssetOrStruct = nullptr;
+					if (UAnimNextSharedVariablesEntry* SharedVariablesEntry = Cast<UAnimNextSharedVariablesEntry>(Entry))
+					{
+						for (FName VariableName : VariableNames)
+						{
+							switch (SharedVariablesEntry->GetType())
+							{
+							case EAnimNextSharedVariablesType::Asset:
+								{
+									if (const UAnimNextRigVMAsset* Asset = SharedVariablesEntry->GetAsset())
+									{
+										const FInstancedPropertyBag& Variables = Asset->GetVariableDefaults();
+										if (Variables.FindPropertyDescByName(VariableName) != nullptr)
+										{
+											AssetOrStruct = Asset;
+											NewVariableName = VariableName;
+										}
+									}
+									break;
+								}
+							case EAnimNextSharedVariablesType::Struct:
+								{
+									if (const UScriptStruct* Struct = SharedVariablesEntry->GetStruct())
+									{
+										if (Struct->FindPropertyByName(VariableName) != nullptr)
+										{
+											AssetOrStruct = Struct;
+											NewVariableName = VariableName;
+										}
+									}
+									break;
+								}
+							}
+
+							if (AssetOrStruct)
+							{
+								break;
+							}
+						}
+					}
+
+					if (AssetOrStruct)
+					{
+						Controller->ReplaceVariableNodeWithSharedVariableNode(VariableNode, NewVariableName, AssetOrStruct, false, false);
+						break;
+					}
+				}
+			}
+		}
+	}
+}
+
+
+#undef LOCTEXT_NAMESPACE
