UAF: Moved/renamed AnimNext and AnimNextAnimGraph plugins

diff --git a/Engine/Plugins/Experimental/UAF/UAF/Source/UAFUncookedOnly/Private/Entries/AnimNextRigVMAssetEntry.cpp b/Engine/Plugins/Experimental/UAF/UAF/Source/UAFUncookedOnly/Private/Entries/AnimNextRigVMAssetEntry.cpp
new file mode 100644
--- /dev/null
+++ b/Engine/Plugins/Experimental/UAF/UAF/Source/UAFUncookedOnly/Private/Entries/AnimNextRigVMAssetEntry.cpp
@@ -0,0 +1,58 @@
+// Copyright Epic Games, Inc. All Rights Reserved.
+
+#include "Entries/AnimNextRigVMAssetEntry.h"
+#include "AnimNextRigVMAssetEditorData.h"
+#include "Misc/TransactionObjectEvent.h"
+
+#define LOCTEXT_NAMESPACE "AnimNextRigVMAssetEntry"
+
+void UAnimNextRigVMAssetEntry::Initialize(UAnimNextRigVMAssetEditorData* InEditorData)
+{
+	InEditorData->RigVMGraphModifiedEvent.RemoveAll(this);
+	InEditorData->RigVMGraphModifiedEvent.AddUObject(this, &UAnimNextRigVMAssetEntry::HandleRigVMGraphModifiedEvent);
+}
+
+bool UAnimNextRigVMAssetEntry::IsAsset() const
+{
+	// Entries are considered assets to allow using the asset logic for save dialogs, etc.
+	// Also, they return true even if pending kill, in order to show up as deleted in these dialogs.
+	return IsPackageExternal() && !GetPackage()->HasAnyFlags(RF_Transient) && !HasAnyFlags(RF_Transient | RF_ClassDefaultObject);
+}
+
+#if WITH_EDITOR
+
+void UAnimNextRigVMAssetEntry::PostEditChangeProperty(FPropertyChangedEvent& PropertyChangedEvent)
+{
+	Super::PostEditChangeProperty(PropertyChangedEvent);
+
+	BroadcastModified(EAnimNextEditorDataNotifType::PropertyChanged);
+}
+
+void UAnimNextRigVMAssetEntry::PostTransacted(const FTransactionObjectEvent& TransactionEvent)
+{
+	Super::PostTransacted(TransactionEvent);
+
+	if (TransactionEvent.GetEventType() == ETransactionObjectEventType::UndoRedo)
+	{
+		BroadcastModified(EAnimNextEditorDataNotifType::UndoRedo);
+	}
+}
+
+void UAnimNextRigVMAssetEntry::PostEditUndo()
+{
+	// Note: Dont call super as that will default to a generic property changed event
+
+	BroadcastModified(EAnimNextEditorDataNotifType::UndoRedo);
+}
+
+#endif
+
+void UAnimNextRigVMAssetEntry::BroadcastModified(EAnimNextEditorDataNotifType InType)
+{
+	if(UAnimNextRigVMAssetEditorData* EditorData = Cast<UAnimNextRigVMAssetEditorData>(GetOuter()))
+	{
+		EditorData->BroadcastModified(InType, this);
+	}
+}
+
+#undef LOCTEXT_NAMESPACE
\ No newline at end of file
diff --git a/Engine/Plugins/Experimental/UAF/UAF/Source/UAFUncookedOnly/Private/Entries/AnimNextSharedVariablesEntry.cpp b/Engine/Plugins/Experimental/UAF/UAF/Source/UAFUncookedOnly/Private/Entries/AnimNextSharedVariablesEntry.cpp
new file mode 100644
--- /dev/null
+++ b/Engine/Plugins/Experimental/UAF/UAF/Source/UAFUncookedOnly/Private/Entries/AnimNextSharedVariablesEntry.cpp
@@ -0,0 +1,153 @@
+﻿// Copyright Epic Games, Inc. All Rights Reserved.
+
+#include "Entries/AnimNextSharedVariablesEntry.h"
+
+#if WITH_LIVE_CODING
+#include "ILiveCodingModule.h"
+#endif
+#include "UncookedOnlyUtils.h"
+#include "Variables/AnimNextSharedVariables_EditorData.h"
+#include "Entries/AnimNextVariableEntry.h"
+#include "Logging/StructuredLog.h"
+#include "Modules/ModuleManager.h"
+#include "Param/ParamType.h"
+#include "Variables/AnimNextSharedVariables.h"
+
+#define LOCTEXT_NAMESPACE "AnimNextSharedVariablesEntry"
+
+void UAnimNextSharedVariablesEntry::Initialize(UAnimNextRigVMAssetEditorData* InEditorData)
+{
+	Super::Initialize(InEditorData);
+
+	if(Asset)
+	{
+		UAnimNextSharedVariables_EditorData* EditorData = UE::UAF::UncookedOnly::FUtils::GetEditorData<UAnimNextSharedVariables_EditorData>(Asset.Get());
+		EditorData->ModifiedDelegate.AddUObject(this, &UAnimNextSharedVariablesEntry::HandleAssetModified);
+	}
+	else if (Struct)
+	{
+#if WITH_LIVE_CODING
+		if (ILiveCodingModule* LiveCoding = FModuleManager::LoadModulePtr<ILiveCodingModule>(LIVE_CODING_MODULE_NAME))
+		{
+			LiveCoding->GetOnPatchCompleteDelegate().AddUObject(this, &UAnimNextSharedVariablesEntry::HandlePatchComplete);
+		}
+#endif
+	}
+}
+
+FName UAnimNextSharedVariablesEntry::GetEntryName() const
+{
+	switch (Type)
+	{
+	case EAnimNextSharedVariablesType::Asset:
+		return Asset ? Asset->GetFName() : NAME_None;
+	case EAnimNextSharedVariablesType::Struct:
+		return Struct ? Struct->GetFName() : NAME_None;
+	}
+	
+	return NAME_None;
+}
+
+FText UAnimNextSharedVariablesEntry::GetDisplayName() const
+{
+	switch (Type)
+	{
+	case EAnimNextSharedVariablesType::Asset:
+		return Asset ? FText::FromName(Asset->GetFName()) : LOCTEXT("InvalidSharedVariables", "Invalid Shared Variables");
+	case EAnimNextSharedVariablesType::Struct:
+		return Struct ? Struct->GetDisplayNameText() : LOCTEXT("InvalidSharedVariables", "Invalid Shared Variables");
+	}
+
+	return FText::GetEmpty();
+}
+
+FText UAnimNextSharedVariablesEntry::GetDisplayNameTooltip() const
+{
+	switch (Type)
+	{
+	case EAnimNextSharedVariablesType::Asset:
+		return Asset ? FText::FromName(Asset->GetFName()) : LOCTEXT("InvalidSharedVariablesTooltip", "Invalid or deleted Shared Variables");
+	case EAnimNextSharedVariablesType::Struct:
+		return Struct ? FText::FromString(Struct->GetPathName()) : LOCTEXT("InvalidSharedVariables", "Invalid Shared Variables");
+	}
+
+	return FText::GetEmpty();
+}
+
+void UAnimNextSharedVariablesEntry::SetAsset(const UAnimNextSharedVariables* InAsset, bool bSetupUndoRedo)
+{
+	check(InAsset != nullptr);
+
+	if(bSetupUndoRedo)
+	{
+		Modify();
+	}
+
+	Type = EAnimNextSharedVariablesType::Asset;
+	Asset = InAsset;
+	ObjectPath = FSoftObjectPath(Asset);
+	Struct = nullptr;
+}
+
+const UAnimNextSharedVariables* UAnimNextSharedVariablesEntry::GetAsset() const
+{
+	return Type == EAnimNextSharedVariablesType::Asset ? Asset : nullptr;
+}
+
+FSoftObjectPath UAnimNextSharedVariablesEntry::GetObjectPath() const
+{
+	return ObjectPath;
+}
+
+void UAnimNextSharedVariablesEntry::SetStruct(const UScriptStruct* InStruct, bool bSetupUndoRedo)
+{
+	check(InStruct != nullptr);
+
+	if(bSetupUndoRedo)
+	{
+		Modify();
+	}
+
+	Type = EAnimNextSharedVariablesType::Struct;
+	Struct = InStruct;
+	ObjectPath = FSoftObjectPath(Struct);
+	Asset = nullptr;
+}
+
+const UScriptStruct* UAnimNextSharedVariablesEntry::GetStruct() const
+{
+	return Type == EAnimNextSharedVariablesType::Struct ? Struct : nullptr;
+}
+
+void UAnimNextSharedVariablesEntry::HandleAssetModified(UAnimNextRigVMAssetEditorData* InEditorData, EAnimNextEditorDataNotifType InType, UObject* InSubject)
+{
+	switch(InType)
+	{
+	case EAnimNextEditorDataNotifType::UndoRedo:
+	case EAnimNextEditorDataNotifType::EntryAdded:
+	case EAnimNextEditorDataNotifType::EntryRemoved:
+	case EAnimNextEditorDataNotifType::EntryRenamed:
+	case EAnimNextEditorDataNotifType::EntryAccessSpecifierChanged:
+	case EAnimNextEditorDataNotifType::VariableTypeChanged:
+	case EAnimNextEditorDataNotifType::VariableDefaultValueChanged:
+		if(UAnimNextRigVMAssetEditorData* EditorData = GetTypedOuter<UAnimNextRigVMAssetEditorData>())
+		{
+			EditorData->RequestAutoVMRecompilation();
+		}
+		break;
+	default:
+		break;
+	}
+}
+
+#if WITH_LIVE_CODING
+void UAnimNextSharedVariablesEntry::HandlePatchComplete()
+{
+	if(UAnimNextRigVMAssetEditorData* EditorData = GetTypedOuter<UAnimNextRigVMAssetEditorData>())
+	{
+		EditorData->RequestAutoVMRecompilation();
+	}
+}
+#endif
+
+#undef LOCTEXT_NAMESPACE
\ No newline at end of file
diff --git a/Engine/Plugins/Experimental/UAF/UAF/Source/UAFUncookedOnly/Private/Module/AnimNextModule_EditorData.cpp b/Engine/Plugins/Experimental/UAF/UAF/Source/UAFUncookedOnly/Private/Module/AnimNextModule_EditorData.cpp
new file mode 100644
--- /dev/null
+++ b/Engine/Plugins/Experimental/UAF/UAF/Source/UAFUncookedOnly/Private/Module/AnimNextModule_EditorData.cpp
@@ -0,0 +1,104 @@
+﻿// Copyright Epic Games, Inc. All Rights Reserved.
+
+#include "Module/AnimNextModule_EditorData.h"
+
+#include "Compilation/AnimNextGetGraphCompileContext.h"
+#include "Compilation/AnimNextProcessGraphCompileContext.h"
+#include "ExternalPackageHelper.h"
+#include "UncookedOnlyUtils.h"
+#include "Module/AnimNextModule.h"
+#include "Entries/AnimNextEventGraphEntry.h"
+#include "Entries/AnimNextSharedVariablesEntry.h"
+#include "Entries/AnimNextVariableEntry.h"
+#include "String/ParseTokens.h"
+#include "UObject/AssetRegistryTagsContext.h"
+#include "UObject/LinkerLoad.h"
+#include "Variables/AnimNextUniversalObjectLocatorBindingData.h"
+
+TConstArrayView<TSubclassOf<UAnimNextRigVMAssetEntry>> UAnimNextModule_EditorData::GetEntryClasses() const
+{
+	static const TSubclassOf<UAnimNextRigVMAssetEntry> Classes[] =
+	{
+		UAnimNextEventGraphEntry::StaticClass(),
+		UAnimNextVariableEntry::StaticClass(),
+		UAnimNextSharedVariablesEntry::StaticClass(),
+	};
+	
+	return Classes;
+}
+
+void UAnimNextModule_EditorData::OnPreCompileAsset(FRigVMCompileSettings& InSettings)
+{
+	using namespace UE::UAF::UncookedOnly;
+
+	UAnimNextModule* Module = FUtils::GetAsset<UAnimNextModule>(this);
+
+	Module->RequiredComponents.Empty();
+	Module->Dependencies.Empty();
+}
+
+void UAnimNextModule_EditorData::OnPreCompileGetProgrammaticGraphs(const FRigVMCompileSettings& InSettings, FAnimNextGetGraphCompileContext& OutCompileContext)
+{
+	using namespace UE::UAF::UncookedOnly;
+
+	FUtils::CompileVariableBindings(InSettings, FUtils::GetAsset<UAnimNextModule>(this), OutCompileContext.GetMutableProgrammaticGraphs());
+}
+
+void UAnimNextModule_EditorData::OnPreCompileProcessGraphs(const FRigVMCompileSettings& InSettings, FAnimNextProcessGraphCompileContext& OutCompileContext)
+{
+	using namespace UE::UAF::UncookedOnly;
+
+	UAnimNextModule* Module = FUtils::GetAsset<UAnimNextModule>(this);
+
+	// Gather any required components from node metadata
+	for(URigVMGraph* Graph : OutCompileContext.GetMutableAllGraphs())
+	{
+		for(URigVMNode* Node : Graph->GetNodes())
+		{
+			URigVMTemplateNode* TemplateNode = Cast<URigVMTemplateNode>(Node);
+			if(TemplateNode == nullptr)
+			{
+				continue;
+			}
+
+			UScriptStruct* Struct = TemplateNode->GetScriptStruct();
+			if(Struct == nullptr)
+			{
+				continue;
+			}
+
+			static FName Metadata_RequiredComponents("RequiredComponents");
+			FString ComponentsString = Struct->GetMetaData(Metadata_RequiredComponents);
+			if(ComponentsString.Len() == 0)
+			{
+				continue;
+			}
+
+			UE::String::ParseTokens(ComponentsString, TEXT(','), [Module](FStringView InToken)
+			{
+				FString StructName(InToken);
+				UScriptStruct* ComponentStruct = FindFirstObject<UScriptStruct>(*StructName);
+				if(ComponentStruct)
+				{
+					Module->RequiredComponents.Add(ComponentStruct);
+				}
+			});
+		}
+	}
+
+	// Copy dependencies
+	Module->Dependencies = Dependencies;
+}
+
+void UAnimNextModule_EditorData::CustomizeNewAssetEntry(UAnimNextRigVMAssetEntry* InNewEntry) const
+{
+	Super::CustomizeNewAssetEntry(InNewEntry);
+	
+	UAnimNextVariableEntry* VariableEntry = Cast<UAnimNextVariableEntry>(InNewEntry);
+	if(VariableEntry == nullptr)
+	{
+		return;
+	}
+
+	VariableEntry->SetBindingType(FAnimNextUniversalObjectLocatorBindingData::StaticStruct(), false);
+}
\ No newline at end of file
diff --git a/Engine/Plugins/Experimental/UAF/UAF/Source/UAFUncookedOnly/Private/UncookedOnlyUtils.cpp b/Engine/Plugins/Experimental/UAF/UAF/Source/UAFUncookedOnly/Private/UncookedOnlyUtils.cpp
new file mode 100644
--- /dev/null
+++ b/Engine/Plugins/Experimental/UAF/UAF/Source/UAFUncookedOnly/Private/UncookedOnlyUtils.cpp
@@ -0,0 +1,1814 @@
+﻿// Copyright Epic Games, Inc. All Rights Reserved.
+
+#include "UncookedOnlyUtils.h"
+
+#include "K2Node_CallFunction.h"
+#include "AssetRegistry/AssetRegistryModule.h"
+#include "AssetRegistry/IAssetRegistry.h"
+#include "Compilation/AnimNextGetVariableCompileContext.h"
+#include "Module/AnimNextModule.h"
+#include "Module/AnimNextModule_EditorData.h"
+#include "Serialization/MemoryReader.h"
+#include "RigVMCore/RigVM.h"
+#include "AnimNextRigVMAsset.h"
+#include "AnimNextRigVMAssetEditorData.h"
+#include "Entries/AnimNextRigVMAssetEntry.h"
+#include "AnimNextUncookedOnlyModule.h"
+#include "IAnimNextRigVMExportInterface.h"
+#include "Logging/StructuredLog.h"
+#include "AnimNextAssetWorkspaceAssetUserData.h"
+#include "IWorkspaceEditor.h"
+#include "IWorkspaceEditorModule.h"
+#include "Misc/EnumerateRange.h"
+#include "Module/RigUnit_AnimNextModuleEvents.h"
+#include "Variables/AnimNextProgrammaticVariable.h"
+#include "Variables/IVariableBindingType.h"
+#include "Variables/RigUnit_CopyModuleProxyVariables.h"
+#include "WorkspaceAssetRegistryInfo.h"
+#include "Entries/AnimNextSharedVariablesEntry.h"
+#include "Variables/AnimNextSharedVariables.h"
+#include "Variables/AnimNextSharedVariables_EditorData.h"
+#include "Logging/MessageLog.h"
+#include "RigVMFunctions/Execution/RigVMFunction_UserDefinedEvent.h"
+#include "UObject/AssetRegistryTagsContext.h"
+#include "Variables/AnimNextSoftVariableReference.h"
+#include "AnimNextExports.h"
+#include "AnimNextScopedCompileJob.h"
+#include "AnimNextSharedVariableNode.h"
+#include "Styling/SlateIconFinder.h"
+
+#define LOCTEXT_NAMESPACE "AnimNextUncookedOnlyUtils"
+
+namespace UE::UAF::UncookedOnly
+{
+
+TAutoConsoleVariable<bool> CVarDumpProgrammaticGraphs(
+	TEXT("UAF.DumpProgrammaticGraphs"),
+	false,
+	TEXT("When true the transient programmatic graphs will be automatically opened for any that are generated."));
+
+void FUtils::RecreateVM(UAnimNextRigVMAsset* InAsset)
+{
+	if (InAsset->VM == nullptr)
+	{
+		InAsset->VM = NewObject<URigVM>(InAsset, TEXT("VM"), RF_NoFlags);
+	}
+	InAsset->VM->Reset(InAsset->ExtendedExecuteContext);
+	InAsset->RigVM = InAsset->VM; // Local serialization
+}
+
+FInstancedPropertyBag FUtils::MakePropertyBagForEditorData(const UAnimNextRigVMAssetEditorData* InEditorData, const FAnimNextGetVariableCompileContext& InCompileContext)
+{
+	FMessageLog Log("AnimNextCompilerResults");
+	
+	struct FStructEntryInfo
+	{
+		FName Name;
+		FAnimNextParamType Type;
+		EAnimNextExportAccessSpecifier AccessSpecifier = EAnimNextExportAccessSpecifier::Private;
+		TConstArrayView<uint8> Value;
+		EPropertyFlags PropertyFlags = CPF_NativeAccessSpecifierPrivate | CPF_Edit;
+	};
+
+	// Gather all variables in this asset.
+	TMap<FName, int32> EntryInfoIndexMap;
+	TArray<FStructEntryInfo> StructEntryInfos;
+	const TArray<FAnimNextProgrammaticVariable>& ProgrammaticVariables = InCompileContext.GetProgrammaticVariables();
+	StructEntryInfos.Reserve(InEditorData->Entries.Num() + ProgrammaticVariables.Num());
+	int32 NumPublicVariables = 0;
+
+	auto AddVariable = [&Log, &NumPublicVariables, &StructEntryInfos, &EntryInfoIndexMap](FName InName, const FAnimNextParamType& InType, EAnimNextExportAccessSpecifier InAccess, TConstArrayView<uint8> InValue, EPropertyFlags InFlags)
+	{
+		if(!InType.IsValid())
+		{
+			Log.Error(FText::Format(LOCTEXT("InvalidVariableTypeFound", "Variable '{0}' with invalid type found"), FText::FromName(InName)));
+			return;
+		}
+
+		// Check for conflicts
+		if(EntryInfoIndexMap.Contains(InName))
+		{
+			Log.Error(FText::Format(LOCTEXT("DuplicateVariableFound", "Variable '{0}' with duplicate name found"), FText::FromName(InName)));
+			return;
+		}
+
+		if(InAccess == EAnimNextExportAccessSpecifier::Public)
+		{
+			NumPublicVariables++;
+		}
+
+		int32 Index = StructEntryInfos.Add(
+			{
+				InName,
+				FAnimNextParamType(InType.GetValueType(), InType.GetContainerType(), InType.GetValueTypeObject()),
+				InAccess,
+				InValue,
+				InFlags
+			});
+
+		EntryInfoIndexMap.Add(InName, Index);
+	};
+
+	for(UAnimNextRigVMAssetEntry* Entry : InEditorData->Entries)
+	{
+		if(const UAnimNextVariableEntry* VariableEntry = Cast<UAnimNextVariableEntry>(Entry))
+		{
+			EPropertyFlags Flags = CPF_Edit;
+			if (VariableEntry->GetExportAccessSpecifier() == EAnimNextExportAccessSpecifier::Private)
+			{
+				Flags |= CPF_NativeAccessSpecifierPrivate;
+			}
+			else
+			{
+				Flags |= CPF_NativeAccessSpecifierPublic;
+			}
+			
+			AddVariable(
+				VariableEntry->GetExportName(),
+				VariableEntry->GetType(),
+				VariableEntry->GetExportAccessSpecifier(),
+				TConstArrayView<uint8>(VariableEntry->GetValuePtr(), VariableEntry->GetType().GetSize()),
+				Flags);
+		}
+	}
+
+	for (const FAnimNextProgrammaticVariable& Variable : ProgrammaticVariables)
+	{
+		AddVariable(
+			Variable.Name,
+			Variable.Type,
+			EAnimNextExportAccessSpecifier::Private,
+			TConstArrayView<uint8>(Variable.GetValuePtr(), Variable.Type.GetSize()),
+			CPF_AdvancedDisplay | CPF_NativeAccessSpecifierPrivate);
+	}
+
+	// Sort by size, largest first, for better packing
+	StructEntryInfos.Sort([](const FStructEntryInfo& InLHS, const FStructEntryInfo& InRHS)
+	{
+		if(InLHS.Type.GetSize() != InRHS.Type.GetSize())
+		{
+			return InLHS.Type.GetSize() > InRHS.Type.GetSize();
+		}
+		else
+		{
+			return InLHS.Name.LexicalLess(InRHS.Name);
+		}
+	});
+
+	FInstancedPropertyBag VariableDefaults;
+
+	if(StructEntryInfos.Num() > 0)
+	{
+		// Build PropertyDescs and values to batch-create the property bag
+		TArray<FPropertyBagPropertyDesc> PropertyDescs;
+		PropertyDescs.Reserve(StructEntryInfos.Num());
+		TArray<TConstArrayView<uint8>> Values;
+		Values.Reserve(StructEntryInfos.Num());
+
+		for (const FStructEntryInfo& StructEntryInfo : StructEntryInfos)
+		{
+			PropertyDescs.Emplace(StructEntryInfo.Name, StructEntryInfo.Type.ContainerType, StructEntryInfo.Type.ValueType, StructEntryInfo.Type.ValueTypeObject, StructEntryInfo.PropertyFlags);
+			Values.Add(StructEntryInfo.Value);
+		}
+
+		// Create new property bags and migrate
+		EPropertyBagResult Result = VariableDefaults.ReplaceAllPropertiesAndValues(PropertyDescs, Values);
+		check(Result == EPropertyBagResult::Success);
+	}
+
+	return VariableDefaults;
+}
+
+void FUtils::CompileVariables(const FRigVMCompileSettings& InSettings, UAnimNextRigVMAsset* InAsset, FAnimNextGetVariableCompileContext& OutCompileContext)
+{
+	check(InAsset);
+
+	UAnimNextRigVMAssetEditorData* EditorData = GetEditorData<UAnimNextRigVMAssetEditorData>(InAsset);
+
+	// Gather programmatic variables regenerated each compile
+	EditorData->OnPreCompileGetProgrammaticVariables(InSettings, OutCompileContext);
+
+	// Generate the internal property bag
+	InAsset->VariableDefaults = MakePropertyBagForEditorData(EditorData, OutCompileContext);
+
+	TSet<const UAnimNextRigVMAsset*> SharedVariableAssets;
+	TSet<const UScriptStruct*> SharedVariableStructs;
+
+	FMessageLog Log("AnimNextCompilerResults");
+
+	for(UAnimNextRigVMAssetEntry* Entry : EditorData->Entries)
+	{
+		if(const UAnimNextSharedVariablesEntry* SharedVariablesEntry = Cast<UAnimNextSharedVariablesEntry>(Entry))
+		{
+			switch (SharedVariablesEntry->GetType())
+			{
+			case EAnimNextSharedVariablesType::Asset:
+				{
+					const UAnimNextSharedVariables* SharedVariablesAsset = SharedVariablesEntry->GetAsset();
+					if (SharedVariablesAsset == nullptr)
+					{
+						Log.Error(FText::Format(LOCTEXT("InvalidSharedVariableAssetFound", "Invalid shared variables asset found ({0})"), FText::FromString(SharedVariablesEntry->GetObjectPath().ToString())));
+						continue;
+					}
+
+					if (SharedVariableAssets.Contains(SharedVariablesAsset))
+					{
+						Log.Error(FText::Format(LOCTEXT("DuplicateSharedVariableAssetFound", "Duplicate shared variables asset found ({0})"), FText::FromString(SharedVariablesEntry->GetObjectPath().ToString())));
+						continue;
+					}
+
+					SharedVariableAssets.Add(SharedVariablesAsset);
+					break;
+				}
+			case EAnimNextSharedVariablesType::Struct:
+				{
+					const UScriptStruct* SharedVariablesStruct = SharedVariablesEntry->GetStruct();
+					if (SharedVariablesStruct == nullptr)
+					{
+						Log.Error(FText::Format(LOCTEXT("InvalidSharedVariableStructFound", "Invalid shared variables struct found ({0})"), FText::FromString(SharedVariablesEntry->GetObjectPath().ToString())));
+						continue;
+					}
+
+					if (SharedVariableStructs.Contains(SharedVariablesStruct))
+					{
+						Log.Error(FText::Format(LOCTEXT("DuplicateSharedVariableStructFound", "Duplicate shared variables struct found ({0})"), FText::FromString(SharedVariablesEntry->GetObjectPath().ToString())));
+						continue;
+					}
+
+					SharedVariableStructs.Add(SharedVariablesStruct);
+					break;
+				}
+			}
+		}
+	}
+	
+	// Set shared variable assets/structs before we setup variables as GetExternalVariablesImpl relies on these
+	if (SharedVariableAssets.Num() > 0)
+	{
+		InAsset->ReferencedVariableAssets = SharedVariableAssets.Array();
+	}
+	else
+	{
+		InAsset->ReferencedVariableAssets.Empty();
+	}
+
+	if (SharedVariableStructs.Num() > 0)
+	{
+		InAsset->ReferencedVariableStructs = SharedVariableStructs.Array();
+	}
+	else
+	{
+		InAsset->ReferencedVariableStructs.Empty();
+	}
+
+	// Now rebuild combined the combined property bag we used for stable properties 
+	InAsset->CombinedPropertyBag = EditorData->GenerateCombinedPropertyBag(OutCompileContext);
+
+	// Rebuild external variables
+	if (InAsset->CombinedPropertyBag.GetNumPropertiesInBag() > 0)
+	{
+		InAsset->VM->SetExternalVariableDefs(InAsset->GetExternalVariablesImpl(false));
+	}
+	else
+	{
+		InAsset->VM->ClearExternalVariables(InAsset->ExtendedExecuteContext);
+	}
+
+	// Validate injection site
+	InAsset->DefaultInjectionSite.Reset();
+
+	if(!EditorData->DefaultInjectionSiteReference.IsNone())
+	{
+		if (!InAsset->DefaultInjectionSite.IsValid())
+		{
+			Log.Error(FText::Format(LOCTEXT("MissingDefaultInjectionSiteError", "Could not find default injection site: {0}"), FText::FromName(EditorData->DefaultInjectionSiteReference.GetName())));
+		}
+		else
+		{
+			// Warn about using the deprecated name-based path
+			if (EditorData->DefaultInjectionSiteReference.GetObject() == nullptr)
+			{
+				Log.Warning(FText::Format(LOCTEXT("DeprecatedNamedInjectionSiteWarning", "Default injection site '{0}' uses a name-based reference. Please select a full reference to an asset's variable."), FText::FromName(EditorData->DefaultInjectionSiteReference.GetName())));
+			}
+
+			InAsset->DefaultInjectionSite = EditorData->DefaultInjectionSiteReference; 
+		}
+	}
+}
+
+void FUtils::CompileVariableBindings(const FRigVMCompileSettings& InSettings, UAnimNextRigVMAsset* InAsset, TArray<URigVMGraph*>& OutGraphs)
+{
+	CompileVariableBindingsInternal(InSettings, InAsset, OutGraphs, true);
+	CompileVariableBindingsInternal(InSettings, InAsset,  OutGraphs, false);
+}
+
+void FUtils::CompileVariableBindingsInternal(const FRigVMCompileSettings& InSettings, UAnimNextRigVMAsset* InAsset, TArray<URigVMGraph*>& OutGraphs, bool bInThreadSafe)
+{
+	check(InAsset);
+
+	FModule& Module = FModuleManager::LoadModuleChecked<FModule>("UAFUncookedOnly");
+	UAnimNextRigVMAssetEditorData* EditorData = GetEditorData(InAsset);
+	TMap<IVariableBindingType*, TArray<IVariableBindingType::FBindingGraphInput>> BindingGroups;
+
+	for(const UAnimNextRigVMAssetEntry* Entry : EditorData->Entries)
+	{
+		const IAnimNextRigVMVariableInterface* Variable = Cast<IAnimNextRigVMVariableInterface>(Entry);
+		if(Variable == nullptr)
+		{
+			continue;
+		}
+
+		TConstStructView<FAnimNextVariableBindingData> Binding = Variable->GetBinding();
+		if(!Binding.IsValid() || !Binding.Get<FAnimNextVariableBindingData>().IsValid())
+		{
+			continue;
+		}
+
+		if(Binding.Get<FAnimNextVariableBindingData>().IsThreadSafe() != bInThreadSafe)
+		{
+			continue;
+		}
+
+		TSharedPtr<IVariableBindingType> BindingType = Module.FindVariableBindingType(Binding.GetScriptStruct());
+		if(!BindingType.IsValid())
+		{
+			continue;
+		}
+
+		TArray<IVariableBindingType::FBindingGraphInput>& Group = BindingGroups.FindOrAdd(BindingType.Get());
+		FRigVMTemplateArgumentType RigVMArg = Variable->GetType().ToRigVMTemplateArgument();
+		Group.Add({ Variable->GetVariableName(), RigVMArg.GetBaseCPPType(), RigVMArg.CPPTypeObject, Binding});
+	}
+
+	const bool bHasBindings = BindingGroups.Num() > 0;
+	const bool bHasPublicVariablesToCopy = EditorData->IsA<UAnimNextModule_EditorData>() && EditorData->HasPublicVariables() && bInThreadSafe;
+	if(!bHasBindings && !bHasPublicVariablesToCopy)
+	{
+		// Nothing to do here
+		return;
+	}
+
+	URigVMGraph* BindingGraph = NewObject<URigVMGraph>(EditorData, NAME_None, RF_Transient);
+
+	FRigVMClient* VMClient = EditorData->GetRigVMClient();
+	URigVMController* Controller = VMClient->GetOrCreateController(BindingGraph);
+	UScriptStruct* BindingsNodeType = bInThreadSafe ? FRigUnit_AnimNextExecuteBindings_WT::StaticStruct() : FRigUnit_AnimNextExecuteBindings_GT::StaticStruct();
+	URigVMNode* ExecuteBindingsNode = Controller->AddUnitNode(BindingsNodeType, FRigVMStruct::ExecuteName, FVector2D::ZeroVector, FString(), false);
+	if(ExecuteBindingsNode == nullptr)
+	{
+		InSettings.ReportError(TEXT("Could not spawn Execute Bindings node"));
+		return;
+	}
+	URigVMPin* ExecuteBindingsExecPin = ExecuteBindingsNode->FindPin(FRigVMStruct::ExecuteContextName.ToString());
+	if(ExecuteBindingsExecPin == nullptr)
+	{
+		InSettings.ReportError(TEXT("Could not find execute pin on Execute Bindings node"));
+		return;
+	}
+	URigVMPin* ExecPin = ExecuteBindingsExecPin;
+
+	// Copy public vars in the WT event
+	if(bHasPublicVariablesToCopy && bInThreadSafe)
+	{
+		URigVMNode* CopyProxyVariablesNode = Controller->AddUnitNode(FRigUnit_CopyModuleProxyVariables::StaticStruct(), FRigVMStruct::ExecuteName, FVector2D(200, 0.0f), FString(), false);
+		if(CopyProxyVariablesNode == nullptr)
+		{
+			InSettings.ReportError(TEXT("Could not spawn Copy Module Proxy Variables node"));
+			return;
+		}
+		URigVMPin* CopyProxyVariablesExecPin = CopyProxyVariablesNode->FindPin(FRigVMStruct::ExecuteContextName.ToString());
+		if(ExecPin == nullptr)
+		{
+			InSettings.ReportError(TEXT("Could not find execute pin on Copy Module Proxy Variables node"));
+			return;
+		}
+		bool bLinkAdded = Controller->AddLink(ExecuteBindingsExecPin, CopyProxyVariablesExecPin, false);
+		if(!bLinkAdded)
+		{
+			InSettings.ReportError(TEXT("Could not link Copy Module Proxy Variables node"));
+			return;
+		}
+		ExecPin = CopyProxyVariablesExecPin;
+	}
+
+	IVariableBindingType::FBindingGraphFragmentArgs Args;
+	Args.Event = BindingsNodeType;
+	Args.Controller = Controller;
+	Args.BindingGraph = BindingGraph;
+	Args.ExecTail = ExecPin;
+	Args.bThreadSafe = bInThreadSafe;
+
+	FVector2D Location(0.0f, 0.0f);
+	for(const TPair<IVariableBindingType*, TArray<IVariableBindingType::FBindingGraphInput>>& BindingGroupPair : BindingGroups)
+	{
+		Args.Inputs = BindingGroupPair.Value;
+		BindingGroupPair.Key->BuildBindingGraphFragment(InSettings, Args, ExecPin, Location);
+	}
+
+	OutGraphs.Add(BindingGraph);
+}
+
+UAnimNextRigVMAsset* FUtils::GetAsset(UAnimNextRigVMAssetEditorData* InEditorData)
+{
+	check(InEditorData);
+	return CastChecked<UAnimNextRigVMAsset>(InEditorData->GetOuter());
+}
+
+UAnimNextRigVMAssetEditorData* FUtils::GetEditorData(UAnimNextRigVMAsset* InAsset)
+{
+	check(InAsset);
+	return CastChecked<UAnimNextRigVMAssetEditorData>(InAsset->EditorData);
+}
+
+FAnimNextParamType FUtils::GetParamTypeFromPinType(const FEdGraphPinType& InPinType)
+{
+	FAnimNextParamType::EValueType ValueType = FAnimNextParamType::EValueType::None;
+	FAnimNextParamType::EContainerType ContainerType = FAnimNextParamType::EContainerType::None;
+	UObject* ValueTypeObject = nullptr;
+
+	if (InPinType.PinCategory == UEdGraphSchema_K2::PC_Boolean)
+	{
+		ValueType = FAnimNextParamType::EValueType::Bool;
+	}
+	else if (InPinType.PinCategory == UEdGraphSchema_K2::PC_Byte)
+	{
+		if (UEnum* Enum = Cast<UEnum>(InPinType.PinSubCategoryObject.Get()))
+		{
+			ValueType = FAnimNextParamType::EValueType::Enum;
+			ValueTypeObject = Enum;
+		}
+		else
+		{
+			ValueType = FAnimNextParamType::EValueType::Byte;
+		}
+	}
+	else if (InPinType.PinCategory == UEdGraphSchema_K2::PC_Int)
+	{
+		ValueType = FAnimNextParamType::EValueType::Int32;
+	}
+	else if (InPinType.PinCategory == UEdGraphSchema_K2::PC_Int64)
+	{
+		ValueType = FAnimNextParamType::EValueType::Int64;
+	}
+	else if (InPinType.PinCategory == UEdGraphSchema_K2::PC_Real)
+	{
+		if (InPinType.PinSubCategory == UEdGraphSchema_K2::PC_Float)
+		{
+			ValueType = FAnimNextParamType::EValueType::Float;
+		}
+		else if (InPinType.PinSubCategory == UEdGraphSchema_K2::PC_Double)
+		{
+			ValueType = FAnimNextParamType::EValueType::Double;
+		}
+		else
+		{
+			ensure(false);	// Reals should be either floats or doubles
+		}
+	}
+	else if (InPinType.PinCategory == UEdGraphSchema_K2::PC_Float)
+	{
+		ValueType = FAnimNextParamType::EValueType::Float;
+	}
+	else if (InPinType.PinCategory == UEdGraphSchema_K2::PC_Double)
+	{
+		ValueType = FAnimNextParamType::EValueType::Double;
+	}
+	else if (InPinType.PinCategory == UEdGraphSchema_K2::PC_Name)
+	{
+		ValueType = FAnimNextParamType::EValueType::Name;
+	}
+	else if (InPinType.PinCategory == UEdGraphSchema_K2::PC_String)
+	{
+		ValueType = FAnimNextParamType::EValueType::String;
+	}
+	else if (InPinType.PinCategory == UEdGraphSchema_K2::PC_Text)
+	{
+		ValueType = FAnimNextParamType::EValueType::Text;
+	}
+	else if (InPinType.PinCategory == UEdGraphSchema_K2::PC_Enum)
+	{
+		ValueType = FAnimNextParamType::EValueType::Enum;
+		ValueTypeObject = Cast<UEnum>(InPinType.PinSubCategoryObject.Get());
+		ensure(ValueTypeObject);
+	}
+	else if (InPinType.PinCategory == UEdGraphSchema_K2::PC_Struct)
+	{
+		ValueType = FAnimNextParamType::EValueType::Struct;
+		ValueTypeObject = Cast<UScriptStruct>(InPinType.PinSubCategoryObject.Get());
+	}
+	else if (InPinType.PinCategory == UEdGraphSchema_K2::PC_Object || InPinType.PinCategory == UEdGraphSchema_K2::AllObjectTypes)
+	{
+		ValueType = FAnimNextParamType::EValueType::Object;
+		ValueTypeObject = Cast<UClass>(InPinType.PinSubCategoryObject.Get());
+	}
+	else if (InPinType.PinCategory == UEdGraphSchema_K2::PC_SoftObject)
+	{
+		ValueType = FAnimNextParamType::EValueType::SoftObject;
+		ValueTypeObject = Cast<UClass>(InPinType.PinSubCategoryObject.Get());
+		ensure(ValueTypeObject);
+	}
+	else if (InPinType.PinCategory == UEdGraphSchema_K2::PC_SoftClass)
+	{
+		ValueType = FAnimNextParamType::EValueType::SoftClass;
+		ValueTypeObject = Cast<UClass>(InPinType.PinSubCategoryObject.Get());
+		ensure(ValueTypeObject);
+	}
+
+	if(InPinType.ContainerType == EPinContainerType::Array)
+	{
+		ContainerType = FAnimNextParamType::EContainerType::Array;
+	}
+	else if(InPinType.ContainerType == EPinContainerType::Set)
+	{
+		ensureMsgf(false, TEXT("Set pins are not yet supported"));
+	}
+	if(InPinType.ContainerType == EPinContainerType::Map)
+	{
+		ensureMsgf(false, TEXT("Map pins are not yet supported"));
+	}
+	
+	return FAnimNextParamType(ValueType, ContainerType, ValueTypeObject);
+}
+
+FEdGraphPinType FUtils::GetPinTypeFromParamType(const FAnimNextParamType& InParamType)
+{
+	FEdGraphPinType PinType;
+	PinType.PinSubCategory = NAME_None;
+
+	// Container type
+	switch (InParamType.ContainerType)
+	{
+	case FAnimNextParamType::EContainerType::Array:
+		PinType.ContainerType = EPinContainerType::Array;
+		break;
+	default:
+		PinType.ContainerType = EPinContainerType::None;
+	}
+
+	// Value type
+	switch (InParamType.ValueType)
+	{
+	case EPropertyBagPropertyType::Bool:
+		PinType.PinCategory = UEdGraphSchema_K2::PC_Boolean;
+		break;
+	case EPropertyBagPropertyType::Byte:
+		PinType.PinCategory = UEdGraphSchema_K2::PC_Byte;
+		break;
+	case EPropertyBagPropertyType::Int32:
+		PinType.PinCategory = UEdGraphSchema_K2::PC_Int;
+		break;
+	case EPropertyBagPropertyType::Int64:
+		PinType.PinCategory = UEdGraphSchema_K2::PC_Int64;
+		break;
+	case EPropertyBagPropertyType::Float:
+		PinType.PinCategory = UEdGraphSchema_K2::PC_Real;
+		PinType.PinSubCategory = UEdGraphSchema_K2::PC_Float;
+		break;
+	case EPropertyBagPropertyType::Double:
+		PinType.PinCategory = UEdGraphSchema_K2::PC_Real;
+		PinType.PinSubCategory = UEdGraphSchema_K2::PC_Double;
+		break;
+	case EPropertyBagPropertyType::Name:
+		PinType.PinCategory = UEdGraphSchema_K2::PC_Name;
+		break;
+	case EPropertyBagPropertyType::String:
+		PinType.PinCategory = UEdGraphSchema_K2::PC_String;
+		break;
+	case EPropertyBagPropertyType::Text:
+		PinType.PinCategory = UEdGraphSchema_K2::PC_Text;
+		break;
+	case EPropertyBagPropertyType::Enum:
+		// @todo: some pin coloring is not correct due to this (byte-as-enum vs enum). 
+		PinType.PinCategory = UEdGraphSchema_K2::PC_Enum;
+		PinType.PinSubCategoryObject = const_cast<UObject*>(InParamType.ValueTypeObject.Get());
+		break;
+	case EPropertyBagPropertyType::Struct:
+		PinType.PinCategory = UEdGraphSchema_K2::PC_Struct;
+		PinType.PinSubCategoryObject = const_cast<UObject*>(InParamType.ValueTypeObject.Get());
+		break;
+	case EPropertyBagPropertyType::Object:
+		PinType.PinCategory = UEdGraphSchema_K2::PC_Object;
+		PinType.PinSubCategoryObject = const_cast<UObject*>(InParamType.ValueTypeObject.Get());
+		break;
+	case EPropertyBagPropertyType::SoftObject:
+		PinType.PinCategory = UEdGraphSchema_K2::PC_SoftObject;
+		PinType.PinSubCategoryObject = const_cast<UObject*>(InParamType.ValueTypeObject.Get());
+		break;
+	case EPropertyBagPropertyType::Class:
+		PinType.PinCategory = UEdGraphSchema_K2::PC_Class;
+		PinType.PinSubCategoryObject = const_cast<UObject*>(InParamType.ValueTypeObject.Get());
+		break;
+	case EPropertyBagPropertyType::SoftClass:
+		PinType.PinCategory = UEdGraphSchema_K2::PC_SoftClass;
+		PinType.PinSubCategoryObject = const_cast<UObject*>(InParamType.ValueTypeObject.Get());
+		break;
+	default:
+		break;
+	}
+
+	return PinType;
+}
+
+FRigVMTemplateArgumentType FUtils::GetRigVMArgTypeFromParamType(const FAnimNextParamType& InParamType)
+{
+	FRigVMTemplateArgumentType ArgType;
+
+	FString CPPTypeString;
+
+	// Value type
+	switch (InParamType.ValueType)
+	{
+	case EPropertyBagPropertyType::Bool:
+		CPPTypeString = RigVMTypeUtils::BoolType;
+		break;
+	case EPropertyBagPropertyType::Byte:
+		CPPTypeString = RigVMTypeUtils::UInt8Type;
+		break;
+	case EPropertyBagPropertyType::Int32:
+		CPPTypeString = RigVMTypeUtils::UInt32Type;
+		break;
+	case EPropertyBagPropertyType::Int64:
+		ensureMsgf(false, TEXT("Unhandled value type %d"), InParamType.ValueType);
+		break;
+	case EPropertyBagPropertyType::Float:
+		CPPTypeString = RigVMTypeUtils::FloatType;
+		break;
+	case EPropertyBagPropertyType::Double:
+		CPPTypeString = RigVMTypeUtils::DoubleType;
+		break;
+	case EPropertyBagPropertyType::Name:
+		CPPTypeString = RigVMTypeUtils::FNameType;
+		break;
+	case EPropertyBagPropertyType::String:
+		CPPTypeString = RigVMTypeUtils::FStringType;
+		break;
+	case EPropertyBagPropertyType::Text:
+		ensureMsgf(false, TEXT("Unhandled value type %d"), InParamType.ValueType);
+		break;
+	case EPropertyBagPropertyType::Enum:
+		CPPTypeString = RigVMTypeUtils::CPPTypeFromEnum(Cast<UEnum>(InParamType.ValueTypeObject.Get()));
+		ArgType.CPPTypeObject = const_cast<UObject*>(InParamType.ValueTypeObject.Get());
+		break;
+	case EPropertyBagPropertyType::Struct:
+		CPPTypeString = RigVMTypeUtils::GetUniqueStructTypeName(Cast<UScriptStruct>(InParamType.ValueTypeObject.Get()));
+		ArgType.CPPTypeObject = const_cast<UObject*>(InParamType.ValueTypeObject.Get());
+		break;
+	case EPropertyBagPropertyType::Object:
+		CPPTypeString = RigVMTypeUtils::CPPTypeFromObject(Cast<UClass>(InParamType.ValueTypeObject.Get()), RigVMTypeUtils::EClassArgType::AsObject);
+		ArgType.CPPTypeObject = const_cast<UObject*>(InParamType.ValueTypeObject.Get());
+		break;
+	case EPropertyBagPropertyType::SoftObject:
+		ensureMsgf(false, TEXT("Unhandled value type %d"), InParamType.ValueType);
+		break;
+	case EPropertyBagPropertyType::Class:
+		CPPTypeString = RigVMTypeUtils::CPPTypeFromObject(Cast<UClass>(InParamType.ValueTypeObject.Get()), RigVMTypeUtils::EClassArgType::AsClass);
+		ArgType.CPPTypeObject = const_cast<UObject*>(InParamType.ValueTypeObject.Get());
+		break;
+	case EPropertyBagPropertyType::SoftClass:
+		ensureMsgf(false, TEXT("Unhandled value type %d"), InParamType.ValueType);
+		break;
+	default:
+		ensureMsgf(false, TEXT("Unhandled value type %d"), InParamType.ValueType);
+		break;
+	}
+
+	// Container type
+	switch (InParamType.ContainerType)
+	{
+	case FAnimNextParamType::EContainerType::None:
+		break;
+	case FAnimNextParamType::EContainerType::Array:
+		CPPTypeString = FString::Printf(RigVMTypeUtils::TArrayTemplate, *CPPTypeString);
+		break;
+	default:
+		ensureMsgf(false, TEXT("Unhandled container type %d"), InParamType.ContainerType);
+		break;
+	}
+
+	ArgType.CPPType = *CPPTypeString;
+
+	return ArgType;
+}
+
+void FUtils::SetupEventGraph(URigVMController* InController, UScriptStruct* InEventStruct, FName InEventName, bool bPrintPythonCommand)
+{
+	// Clear the graph
+	InController->RemoveNodes(InController->GetGraph()->GetNodes());
+
+	if (InEventStruct->IsChildOf(FRigUnit_AnimNextUserEvent::StaticStruct()))
+	{
+		FRigUnit_AnimNextUserEvent Defaults;
+		Defaults.Name = InEventName;
+		Defaults.SortOrder = InEventName.GetNumber();
+		InController->AddUnitNodeWithDefaults(InEventStruct, Defaults, FRigVMStruct::ExecuteName, FVector2D(-200.0f, 0.0f), FString(), false);
+	}
+	else if (InEventStruct == FRigVMFunction_UserDefinedEvent::StaticStruct())
+	{
+		FRigVMFunction_UserDefinedEvent Defaults;
+		Defaults.EventName = InEventName;
+		InController->AddUnitNodeWithDefaults(InEventStruct, Defaults, FRigVMStruct::ExecuteName, FVector2D(-200.0f, 0.0f), FString(), false);
+	}
+	else
+	{
+		InController->AddUnitNode(InEventStruct, FRigVMStruct::ExecuteName, FVector2D(-200.0f, 0.0f), FString(), false);
+	}
+}
+
+FAnimNextParamType FUtils::FindVariableType(const FAnimNextVariableReference& InVariableReference)
+{
+	// If this is likely a struct variable, just resolve & search for the member
+	if (const UScriptStruct* Struct = Cast<UScriptStruct>(InVariableReference.GetObject()))
+	{
+		if (const FProperty* Property = Struct->FindPropertyByName(InVariableReference.GetName()))
+		{
+			return FAnimNextParamType::FromProperty(Property);
+		}
+	}
+	else if (const UAnimNextRigVMAsset* Asset = Cast<UAnimNextRigVMAsset>(InVariableReference.GetObject()))
+	{
+		// Query the asset registry for an asset's variable exports
+		FAssetData AssetData(Asset);
+		FAnimNextAssetRegistryExports Exports;
+		if (GetExportedVariablesForAsset(AssetData, Exports))
+		{
+			for(const FAnimNextExport& Export : Exports.Exports)
+			{
+				if(Export.Identifier == InVariableReference.GetName())
+				{
+					if (const FAnimNextVariableDeclarationData* VariableDeclaration = Export.Data.GetPtr<FAnimNextVariableDeclarationData>())
+					{
+						return VariableDeclaration->Type;
+					}
+				}
+			}
+		}
+	}
+
+	return FAnimNextParamType();
+}
+
+
+FAnimNextParamType FUtils::FindVariableType(const FAnimNextSoftVariableReference& InSoftVariableReference)
+{
+	// If this is likely a struct variable, just resolve & search for the member
+	if (InSoftVariableReference.GetSoftObjectPath().GetLongPackageName().StartsWith(TEXT("/Script/")))
+	{
+		if (const UScriptStruct* Struct = Cast<UScriptStruct>(InSoftVariableReference.GetSoftObjectPath().TryLoad()))
+		{
+			if (const FProperty* Property = Struct->FindPropertyByName(InSoftVariableReference.GetName()))
+			{
+				return FAnimNextParamType::FromProperty(Property);
+			}
+		}
+	}
+	else
+	{
+		// Query the asset registry for the asset
+		FAssetData AssetData = IAssetRegistry::GetChecked().GetAssetByObjectPath(InSoftVariableReference.GetSoftObjectPath());
+
+		// If the asset is loaded, just use it
+		if (const UAnimNextRigVMAsset* Asset = Cast<UAnimNextRigVMAsset>(AssetData.FastGetAsset(false)))
+		{
+			const UAnimNextRigVMAssetEditorData* EditorData = GetEditorData<UAnimNextRigVMAssetEditorData>(Asset);
+			FAnimNextAssetRegistryExports Exports;
+
+			FAssetRegistryTagsContextData ContextData(EditorData, EAssetRegistryTagsCaller::Uncategorized);
+			const FAssetRegistryTagsContext Context(ContextData);
+			GetAssetVariableExports(EditorData, Exports, Context);
+
+			for(const FAnimNextExport& Export : Exports.Exports)
+			{
+				const FAnimNextVariableDeclarationData* Data = Export.Data.GetPtr<FAnimNextVariableDeclarationData>();
+				if(Export.Identifier == InSoftVariableReference.GetName() && Data)
+				{
+					return Data->Type;
+				}
+			}
+		}
+		else
+		{
+			// Otherwise use AR data to find its exports
+			FAnimNextAssetRegistryExports Exports;
+			if (GetExportedVariablesForAsset(AssetData, Exports))
+			{
+				for(const FAnimNextExport& Export : Exports.Exports)
+				{
+					if(Export.Identifier == InSoftVariableReference.GetName())
+					{
+						if (const FAnimNextVariableDeclarationData* VariableDeclaration = Export.Data.GetPtr<FAnimNextVariableDeclarationData>())
+						{
+							return VariableDeclaration->Type;
+						}
+					}
+				}
+			}
+		}
+	}
+
+	return FAnimNextParamType();
+}
+
+bool FUtils::GetExportedVariablesForAsset(const FAssetData& InAsset, FAnimNextAssetRegistryExports& OutExports)
+{
+	return GetExportsOfTypeForAsset<FAnimNextVariableDeclarationData>(InAsset, OutExports);
+}
+
+bool FUtils::GetExportedVariablesFromAssetRegistry(TMap<FAssetData, FAnimNextAssetRegistryExports>& OutExports)
+{	
+	return GetExportsOfTypeFromAssetRegistry<FAnimNextVariableDeclarationData>(OutExports);
+}
+
+void FUtils::GetAssetFunctions(const UAnimNextRigVMAssetEditorData* InEditorData, FRigVMGraphFunctionHeaderArray& OutExports)
+{
+	for (const FRigVMGraphFunctionData& FunctionData : InEditorData->GraphFunctionStore.PublicFunctions)
+	{
+		if (FunctionData.CompilationData.IsValid())
+		{
+			OutExports.Headers.Add(FunctionData.Header);
+		}
+	}
+}
+
+void FUtils::GetAssetPrivateFunctions(const UAnimNextRigVMAssetEditorData* InEditorData, FRigVMGraphFunctionHeaderArray& OutExports)
+{
+	for (const FRigVMGraphFunctionData& FunctionData : InEditorData->GraphFunctionStore.PrivateFunctions)
+	{
+		// Note: We dont check compilation data here as private functions are not compiled if they are not referenced
+		OutExports.Headers.Add(FunctionData.Header);
+	}
+}
+
+bool FUtils::GetExportedFunctionsFromAssetRegistry(FName Tag, TMap<FAssetData, FRigVMGraphFunctionHeaderArray>& OutExports)
+{
+	TArray<FAssetData> AssetData;
+	IAssetRegistry::GetChecked().GetAssetsByTags({ Tag }, AssetData);
+
+	const FArrayProperty* HeadersProperty = CastField<FArrayProperty>(FRigVMGraphFunctionHeaderArray::StaticStruct()->FindPropertyByName(TEXT("Headers")));
+
+	for (const FAssetData& Asset : AssetData)
+	{
+		const FString TagValue = Asset.GetTagValueRef<FString>(Tag);
+		FRigVMGraphFunctionHeaderArray AssetExports;
+
+		if (HeadersProperty->ImportText_Direct(*TagValue, &AssetExports, nullptr, EPropertyPortFlags::PPF_None) != nullptr)
+		{
+			if (AssetExports.Headers.Num() > 0)
+			{
+				FRigVMGraphFunctionHeaderArray& AssetArray = OutExports.FindOrAdd(Asset);
+				AssetArray.Headers.Append(MoveTemp(AssetExports.Headers));
+			}
+		}
+	}
+
+	return OutExports.Num() > 0;
+}
+
+static void AddParamToSet(const FAnimNextExport& InNewParam, TSet<FAnimNextExport>& OutExports)
+{
+	if(FAnimNextExport* ExistingEntry = OutExports.Find(InNewParam))
+	{
+		const FAnimNextVariableDeclarationData* NewData = InNewParam.Data.GetPtr<FAnimNextVariableDeclarationData>();
+		FAnimNextVariableDeclarationData* ExistingData = ExistingEntry->Data.GetMutablePtr<FAnimNextVariableDeclarationData>();
+		check(NewData && ExistingData);
+		
+		if(ExistingData->Type != NewData->Type)
+		{
+			UE_LOGFMT(LogAnimation, Warning, "Type mismatch between parameter {ParameterName}. {ParamType1} vs {ParamType1}", InNewParam.Identifier, NewData->Type.ToString(), ExistingData->Type.ToString());
+		}
+		ExistingData->Flags |= NewData->Flags;
+	}
+	else
+	{
+		OutExports.Add(InNewParam);
+	}
+}
+
+void FUtils::GetAssetVariableExports(const UAnimNextRigVMAssetEditorData* EditorData, FAnimNextAssetRegistryExports& OutExports, FAssetRegistryTagsContext Context)
+{
+	OutExports.Exports.Reserve(EditorData->Entries.Num());
+
+	TSet<FAnimNextExport> ExportSet;
+	GetAssetVariableExports(EditorData, ExportSet, Context);
+	OutExports.Exports.Append(ExportSet.Array());
+}
+
+void FUtils::GetAssetVariableExports(const UAnimNextRigVMAssetEditorData* InEditorData, TSet<FAnimNextExport>& OutExports, FAssetRegistryTagsContext Context, EAnimNextExportedVariableFlags InFlags /*= EAnimNextExportedVariableFlags::NoFlags*/)
+{
+	for(const UAnimNextRigVMAssetEntry* Entry : InEditorData->Entries)
+	{
+		if(const IAnimNextRigVMExportInterface* ExportInterface = Cast<IAnimNextRigVMExportInterface>(Entry))
+		{
+			EAnimNextExportedVariableFlags Flags = EAnimNextExportedVariableFlags::Declared | InFlags;
+			if(ExportInterface->GetExportAccessSpecifier() == EAnimNextExportAccessSpecifier::Public)
+			{
+				Flags |= EAnimNextExportedVariableFlags::Public;
+			}
+						
+			FAnimNextExport ParameterExport = FAnimNextExport::MakeExport<FAnimNextVariableDeclarationData>(ExportInterface->GetExportName(), ExportInterface->GetExportType(), Flags);
+			AddParamToSet(ParameterExport, OutExports);
+		}
+		else if(const UAnimNextSharedVariablesEntry* DataInterfaceEntry = Cast<UAnimNextSharedVariablesEntry>(Entry))
+		{
+			switch(DataInterfaceEntry->Type)
+			{
+			case EAnimNextSharedVariablesType::Asset:
+				if(const UAnimNextSharedVariables* SharedVariablesAsset = DataInterfaceEntry->Asset)
+				{
+					UAnimNextSharedVariables_EditorData* EditorData = GetEditorData<UAnimNextSharedVariables_EditorData>(SharedVariablesAsset);
+					GetAssetVariableExports(EditorData, OutExports, Context, EAnimNextExportedVariableFlags::Referenced);
+				}
+				break;
+			case EAnimNextSharedVariablesType::Struct:
+				if(const UScriptStruct* Struct = DataInterfaceEntry->Struct)
+				{
+					GetStructVariableExports(Struct, OutExports);
+				}
+				break;
+			}
+		}
+
+		if(const IAnimNextRigVMGraphInterface* GraphInterface = Cast<IAnimNextRigVMGraphInterface>(Entry))
+		{
+			if (URigVMEdGraph* RigVMEdGraph = GraphInterface->GetEdGraph())
+			{
+				GetSubGraphVariableExportsRecursive(RigVMEdGraph, OutExports);
+			}
+		}
+	}
+}
+
+void FUtils::GetStructVariableExports(const UScriptStruct* Struct, TSet<FAnimNextExport>& OutExports)
+{
+	constexpr EAnimNextExportedVariableFlags Flags = EAnimNextExportedVariableFlags::Declared | EAnimNextExportedVariableFlags::Public;
+	for (TFieldIterator<FProperty> It(Struct); It; ++It)
+	{
+		if(const FProperty* Property = *It)
+		{
+			const FName PropertyName = Property->GetFName();
+			FAnimNextExport ParameterExport = FAnimNextExport::MakeExport<FAnimNextVariableDeclarationData>(PropertyName, FAnimNextParamType::FromProperty(Property), Flags);
+			AddParamToSet(ParameterExport, OutExports);
+		}
+	}
+}
+
+void FUtils::GetAssetWorkspaceExports(const UAnimNextRigVMAssetEditorData* EditorData, FWorkspaceOutlinerItemExports& OutExports, FAssetRegistryTagsContext Context)
+{
+	FWorkspaceOutlinerItemExport AssetIdentifier = FWorkspaceOutlinerItemExport(EditorData->GetOuter()->GetFName(), EditorData->GetOuter());
+
+	const int32 RootExportIndex = OutExports.Exports.Num() - 1;
+
+	int32 SharedVariablesGroupIndex = INDEX_NONE;
+	auto AddSharedVariablesGroupEntry = [&SharedVariablesGroupIndex, &OutExports, RootExportIndex]()
+	{
+		const FWorkspaceOutlinerItemExport& ParentExport = OutExports.Exports[RootExportIndex];
+		FWorkspaceOutlinerItemExport& GroupExport = OutExports.Exports.Add_GetRef(FWorkspaceOutlinerItemExport(FName("SharedVariablesGroup"), ParentExport));
+		GroupExport.GetData().InitializeAs<FWorkspaceOutlinerGroupItemData>();
+		GroupExport.GetData().GetMutable<FWorkspaceOutlinerGroupItemData>().GroupName = TEXT("Shared Variables");
+		GroupExport.GetData().GetMutable<FWorkspaceOutlinerGroupItemData>().GroupIcon = *FSlateIconFinder::FindIconBrushForClass(UAnimNextSharedVariables::StaticClass());
+
+		SharedVariablesGroupIndex = OutExports.Exports.Num() - 1;
+	};
+
+	// Referred SharedVariables
+	EditorData->ForEachEntryOfType<UAnimNextSharedVariablesEntry>([&OutExports, &SharedVariablesGroupIndex, AddSharedVariablesGroupEntry](const UAnimNextSharedVariablesEntry* SharedVariablesEntry)
+	{
+		if(SharedVariablesEntry)
+		{
+			if (SharedVariablesGroupIndex == INDEX_NONE)
+			{
+				AddSharedVariablesGroupEntry();
+			}
+
+			switch(SharedVariablesEntry->Type)
+			{
+			case EAnimNextSharedVariablesType::Asset:
+				{
+					const UAnimNextRigVMAsset* Asset = SharedVariablesEntry->Asset;
+					const FSoftObjectPath ObjectPath(Asset);
+					const FWorkspaceOutlinerItemExport& ParentExport = OutExports.Exports[SharedVariablesGroupIndex];
+					FWorkspaceOutlinerItemExport& ReferenceExport = OutExports.Exports.Add_GetRef(FWorkspaceOutlinerItemExport(FName(ObjectPath.ToString()), ParentExport));
+					ReferenceExport.GetData().InitializeAs<FWorkspaceOutlinerAssetReferenceItemData>();
+					FWorkspaceOutlinerAssetReferenceItemData& Data = ReferenceExport.GetData().GetMutable<FWorkspaceOutlinerAssetReferenceItemData>();
+					Data.ReferredObjectPath = ObjectPath;
+					Data.bShouldExpandReference = false;
+				}
+				break;
+			case EAnimNextSharedVariablesType::Struct:
+				{
+					const UScriptStruct* Struct = SharedVariablesEntry->Struct;
+					const FSoftObjectPath ObjectPath(Struct);
+					const FWorkspaceOutlinerItemExport& ParentExport = OutExports.Exports[SharedVariablesGroupIndex];
+					FWorkspaceOutlinerItemExport& ReferenceExport = OutExports.Exports.Add_GetRef(FWorkspaceOutlinerItemExport(FName(ObjectPath.ToString()), ParentExport));
+					ReferenceExport.GetData().InitializeAs<FWorkspaceOutlinerAssetReferenceItemData>();
+					FWorkspaceOutlinerAssetReferenceItemData& Data = ReferenceExport.GetData().GetMutable<FWorkspaceOutlinerAssetReferenceItemData>();
+					Data.ReferredObjectPath = ObjectPath;
+					Data.bShouldExpandReference = false;
+				}
+				break;
+			}
+		}
+		return true;
+	});
+
+
+	int32 GraphsGroupIndex = INDEX_NONE;
+	auto AddGraphsGroupEntry = [&GraphsGroupIndex, &OutExports, RootExportIndex]()
+	{
+		const FWorkspaceOutlinerItemExport& ParentExport = OutExports.Exports[RootExportIndex];
+		FWorkspaceOutlinerItemExport& GroupExport = OutExports.Exports.Add_GetRef(FWorkspaceOutlinerItemExport(FName("GraphsGroup"), ParentExport));
+		GroupExport.GetData().InitializeAs<FWorkspaceOutlinerGroupItemData>();
+		GroupExport.GetData().GetMutable<FWorkspaceOutlinerGroupItemData>().GroupName = TEXT("Graphs");
+		GroupExport.GetData().GetMutable<FWorkspaceOutlinerGroupItemData>().GroupIcon = *FAppStyle::GetBrush(TEXT("GraphEditor.EventGraph_24x"));
+
+		GraphsGroupIndex = OutExports.Exports.Num() - 1;
+	};
+
+	EditorData->ForEachEntryOfType<IAnimNextRigVMGraphInterface>([&EditorData, &OutExports, AssetIdentifier, &Context, &GraphsGroupIndex, &AddGraphsGroupEntry](IAnimNextRigVMGraphInterface* GraphInterface)
+	{
+		UAnimNextRigVMAssetEntry* Entry = CastChecked<UAnimNextRigVMAssetEntry>(GraphInterface);		
+		if (Entry->IsHiddenInOutliner())
+		{
+			if (URigVMEdGraph* RigVMEdGraph = GraphInterface->GetEdGraph())
+			{
+				GetSubGraphWorkspaceExportsRecursive(EditorData, OutExports, AssetIdentifier, INDEX_NONE, RigVMEdGraph, Context);
+			}
+		}
+		else
+		{
+			if (GraphsGroupIndex == INDEX_NONE)
+			{
+				AddGraphsGroupEntry();
+			}
+			
+			FWorkspaceOutlinerItemExport& Export = OutExports.Exports.Add_GetRef(FWorkspaceOutlinerItemExport(Entry->GetEntryName(), OutExports.Exports[GraphsGroupIndex]));
+
+			Export.GetData().InitializeAsScriptStruct(FAnimNextGraphOutlinerData::StaticStruct());
+			FAnimNextGraphOutlinerData& GraphData = Export.GetData().GetMutable<FAnimNextGraphOutlinerData>();
+			GraphData.SoftEntryPtr = Entry;
+
+			if (URigVMEdGraph* RigVMEdGraph = GraphInterface->GetEdGraph())
+			{					
+				const int32 ExportIndex = OutExports.Exports.Num() - 1;
+				GetSubGraphWorkspaceExportsRecursive(EditorData, OutExports, OutExports.Exports[ExportIndex], ExportIndex, RigVMEdGraph, Context);
+			}
+		}
+		return true;
+	});	
+
+	GetFunctionLibraryWorkspaceExportsRecursive(EditorData, OutExports, AssetIdentifier, INDEX_NONE, EditorData->GetRigVMGraphFunctionStore()->PublicFunctions, EditorData->GetRigVMGraphFunctionStore()->PrivateFunctions);
+}
+
+void FUtils::ProcessPinAssetReferences(const URigVMPin* InPin, FWorkspaceOutlinerItemExports& OutExports, const FWorkspaceOutlinerItemExport& RootExport, int32 ParentExportIndex)
+{
+	auto HandleSoftObjectPath = [&ParentExportIndex, &RootExport, &OutExports](const FSoftObjectPath& InSoftObjectPath) 
+	{
+		// Only add export if object is loaded, or the path actually points to an asset
+		FAssetData ReferenceAssetData;
+		if (InSoftObjectPath.ResolveObject() || IAssetRegistry::GetChecked().TryGetAssetByObjectPath(InSoftObjectPath, ReferenceAssetData) == UE::AssetRegistry::EExists::Exists)
+		{
+			const FWorkspaceOutlinerItemExport& ParentExport = ParentExportIndex == INDEX_NONE ? RootExport : OutExports.Exports[ParentExportIndex];
+			FWorkspaceOutlinerItemExport& ReferenceExport = OutExports.Exports.Add_GetRef(FWorkspaceOutlinerItemExport(FName(InSoftObjectPath.ToString()), ParentExport));
+			ReferenceExport.GetData().InitializeAs<FWorkspaceOutlinerAssetReferenceItemData>();
+			ReferenceExport.GetData().GetMutable<FWorkspaceOutlinerAssetReferenceItemData>().ReferredObjectPath = InSoftObjectPath;
+		}
+	};
+	
+	if (InPin)
+	{
+		const UObject* TypeObject = InPin->GetCPPTypeObject();
+		if (Cast<UClass>(TypeObject))
+		{
+			FSoftObjectPath ObjectPath(InPin->GetDefaultValue());
+			if (ObjectPath.IsValid())
+			{
+				HandleSoftObjectPath(ObjectPath);
+			}
+			else
+			{
+
+				auto HandleDirectPinVariableLink = [&HandleSoftObjectPath](const IAnimNextRigVMVariableInterface* InVariableInterface)
+				{
+					FString DefaultValue;
+					InVariableInterface->GetDefaultValueString(DefaultValue);
+					const FSoftObjectPath ObjectPath = FSoftObjectPath(DefaultValue);
+
+					if (ObjectPath.IsValid())
+					{
+						HandleSoftObjectPath(ObjectPath);
+					}
+				};
+				
+				// Check for variable nodes linked directly to this pin
+				for (URigVMPin* LinkedPin : InPin->GetLinkedSourcePins())
+				{
+					const URigVMNode* LinkedNode = LinkedPin->GetNode();
+					if (const UAnimNextSharedVariableNode* LinkedSharedVariableNode = Cast<UAnimNextSharedVariableNode>(LinkedNode))
+					{
+						if (LinkedSharedVariableNode->Type == EAnimNextSharedVariablesType::Asset && LinkedSharedVariableNode->Asset)
+						{
+							if (const UAnimNextRigVMAssetEditorData* EditorData = GetEditorData<const UAnimNextRigVMAssetEditorData, const UAnimNextSharedVariables>(LinkedSharedVariableNode->Asset))
+							{
+								if (IAnimNextRigVMVariableInterface* VariableInterface = Cast<IAnimNextRigVMVariableInterface>(EditorData->FindEntry(LinkedSharedVariableNode->GetVariableName())))
+								{
+									HandleDirectPinVariableLink(VariableInterface);
+								}
+							}
+						}
+					}
+					else if (const URigVMVariableNode* LinkedVariableNode = Cast<URigVMVariableNode>(LinkedNode))
+					{
+						if (UAnimNextRigVMAssetEditorData* EditorData = InPin->GetTypedOuter<UAnimNextRigVMAssetEditorData>())
+						{
+							if (IAnimNextRigVMVariableInterface* VariableInterface = Cast<IAnimNextRigVMVariableInterface>(EditorData->FindEntry(LinkedVariableNode->GetVariableName())))
+							{
+								HandleDirectPinVariableLink(VariableInterface);
+							}
+						}
+					}
+				}
+
+				// Check for variable nodes linked to this pin its parent pin
+				if (InPin->GetParentPin())
+				{
+					auto HandleParentPinVariableLink = [&HandleSoftObjectPath](const IAnimNextRigVMVariableInterface* InVariableInterface, const URigVMPin* InPin)
+					{
+						const FAnimNextParamType& ParameterType = InVariableInterface->GetType();
+						if (ParameterType.GetValueType() == FAnimNextParamType::EValueType::Struct)
+						{
+							if (const UScriptStruct* TypeStruct = Cast<UScriptStruct>(ParameterType.GetValueTypeObject()))
+							{
+								if (const FProperty* StructPinProperty = TypeStruct->FindPropertyByName(InPin->GetFName()))
+								{
+									FString DefaultValue;
+									StructPinProperty->ExportTextItem_InContainer(DefaultValue, InVariableInterface->GetValuePtr(), nullptr, nullptr, 0 );
+									const FSoftObjectPath ObjectPath = FSoftObjectPath(DefaultValue);
+									if (ObjectPath.IsValid())
+									{
+										HandleSoftObjectPath(ObjectPath);
+									}
+								}
+							}
+						}
+					};
+					
+					for (URigVMPin* LinkedPin : InPin->GetParentPin()->GetLinkedSourcePins())
+					{
+						const URigVMNode* LinkedNode = LinkedPin->GetNode();
+						if (const UAnimNextSharedVariableNode* LinkedSharedVariableNode = Cast<UAnimNextSharedVariableNode>(LinkedNode))
+						{
+							if (LinkedSharedVariableNode->Type == EAnimNextSharedVariablesType::Asset && LinkedSharedVariableNode->Asset)
+							{
+								if (const UAnimNextRigVMAssetEditorData* EditorData = GetEditorData<const UAnimNextRigVMAssetEditorData, const UAnimNextSharedVariables>(LinkedSharedVariableNode->Asset))
+								{
+									if (IAnimNextRigVMVariableInterface* VariableInterface = Cast<IAnimNextRigVMVariableInterface>(EditorData->FindEntry(LinkedSharedVariableNode->GetVariableName())))
+									{
+										HandleParentPinVariableLink(VariableInterface, InPin);
+									}
+								}
+							}
+						}
+						else if (const URigVMVariableNode* LinkedVariableNode = Cast<URigVMVariableNode>(LinkedNode))
+						{
+							if (UAnimNextRigVMAssetEditorData* EditorData = InPin->GetTypedOuter<UAnimNextRigVMAssetEditorData>())
+							{
+								if (IAnimNextRigVMVariableInterface* VariableInterface = Cast<IAnimNextRigVMVariableInterface>(EditorData->FindEntry(LinkedVariableNode->GetVariableName())))
+								{
+									HandleParentPinVariableLink(VariableInterface, InPin);
+								}
+							}
+						}
+					}
+				}
+			}
+		}
+
+		for (const URigVMPin* SubPin : InPin->GetSubPins())
+		{
+			ProcessPinAssetReferences(SubPin, OutExports, RootExport, ParentExportIndex);
+		}
+	}	
+}
+
+void FUtils::GetSubGraphWorkspaceExportsRecursive(const UAnimNextRigVMAssetEditorData* EditorData, FWorkspaceOutlinerItemExports& OutExports, const FWorkspaceOutlinerItemExport& RootExport, int32 ParentExportIndex, URigVMEdGraph* RigVMEdGraph, FAssetRegistryTagsContext Context)
+{
+	if (RigVMEdGraph == nullptr)
+	{
+		return;
+	}
+
+	// Handle pin asset references (disabled during save as GetMetaData can cause StaticFindFast calls which is prohibited during save)
+	if (!Context.IsSaving())
+	{
+		for (const TObjectPtr<class UEdGraphNode>& Node : RigVMEdGraph->Nodes)
+		{
+			if (URigVMEdGraphNode* RigVMEdNode = Cast<URigVMEdGraphNode>(Node))
+			{
+				if (URigVMTemplateNode* TemplateRigVMNode = Cast<URigVMTemplateNode>(RigVMEdNode->GetModelNode()))
+				{
+					if (TemplateRigVMNode->GetScriptStruct() && TemplateRigVMNode->GetScriptStruct()->IsChildOf(FRigUnit_AnimNextBase::StaticStruct()))
+					{
+						for (URigVMPin* ModelPin : TemplateRigVMNode->GetPins())
+						{
+							if (ModelPin->GetDirection() == ERigVMPinDirection::Input)
+							{
+								auto HandlePin = [&OutExports, &RootExport, ParentExportIndex](const URigVMPin* InPin)
+								{
+									if (InPin->GetMetaData(TEXT("ExportAsReference")) == TEXT("true"))
+									{
+										ProcessPinAssetReferences(InPin, OutExports, RootExport, ParentExportIndex);
+									}										
+								};
+								
+								HandlePin(ModelPin);
+								
+								for (URigVMPin* TraitPin : ModelPin->GetSubPins())
+								{									
+									HandlePin(TraitPin);
+								}
+							}
+						}
+					}
+				}
+			}
+		}
+	}
+
+	// ---- Collapsed graphs ---
+	for (const TObjectPtr<UEdGraph>& SubGraph : RigVMEdGraph->SubGraphs)
+	{
+		URigVMEdGraph* EditorObject = Cast<URigVMEdGraph>(SubGraph);
+		if (IsValid(EditorObject))
+		{
+			if(ensure(EditorObject->GetModel()))
+			{
+				URigVMCollapseNode* CollapseNode = CastChecked<URigVMCollapseNode>(EditorObject->GetModel()->GetOuter());
+
+				const FWorkspaceOutlinerItemExport& ParentExport = ParentExportIndex == INDEX_NONE ? RootExport : OutExports.Exports[ParentExportIndex];
+				FWorkspaceOutlinerItemExport& Export = OutExports.Exports.Add_GetRef(FWorkspaceOutlinerItemExport(CollapseNode->GetFName(), ParentExport));
+				Export.GetData().InitializeAsScriptStruct(FAnimNextCollapseGraphOutlinerData::StaticStruct());
+			
+				FAnimNextCollapseGraphOutlinerData& FnGraphData = Export.GetData().GetMutable<FAnimNextCollapseGraphOutlinerData>();
+				FnGraphData.SoftEditorObject = EditorObject;
+
+				int32 ExportIndex = OutExports.Exports.Num() - 1;
+				GetSubGraphWorkspaceExportsRecursive(EditorData, OutExports, RootExport, ExportIndex, EditorObject, Context);
+			}
+		}
+	}
+
+	// ---- Function References ---
+	TArray<URigVMEdGraphNode*> EdNodes;
+	RigVMEdGraph->GetNodesOfClass(EdNodes);
+
+	for (URigVMEdGraphNode* EdNode : EdNodes)
+	{
+		if (URigVMFunctionReferenceNode* FunctionReferenceNode = Cast<URigVMFunctionReferenceNode>(EdNode->GetModelNode()))
+		{
+			// Only export referenced functions which are part of same outer
+			const FRigVMGraphFunctionIdentifier FunctionIdentifier = FunctionReferenceNode->GetFunctionIdentifier();
+			if (FunctionIdentifier.HostObject == FSoftObjectPath(EditorData))
+			{
+				const URigVMLibraryNode* FunctionNode = EditorData->RigVMClient.GetFunctionLibrary()->FindFunction(FunctionIdentifier.GetFunctionFName());
+				if (URigVMGraph* ContainedGraph = FunctionNode->GetContainedGraph())
+				{
+					const FWorkspaceOutlinerItemExport& ParentExport = ParentExportIndex == INDEX_NONE ? RootExport : OutExports.Exports[ParentExportIndex];
+					FWorkspaceOutlinerItemExport& Export = OutExports.Exports.Add_GetRef(FWorkspaceOutlinerItemExport(FunctionNode->GetFName(), ParentExport));
+					Export.GetData().InitializeAsScriptStruct(FAnimNextGraphFunctionOutlinerData::StaticStruct());
+					FAnimNextGraphFunctionOutlinerData& FnGraphData = Export.GetData().GetMutable<FAnimNextGraphFunctionOutlinerData>();
+
+					if (URigVMEdGraph* ContainedEdGraph = Cast<URigVMEdGraph>(EditorData->GetEditorObjectForRigVMGraph(ContainedGraph)))
+					{
+						FnGraphData.SoftEditorObject = ContainedEdGraph;
+						FnGraphData.SoftEdGraphNode = EdNode;
+						
+						int32 ExportIndex = OutExports.Exports.Num() - 1;
+						GetSubGraphWorkspaceExportsRecursive(EditorData, OutExports, RootExport, ExportIndex, ContainedEdGraph, Context);
+					}
+				}
+			}
+		}
+	}
+}
+
+void FUtils::GetFunctionLibraryWorkspaceExportsRecursive(const UAnimNextRigVMAssetEditorData* EditorData, FWorkspaceOutlinerItemExports& OutExports, const FWorkspaceOutlinerItemExport& RootExport, int32 ParentExportIndex, const TArray<FRigVMGraphFunctionData>& PublicFunctions, const TArray<FRigVMGraphFunctionData>& PrivateFunctions)
+{
+	if (PrivateFunctions.Num() > 0 || PublicFunctions.Num() > 0)
+	{
+		const FWorkspaceOutlinerItemExport& ParentExport = ParentExportIndex != INDEX_NONE ? OutExports.Exports[ParentExportIndex] : RootExport;
+		FWorkspaceOutlinerItemExport& GroupExport = OutExports.Exports.Add_GetRef(FWorkspaceOutlinerItemExport(FName("FunctionsGroup"), ParentExport));
+		GroupExport.GetData().InitializeAs<FWorkspaceOutlinerGroupItemData>();
+		GroupExport.GetData().GetMutable<FWorkspaceOutlinerGroupItemData>().GroupName = TEXT("Functions");
+		GroupExport.GetData().GetMutable<FWorkspaceOutlinerGroupItemData>().GroupIcon = *FAppStyle::GetBrush(TEXT("GraphEditor.Function_24x"));
+		const int32 GroupExportIndex = OutExports.Exports.Num() - 1;
+
+		GetFunctionsWorkspaceExportsRecursive(EditorData, OutExports, GroupExport, GroupExportIndex, PrivateFunctions, false);
+		GetFunctionsWorkspaceExportsRecursive(EditorData, OutExports, GroupExport, GroupExportIndex, PublicFunctions, true);
+	}
+}
+
+void FUtils::GetFunctionsWorkspaceExportsRecursive(const UAnimNextRigVMAssetEditorData* EditorData, FWorkspaceOutlinerItemExports& OutExports, const FWorkspaceOutlinerItemExport& RootExport, int32 ParentExportIndex, const TArray<FRigVMGraphFunctionData>& Functions, bool bPublicFunctions)
+{
+	for (const FRigVMGraphFunctionData& FunctionData : Functions)
+	{
+		if (const URigVMLibraryNode* FunctionNode = EditorData->RigVMClient.GetFunctionLibrary()->FindFunction(FunctionData.Header.LibraryPointer.GetFunctionFName()))
+		{
+			if (URigVMGraph* ContainedModelGraph = FunctionNode->GetContainedGraph())
+			{
+				if (URigVMEdGraph* EditorObject = Cast<URigVMEdGraph>(EditorData->GetEditorObjectForRigVMGraph(ContainedModelGraph)))
+				{
+					const FWorkspaceOutlinerItemExport& ParentExport = ParentExportIndex == INDEX_NONE ? RootExport : OutExports.Exports[ParentExportIndex];
+					FWorkspaceOutlinerItemExport& Export = OutExports.Exports.Add_GetRef(FWorkspaceOutlinerItemExport(FunctionData.Header.Name, ParentExport));
+
+					Export.GetData().InitializeAsScriptStruct(FAnimNextGraphFunctionOutlinerData::StaticStruct());
+					FAnimNextGraphFunctionOutlinerData& FnGraphData = Export.GetData().GetMutable<FAnimNextGraphFunctionOutlinerData>();
+					FnGraphData.SoftEditorObject = EditorObject;
+				}
+			}
+		}
+	}
+}
+
+void FUtils::GetSubGraphVariableExportsRecursive(const URigVMEdGraph* RigVMEdGraph, TSet<FAnimNextExport>& OutExports)
+{
+	TArray<URigVMEdGraphNode*> EdNodes;
+	RigVMEdGraph->GetNodesOfClass(EdNodes);
+
+	for (URigVMEdGraphNode* EdNode : EdNodes)
+	{
+		if (URigVMVariableNode* VariableNode = Cast<URigVMVariableNode>(EdNode->GetModelNode()))
+		{
+			OutExports.Add(
+				FAnimNextExport::MakeExport<FAnimNextVariableReferenceData>
+				(
+					VariableNode->GetVariableName(),
+					RigVMEdGraph->GetFName(),
+					VariableNode->GetNodePath(),
+					VariableNode->GetValuePin()->GetPinPath(),
+					VariableNode->IsGetter() ? EAnimNextExportedVariableFlags::Read : EAnimNextExportedVariableFlags::Write
+				)
+			);
+		}
+	}
+}
+
+const FText& FUtils::GetFunctionLibraryDisplayName()
+{
+	static const FText FunctionLibraryName = LOCTEXT("WorkspaceFunctionLibraryName", "Function Library");
+	return FunctionLibraryName;
+}
+
+#if WITH_EDITOR
+void FUtils::OpenProgrammaticGraphs(UAnimNextRigVMAssetEditorData* EditorData, const TArray<URigVMGraph*>& ProgrammaticGraphs)
+{
+	UAnimNextRigVMAsset* OwningAsset = FUtils::GetAsset(EditorData);
+	UE::Workspace::IWorkspaceEditorModule& WorkspaceEditorModule = FModuleManager::LoadModuleChecked<UE::Workspace::IWorkspaceEditorModule>("WorkspaceEditor");
+	if(UE::Workspace::IWorkspaceEditor* WorkspaceEditor = WorkspaceEditorModule.OpenWorkspaceForObject(OwningAsset, UE::Workspace::EOpenWorkspaceMethod::Default))
+	{
+		TArray<UObject*> Graphs;
+		for(URigVMGraph* ProgrammaticGraph : ProgrammaticGraphs)
+		{
+			// Some explanation needed here!
+			// URigVMEdGraph caches its underlying model internally in GetModel depending on its outer if it is no attached to a RigVMClient
+			// So here we rename the graph into the transient package so we dont get any notifications
+			ProgrammaticGraph->Rename(nullptr, GetTransientPackage(), REN_ForceNoResetLoaders | REN_DoNotDirty | REN_DontCreateRedirectors | REN_NonTransactional);
+
+			// then create the graph (transient so it outers to the RigVMGraph)
+			URigVMEdGraph* EdGraph = CastChecked<URigVMEdGraph>(EditorData->CreateEdGraph(ProgrammaticGraph, true));
+
+			// Then cache the model
+			EdGraph->GetModel();
+			Graphs.Add(EdGraph);
+
+			// Now rename into this asset again to be able to correctly create a controller (needed to view the graph and interact with it)
+			ProgrammaticGraph->Rename(nullptr, EditorData, REN_ForceNoResetLoaders | REN_DoNotDirty | REN_DontCreateRedirectors | REN_NonTransactional);
+			URigVMController* ProgrammaticController = EditorData->GetOrCreateController(ProgrammaticGraph);
+
+			// Resend notifications to rebuild the EdGraph
+			ProgrammaticController->ResendAllNotifications();
+		}
+
+		WorkspaceEditor->OpenObjects(Graphs);
+	}
+}
+#endif // WITH_EDITOR
+
+FString FUtils::MakeFunctionWrapperVariableName(FName InFunctionName, FName InVariableName)
+{
+	// We assume the function name is enough for variable name uniqueness in this graph (We don't yet desire global uniqueness).
+	return TEXT("__InternalVar_") + InFunctionName.ToString() + "_" + InVariableName.ToString();
+}
+
+FString FUtils::MakeFunctionWrapperEventName(FName InFunctionName)
+{
+	return TEXT("__InternalCall_") + InFunctionName.ToString();
+}
+
+void FUtils::GetVariableNames(UAnimNextRigVMAssetEditorData* InEditorData, TArray<FName>& OutVariableNames, bool bRecursive)
+{
+	for(UAnimNextRigVMAssetEntry* Entry : InEditorData->Entries)
+	{
+		if(UAnimNextVariableEntry* VariableEntry = Cast<UAnimNextVariableEntry>(Entry))
+		{
+			if(VariableEntry->GetExportAccessSpecifier() == EAnimNextExportAccessSpecifier::Public)
+			{
+				OutVariableNames.Add(Entry->GetEntryName());
+			}
+		}
+		else if(UAnimNextSharedVariablesEntry* DataInterfaceEntry = Cast<UAnimNextSharedVariablesEntry>(Entry))
+		{
+			if (bRecursive && DataInterfaceEntry->GetAsset())
+			{
+				UAnimNextSharedVariables_EditorData* EditorData = GetEditorData<UAnimNextSharedVariables_EditorData>(DataInterfaceEntry->GetAsset());
+				GetVariableNames(EditorData, OutVariableNames);
+			}
+		}
+	}
+}
+
+FName FUtils::GetValidVariableName(UAnimNextRigVMAssetEditorData* InEditorData, FName InBaseName)
+{
+	TArray<FName> ExistingNames;
+	GetVariableNames(InEditorData, ExistingNames, true);	
+	return GetValidVariableName(InBaseName, ExistingNames);	
+}
+
+FName FUtils::GetValidVariableName(FName InBaseName, const TArrayView<FName> ExistingNames)
+{
+	auto NameExists = [&ExistingNames](FName InName)
+	{
+		for(FName AdditionalName : ExistingNames)
+		{
+			if(AdditionalName == InName)
+			{
+				return true;
+			}
+		}
+
+		return false;
+	};
+
+	if(!NameExists(InBaseName))
+	{
+		// Early out - name is valid
+		return InBaseName;
+	}
+
+	int32 PostFixIndex = 0;
+	TStringBuilder<128> StringBuilder;
+	while(true)
+	{
+		StringBuilder.Reset();
+		InBaseName.GetDisplayNameEntry()->AppendNameToString(StringBuilder);
+		StringBuilder.Appendf(TEXT("_%d"), PostFixIndex++);
+
+		FName TestName(StringBuilder.ToString()); 
+		if(!NameExists(TestName))
+		{
+			return TestName;
+		}
+	}
+}
+
+void FUtils::DeleteVariable(UAnimNextVariableEntry* VariableEntry, bool bSetupUndoRedo, bool bPrintPythonCommands)
+{
+	if (VariableEntry)
+	{
+		const UAnimNextRigVMAsset* OuterRigVMAsset = VariableEntry->GetTypedOuter<UAnimNextRigVMAsset>();
+		check(OuterRigVMAsset);
+		
+		const FAnimNextSoftVariableReference FindReference(VariableEntry->GetEntryName(), OuterRigVMAsset);
+
+		// Replace references with None
+		UAnimNextRigVMAssetEditorData* EditorData = GetEditorData<UAnimNextRigVMAssetEditorData>(OuterRigVMAsset);
+
+		// Replace reference across project for public variables
+		const bool bIsPublicVariable = VariableEntry->GetExportAccessSpecifier() == EAnimNextExportAccessSpecifier::Public;
+		if (bIsPublicVariable)
+		{
+			ReplaceVariableReferencesAcrossProject(FindReference, FAnimNextSoftVariableReference(), bSetupUndoRedo, bPrintPythonCommands);
+		}
+		else
+		{
+			// Otherwise replace, potential, local references only 
+			ReplaceVariableReferences(EditorData, FindReference, FAnimNextSoftVariableReference(), bSetupUndoRedo, bPrintPythonCommands);		
+		}
+
+		EditorData->RemoveEntry(VariableEntry);
+	}
+}
+
+void FUtils::RenameVariable(UAnimNextVariableEntry* VariableEntry, const FName NewName, bool bSetupUndoRedo, bool bPrintPythonCommands)
+{
+	if (VariableEntry)
+	{
+		UAnimNextRigVMAsset* OuterRigVMAsset = VariableEntry->GetTypedOuter<UAnimNextRigVMAsset>();
+		check(OuterRigVMAsset);
+		
+		const FAnimNextSoftVariableReference FindReference(VariableEntry->GetEntryName(), OuterRigVMAsset);
+		
+		// Rename variable first
+		VariableEntry->SetEntryName(NewName);
+		
+		const FAnimNextSoftVariableReference ReplaceReference(NewName, OuterRigVMAsset);
+		
+		// Replace reference across project for public variables
+		const bool bIsPublicVariable = VariableEntry->GetExportAccessSpecifier() == EAnimNextExportAccessSpecifier::Public;
+		if (bIsPublicVariable)
+		{
+			ReplaceVariableReferencesAcrossProject(FindReference, ReplaceReference, bSetupUndoRedo, bPrintPythonCommands);
+		}
+		else
+		{
+			// Otherwise replace, potential, local references only 
+			ReplaceVariableReferences(GetEditorData(OuterRigVMAsset), FindReference, ReplaceReference, bSetupUndoRedo, bPrintPythonCommands);		
+		}
+	}	
+}
+
+void FUtils::MoveVariableToAsset(UAnimNextVariableEntry* VariableEntry, UAnimNextRigVMAsset* NewOuter, bool bSetupUndoRedo, bool bPrintPythonCommands)
+{
+	if (VariableEntry && NewOuter)
+	{
+		if (UAnimNextRigVMAsset* CurrentOuter = VariableEntry->GetTypedOuter<UAnimNextRigVMAsset>())
+		{
+			UAnimNextRigVMAssetEditorData* CurrentEditorData = GetEditorData(CurrentOuter);
+			UAnimNextRigVMAssetEditorData* NewEditorData = GetEditorData(NewOuter);
+						
+			FString DefaultValueString;
+			VariableEntry->GetDefaultValueString(DefaultValueString);
+
+			check(NewEditorData->AddVariable(VariableEntry->GetVariableName(), VariableEntry->GetType(), DefaultValueString, bSetupUndoRedo, bPrintPythonCommands));
+
+			const FAnimNextSoftVariableReference FindReference(VariableEntry->GetEntryName(), CurrentOuter);
+			const FAnimNextSoftVariableReference ReplaceReference(VariableEntry->GetEntryName(), NewOuter);
+
+			// Replace reference across project for public variables
+			const bool bIsPublicVariable = VariableEntry->GetExportAccessSpecifier() == EAnimNextExportAccessSpecifier::Public;
+			if (bIsPublicVariable)
+			{
+				ReplaceVariableReferencesAcrossProject(FindReference, ReplaceReference, bSetupUndoRedo, bPrintPythonCommands);
+			}
+			else
+			{
+				// Otherwise replace, potential, local references only 
+				ReplaceVariableReferences(CurrentEditorData, FindReference, ReplaceReference, bSetupUndoRedo, bPrintPythonCommands);		
+			}
+
+			CurrentEditorData->RemoveEntry(VariableEntry, bSetupUndoRedo, bPrintPythonCommands);
+		}
+	}
+}
+
+void FUtils::SetVariableType(UAnimNextVariableEntry* VariableEntry, const FAnimNextParamType NewType, bool bSetupUndoRedo, bool bPrintPythonCommands)
+{
+	if (VariableEntry)
+	{
+		UAnimNextRigVMAsset* OuterRigVMAsset = VariableEntry->GetTypedOuter<UAnimNextRigVMAsset>();
+		check(OuterRigVMAsset);
+		
+		const FAnimNextSoftVariableReference FindReference(VariableEntry->GetEntryName(), OuterRigVMAsset);
+		
+		// Change variable type first
+		VariableEntry->SetType(NewType);
+
+		// Replace reference across project for public variables
+		const bool bIsAPublicVariable = VariableEntry->GetExportAccessSpecifier() == EAnimNextExportAccessSpecifier::Public;
+		if (bIsAPublicVariable)
+		{
+			ReplaceVariableReferencesAcrossProject(FindReference, FindReference, bSetupUndoRedo, bPrintPythonCommands);
+		}
+		else
+		{
+			// Otherwise replace, potential, local references only 
+			ReplaceVariableReferences(GetEditorData(OuterRigVMAsset), FindReference, FindReference, bSetupUndoRedo, bPrintPythonCommands);		
+		}
+	}	
+}
+
+void FUtils::ReplaceVariableReferences(UAnimNextRigVMAssetEditorData* InEditorData, const FAnimNextSoftVariableReference& SoftReferenceToFind, const FAnimNextSoftVariableReference& SoftReferenceToReplaceWith, bool bSetupUndoRedo /*= true*/, bool bPrintPythonCommands /*= true*/)
+{
+	if (InEditorData == nullptr)
+	{
+		return;
+	}
+
+	const FAnimNextVariableReference ReferenceToFind(SoftReferenceToFind);	
+	if (ReferenceToFind.IsNone())
+	{
+		// Cannot replace an empty reference
+		return;
+	}
+		
+	UAnimNextRigVMAsset* ThisOuterAsset = GetAsset(InEditorData);
+	const UAnimNextRigVMAsset* FindOuterAsset = Cast<UAnimNextRigVMAsset>(ReferenceToFind.GetObject());
+	const FAnimNextVariableReference ReferenceToReplaceWith(SoftReferenceToReplaceWith);
+	const UAnimNextRigVMAsset* ReplaceOuterAsset = Cast<UAnimNextRigVMAsset>(ReferenceToReplaceWith.GetObject());
+
+	const bool bReplacingWithNone = ReferenceToReplaceWith.IsNone();
+
+	// Replacing variable ref local to this asset, with another local variable or none
+	const bool bInternalReplace = (ThisOuterAsset == FindOuterAsset) && (FindOuterAsset == ReplaceOuterAsset || bReplacingWithNone);
+	
+	// Replacing with variable from a different object vs existing reference
+	const bool bDifferentVariableSource = FindOuterAsset != ReplaceOuterAsset;
+
+	FAnimNextParamType FindReferenceType = FindVariableType(ReferenceToFind);
+	const FAnimNextParamType ReplaceReferenceType = FindVariableType(ReferenceToReplaceWith);
+
+	// In case ReferenceToFind has already been removed/changed match it up with the replacing type (if not None) 
+	if (!FindReferenceType.IsValid() && !ReferenceToReplaceWith.IsNone())
+	{
+		FindReferenceType = ReplaceReferenceType;
+	}
+	
+	const bool bVariableTypeChanged = FindReferenceType != ReplaceReferenceType;
+	const bool bVariableNameChanged = ReferenceToFind.GetName() != ReferenceToReplaceWith.GetName();
+
+	// Ensure that SharedVariables have been included for replace variable
+	if (const UAnimNextSharedVariables* SharedVariablesReplaceOuter = Cast<UAnimNextSharedVariables>(ReplaceOuterAsset))
+	{
+		bool bHasEntry = false;
+		InEditorData->ForEachEntryOfType<UAnimNextSharedVariablesEntry>([&bHasEntry, SharedVariablesReplaceOuter](const UAnimNextSharedVariablesEntry* Entry)
+		{
+			if (Entry->GetType() == EAnimNextSharedVariablesType::Asset)
+			{
+				if (Entry->GetAsset() == SharedVariablesReplaceOuter)
+				{
+					bHasEntry = true;
+					return false;
+				}
+			}
+
+			return true;
+		});
+
+		if (!bHasEntry)
+		{
+			UE::UAF::UncookedOnly::FScopedCompileJob CompilerResults(LOCTEXT("AddSharedAssetJobName", "Add Shared Asset dependency"), { ThisOuterAsset });
+			InEditorData->AddSharedVariables(SharedVariablesReplaceOuter, bSetupUndoRedo, bPrintPythonCommands);
+			InEditorData->RecompileVM();
+		}
+	}			
+	
+	TArray<URigVMEdGraphNode*> EdNodes;
+	for(const UAnimNextRigVMAssetEntry* Entry : InEditorData->Entries)
+	{
+		if(const IAnimNextRigVMGraphInterface* GraphInterface = Cast<IAnimNextRigVMGraphInterface>(Entry))
+		{			
+			if (URigVMEdGraph* RigVMEdGraph = GraphInterface->GetEdGraph())
+			{
+				UAnimNextControllerBase* Controller = CastChecked<UAnimNextControllerBase>(RigVMEdGraph->GetController());
+				// Rely on RigVM controller to handle local variable references
+				if (bInternalReplace)
+				{
+					if (bReplacingWithNone)
+					{
+						Controller->OnExternalVariableRemoved(ReferenceToFind.GetName(), bSetupUndoRedo);
+					}
+					else
+					{
+						if (bVariableTypeChanged)
+						{
+							const FRigVMTemplateArgumentType NewRigVMType = ReplaceReferenceType.ToRigVMTemplateArgument();
+							Controller->OnExternalVariableTypeChanged(ReferenceToFind.GetName(), NewRigVMType.CPPType.ToString(), NewRigVMType.CPPTypeObject.Get(), bSetupUndoRedo);
+						}
+
+						if (bVariableNameChanged)
+						{
+							Controller->OnExternalVariableRenamed(ReferenceToFind.GetName(), ReferenceToReplaceWith.GetName(), bSetupUndoRedo);
+						}
+					}
+				}
+				else
+				{
+					RigVMEdGraph->GetNodesOfClass(EdNodes);
+
+					for (URigVMEdGraphNode* EdNode : EdNodes)
+					{
+						if (URigVMVariableNode* VariableNode = Cast<URigVMVariableNode>(EdNode->GetModelNode()))
+						{
+							const bool bNameMatch = VariableNode->GetVariableName() == ReferenceToFind.GetName();
+
+							if (!bNameMatch)
+							{
+								continue;
+							}								
+							
+							if (UAnimNextSharedVariableNode* SharedVariableNode = Cast<UAnimNextSharedVariableNode>(VariableNode))
+							{
+								const UObject* SearchObject = ReferenceToFind.GetObject(); 
+								const bool bPropertySourceMatch = (SharedVariableNode->Asset && SharedVariableNode->Asset == SearchObject)  || (SharedVariableNode->Struct && SharedVariableNode->Struct == SearchObject);
+								const bool bValidProperty = ReferenceToReplaceWith.ResolveProperty() != nullptr;
+								
+								if (bPropertySourceMatch && (bValidProperty || bReplacingWithNone))
+								{
+									if (bReplacingWithNone)
+									{
+										// [TODO] decide to either remove node or replace with None variable references (keeps node in place but will fail compilation until fixed up)
+										Controller->RemoveNode(SharedVariableNode, bSetupUndoRedo, bPrintPythonCommands);
+										//Controller->RefreshSharedVariableNode(SharedVariableNode->GetFName(), TEXT(""), ReferenceToReplaceWith.GetName(), TEXT(""), nullptr, true, true);	
+									}
+									else
+									{
+										if (bDifferentVariableSource && ReplaceOuterAsset == ThisOuterAsset)
+										{
+											Controller->ReplaceSharedVariableNodeWithVariableNode(SharedVariableNode, ReferenceToReplaceWith.GetName(), bSetupUndoRedo, bPrintPythonCommands);
+										}
+										else
+										{
+											const FRigVMTemplateArgumentType NewRigVMType = ReplaceReferenceType.ToRigVMTemplateArgument();
+											Controller->RefreshSharedVariableNode(SharedVariableNode->GetFName(), ReferenceToReplaceWith.GetObject().GetPath(), ReferenceToReplaceWith.GetName(), NewRigVMType.CPPType.ToString(), NewRigVMType.CPPTypeObject.Get(), bSetupUndoRedo, bPrintPythonCommands);											
+										}
+									}
+								}
+							}
+							else
+							{
+								if (bReplacingWithNone)
+								{
+									// Remove the node for now (can also leave an invalid one causing compile error/warning)
+									Controller->RemoveNode(VariableNode, bSetupUndoRedo, bPrintPythonCommands);
+								}
+								else
+								{
+									if (bDifferentVariableSource)
+									{
+										check(FindOuterAsset == ThisOuterAsset && ReplaceOuterAsset != ThisOuterAsset);
+										{
+											Controller->ReplaceVariableNodeWithSharedVariableNode(VariableNode, ReferenceToReplaceWith.GetName(), ReferenceToReplaceWith.GetObject(), bSetupUndoRedo, bPrintPythonCommands);
+										}
+									}
+									else
+									{					
+										if (bVariableTypeChanged)
+										{
+											const FRigVMTemplateArgumentType NewRigVMType = ReplaceReferenceType.ToRigVMTemplateArgument();
+											Controller->OnExternalVariableTypeChanged(ReferenceToFind.GetName(), NewRigVMType.CPPType.ToString(), NewRigVMType.CPPTypeObject.Get(), bSetupUndoRedo);
+										}
+
+										if (bVariableNameChanged)
+										{
+											Controller->OnExternalVariableRenamed(ReferenceToFind.GetName(), ReferenceToReplaceWith.GetName(), bSetupUndoRedo);
+										}
+									}									
+								}
+							}
+						}
+					}
+				}			
+			}
+		}
+
+		EdNodes.Reset();
+	}	
+}
+
+void FUtils::ReplaceVariableReferencesAcrossProject(const FAnimNextSoftVariableReference& ReferenceToFind, const FAnimNextSoftVariableReference& ReferenceToReplaceWith, bool bSetupUndoRedo, bool bPrintPythonCommands)
+{
+	if (ReferenceToFind.IsNone())
+	{
+		return;
+	}
+	
+	FARFilter AssetFilter;
+	AssetFilter.ClassPaths = { UAnimNextRigVMAsset::StaticClass()->GetClassPathName() };
+	AssetFilter.bRecursiveClasses = true;
+
+	// [TODO] would this miss out assets
+	//AssetFilter.TagsAndValues.Add(UE::UAF::ExportsAnimNextAssetRegistryTag, TOptional<FString>());
+		
+	TArray<FAssetData> AssetDatas;
+	FModuleManager::GetModuleChecked<FAssetRegistryModule>("AssetRegistry").Get().GetAssets(AssetFilter, AssetDatas);
+
+	TArray<UAnimNextRigVMAssetEditorData*> EditorDataRequiringReplacing;	
+	for (const FAssetData& AssetData : AssetDatas)
+	{
+		// Relates to [TODO] above
+		ensure(AssetData.FindTag(UE::UAF::ExportsAnimNextAssetRegistryTag));
+		
+		FAnimNextAssetRegistryExports AnimNextExports;
+		GetExportsOfTypeForAsset<FAnimNextVariableReferenceData>(AssetData, AnimNextExports);
+
+		for (const FAnimNextExport& Export : AnimNextExports.Exports)
+		{
+			if (ReferenceToFind.GetName() == Export.Identifier)
+			{
+				if (UAnimNextRigVMAsset* RigVMAsset = Cast<UAnimNextRigVMAsset>(AssetData.GetAsset()))
+				{
+					if (UAnimNextRigVMAssetEditorData* EditorData = GetEditorData(RigVMAsset))
+					{
+						EditorDataRequiringReplacing.Add(EditorData);
+					}
+				}
+			}
+		}
+	}
+
+	for (UAnimNextRigVMAssetEditorData* EditorData : EditorDataRequiringReplacing)
+	{
+		ReplaceVariableReferences(EditorData, ReferenceToFind, ReferenceToReplaceWith, bSetupUndoRedo, bPrintPythonCommands);
+	}
+}
+}
+
+#undef LOCTEXT_NAMESPACE
diff --git a/Engine/Plugins/Experimental/UAF/UAF/Source/UAFUncookedOnly/Private/Variables/AnimNextSharedVariables_EditorData.cpp b/Engine/Plugins/Experimental/UAF/UAF/Source/UAFUncookedOnly/Private/Variables/AnimNextSharedVariables_EditorData.cpp
new file mode 100644
--- /dev/null
+++ b/Engine/Plugins/Experimental/UAF/UAF/Source/UAFUncookedOnly/Private/Variables/AnimNextSharedVariables_EditorData.cpp
@@ -0,0 +1,35 @@
+﻿// Copyright Epic Games, Inc. All Rights Reserved.
+
+#include "Variables/AnimNextSharedVariables_EditorData.h"
+
+#include "Entries/AnimNextSharedVariablesEntry.h"
+#include "Entries/AnimNextVariableEntry.h"
+
+TConstArrayView<TSubclassOf<UAnimNextRigVMAssetEntry>> UAnimNextSharedVariables_EditorData::GetEntryClasses() const
+{
+	static const TSubclassOf<UAnimNextRigVMAssetEntry> Classes[] =
+	{
+		UAnimNextVariableEntry::StaticClass(),
+		UAnimNextSharedVariablesEntry::StaticClass(),
+	};
+
+	return Classes;
+}
+
+void UAnimNextSharedVariables_EditorData::CustomizeNewAssetEntry(UAnimNextRigVMAssetEntry* InNewEntry) const
+{
+	UAnimNextVariableEntry* VariableEntry = Cast<UAnimNextVariableEntry>(InNewEntry);
+	if(VariableEntry == nullptr)
+	{
+		return;
+	}
+	
+	const bool bIsSharedVariables = ExactCast<UAnimNextSharedVariables_EditorData>(this) != nullptr;
+	if(!bIsSharedVariables)
+	{
+		return;
+	}
+
+	// Force all variables in 'pure' shared variables assets to be public
+	VariableEntry->SetExportAccessSpecifier(EAnimNextExportAccessSpecifier::Public, false);
+}
\ No newline at end of file
diff --git a/Engine/Plugins/Experimental/UAF/UAF/Source/UAFUncookedOnly/Public/AnimNextScopedCompileJob.h b/Engine/Plugins/Experimental/UAF/UAF/Source/UAFUncookedOnly/Public/AnimNextScopedCompileJob.h
new file mode 100644
--- /dev/null
+++ b/Engine/Plugins/Experimental/UAF/UAF/Source/UAFUncookedOnly/Public/AnimNextScopedCompileJob.h
@@ -0,0 +1,37 @@
+﻿// Copyright Epic Games, Inc. All Rights Reserved.
+
+#pragma once
+
+#include "CoreMinimal.h"
+#include "Kismet2/CompilerResultsLog.h"
+
+class UAnimNextRigVMAssetEditorData;
+class FCompilerResultsLog;
+
+namespace UE::UAF::UncookedOnly
+{
+
+// RAII helper for scoping compilation batches
+// The re-allocation of compiled assets and their dependencies is deferred until the outermost scope exits.
+class FScopedCompileJob
+{
+public:
+	UAFUNCOOKEDONLY_API explicit FScopedCompileJob(const FText& InJobName, TConstArrayView<UObject*> InAssets);
+	UAFUNCOOKEDONLY_API explicit FScopedCompileJob(UObject* InObject);
+	UAFUNCOOKEDONLY_API ~FScopedCompileJob();
+
+	// Get the log that is currently in scope, asserts if no scope is active
+	UAFUNCOOKEDONLY_API static FCompilerResultsLog& GetLog();
+
+private:
+	UAFUNCOOKEDONLY_API FScopedCompileJob(const FText& InJobName, UObject* InObject, TConstArrayView<UObject*> InAssets);
+
+private:
+	TSharedPtr<FCompilerResultsLog> Log;
+	FText JobName;
+	UObject* Object = nullptr;
+	double StartTime = 0.0;
+	double FinishTime = 0.0;
+};
+
+}
