Fix up asset wizard default character template assets

#rnx

#virtualized

diff --git a/Engine/Plugins/Experimental/UAF/UAF/Source/UAFEditor/Private/Wizard/STemplatePicker.cpp b/Engine/Plugins/Experimental/UAF/UAF/Source/UAFEditor/Private/Wizard/STemplatePicker.cpp
--- a/Engine/Plugins/Experimental/UAF/UAF/Source/UAFEditor/Private/Wizard/STemplatePicker.cpp
+++ b/Engine/Plugins/Experimental/UAF/UAF/Source/UAFEditor/Private/Wizard/STemplatePicker.cpp
@@ -458,7 +458,7 @@ namespace UE::UAF::Editor
 
 		FTemplateItemPtr* DefaultSelectedTemplate = Templates.FindByPredicate([](const FTemplateItemPtr TemplateItem)
 		{
-			return TemplateItem->Template->GetPathName() == TEXT("/AnimNext/Templates/BasicCharacter/DA_DefaultCharacter.DA_DefaultCharacter");
+			return TemplateItem->Template->GetPathName() == TEXT("/UAF/Templates/BasicCharacter/DA_DefaultCharacter.DA_DefaultCharacter");
 		});
 
 		if (DefaultSelectedTemplate)
