More loctext fixes

diff --git a/Engine/Plugins/Experimental/UAF/UAF/Source/UAFEditor/Private/Common/GraphEditorSchemaActions.cpp b/Engine/Plugins/Experimental/UAF/UAF/Source/UAFEditor/Private/Common/GraphEditorSchemaActions.cpp
--- a/Engine/Plugins/Experimental/UAF/UAF/Source/UAFEditor/Private/Common/GraphEditorSchemaActions.cpp
+++ b/Engine/Plugins/Experimental/UAF/UAF/Source/UAFEditor/Private/Common/GraphEditorSchemaActions.cpp
@@ -476,7 +476,7 @@ UEdGraphNode* FAnimNextSchemaAction_PromoteToVariable::PerformAction(UEdGraph* P
 					const FAnimNextParamType ParamType = UE::UAF::UncookedOnly::FUtils::GetParamTypeFromPinType(FromPin->PinType);
 					if (ParamType.IsValid() && EditorData)
 					{
-						FScopedTransaction Transaction(LOCTEXT("PromoteToVariable", "Promote to Variable"));
+						FScopedTransaction Transaction(LOCTEXT("PromoteToVariableTransaction", "Promote to Variable"));
 						
 						FString BaseName = FromPin->GetDisplayName().ToString();
 						BaseName.ReplaceCharInline(' ', '_');
diff --git a/Engine/Plugins/Experimental/UAF/UAF/Source/UAFEditor/Private/Variables/Outliner/VariablesOutlinerAssetItem.cpp b/Engine/Plugins/Experimental/UAF/UAF/Source/UAFEditor/Private/Variables/Outliner/VariablesOutlinerAssetItem.cpp
--- a/Engine/Plugins/Experimental/UAF/UAF/Source/UAFEditor/Private/Variables/Outliner/VariablesOutlinerAssetItem.cpp
+++ b/Engine/Plugins/Experimental/UAF/UAF/Source/UAFEditor/Private/Variables/Outliner/VariablesOutlinerAssetItem.cpp
@@ -401,7 +401,7 @@ class SVariablesOutlinerAssetLabel : FSceneOutlinerCommonLabelData, public SComp
 							const FText VariableNameText = FText::FromName(VariableEntry->GetVariableName());
 							const FText AssetNameText = FText::FromName(Asset->GetFName());
 							
-							const FText FormattedMessage = FText::Format(LOCTEXT("MoveVariableToAssetFormat", "Moving {0} to {1}"), VariableNameText, AssetNameText);
+							const FText FormattedMessage = FText::Format(LOCTEXT("MoveVariableToAssetFormatTransaction", "Moving {0} to {1}"), VariableNameText, AssetNameText);
 							FScopedTransaction Transaction(FormattedMessage);
 							UncookedOnly::FUtils::MoveVariableToAsset(VariableEntry, Asset, true, true);
 
diff --git a/Engine/Plugins/Experimental/UAF/UAF/Source/UAFEditor/Private/Variables/Outliner/VariablesOutlinerEntryItem.cpp b/Engine/Plugins/Experimental/UAF/UAF/Source/UAFEditor/Private/Variables/Outliner/VariablesOutlinerEntryItem.cpp
--- a/Engine/Plugins/Experimental/UAF/UAF/Source/UAFEditor/Private/Variables/Outliner/VariablesOutlinerEntryItem.cpp
+++ b/Engine/Plugins/Experimental/UAF/UAF/Source/UAFEditor/Private/Variables/Outliner/VariablesOutlinerEntryItem.cpp
@@ -229,7 +229,7 @@ class SVariablesOutlinerEntryLabel : FSceneOutlinerCommonLabelData, public SComp
 								const FText VariableNameText = FText::FromName(OtherEntry->GetVariableName());
 								const FText AssetNameText = FText::FromName(ThisAsset->GetFName());
 																		
-								const FText FormattedMessage = FText::Format(LOCTEXT("MoveVariableToAssetFormat", "Moving {0} to {1}"), VariableNameText, AssetNameText);
+								const FText FormattedMessage = FText::Format(LOCTEXT("MoveVariableToAssetFormatTransaction", "Moving {0} to {1}"), VariableNameText, AssetNameText);
 								FScopedTransaction Transaction(FormattedMessage);
 								UncookedOnly::FUtils::MoveVariableToAsset(VariableEntry, ThisAsset, true, true);
 							}
