Add Initialize as a default implemented event in UAF module

diff --git a/Engine/Plugins/Experimental/UAF/UAF/Source/UAFEditor/Private/Module/AnimNextModuleFactory.cpp b/Engine/Plugins/Experimental/UAF/UAF/Source/UAFEditor/Private/Module/AnimNextModuleFactory.cpp
--- a/Engine/Plugins/Experimental/UAF/UAF/Source/UAFEditor/Private/Module/AnimNextModuleFactory.cpp
+++ b/Engine/Plugins/Experimental/UAF/UAF/Source/UAFEditor/Private/Module/AnimNextModuleFactory.cpp
@@ -5,6 +5,7 @@
 #include "Module/AnimNextModule_EditorData.h"
 #include "UncookedOnlyUtils.h"
 #include "Module/RigUnit_AnimNextModuleEvents.h"
+#include "Entries/AnimNextEventGraphEntry.h"
 
 UAnimNextModuleFactory::UAnimNextModuleFactory()
 {
@@ -34,7 +35,10 @@ UObject* UAnimNextModuleFactory::FactoryCreateNew(UClass* Class, UObject* InPare
 	EditorData->bUsesExternalPackages = false;
 	EditorData->Initialize(/*bRecompileVM*/false);
 
-	EditorData->AddEventGraph(FRigUnit_AnimNextPrePhysicsEvent::DefaultEventName, FRigUnit_AnimNextPrePhysicsEvent::StaticStruct());
+	UAnimNextEventGraphEntry* GraphEntry = EditorData->AddEventGraph(TEXT("EventGraph"), FRigUnit_AnimNextInitializeEvent::StaticStruct());
+
+	URigVMController* Controller = EditorData->GetController(GraphEntry->Graph);
+	Controller->AddUnitNode(FRigUnit_AnimNextPrePhysicsEvent::StaticStruct(), FRigVMStruct::ExecuteName, FVector2D(0.0f, 200.0f), FString(), false);
 
 	// Compile the initial skeleton
 	EditorData->RecompileVM();
diff --git a/Engine/Plugins/Experimental/UAF/UAF/Source/UAFUncookedOnly/Internal/Entries/AnimNextEventGraphEntry.h b/Engine/Plugins/Experimental/UAF/UAF/Source/UAFUncookedOnly/Internal/Entries/AnimNextEventGraphEntry.h
--- a/Engine/Plugins/Experimental/UAF/UAF/Source/UAFUncookedOnly/Internal/Entries/AnimNextEventGraphEntry.h
+++ b/Engine/Plugins/Experimental/UAF/UAF/Source/UAFUncookedOnly/Internal/Entries/AnimNextEventGraphEntry.h
@@ -35,6 +35,7 @@ class UAnimNextEventGraphEntry : public UAnimNextRigVMAssetEntry, public IAnimNe
 
 	friend class UAnimNextRigVMAssetEditorData;
 	friend class UAnimNextModule_EditorData;
+	friend class UAnimNextModuleFactory;
 	friend class UE::UAF::Tests::FEditor_Variables;
 	friend class UE::UAF::Tests::FVariables_UOLBindings;
 	friend class UE::UAF::Tests::FEditor_AnimGraph_Variables;
diff --git a/Engine/Plugins/Experimental/UAF/UAF/Source/UAFUncookedOnly/Private/UncookedOnlyUtils.cpp b/Engine/Plugins/Experimental/UAF/UAF/Source/UAFUncookedOnly/Private/UncookedOnlyUtils.cpp
--- a/Engine/Plugins/Experimental/UAF/UAF/Source/UAFUncookedOnly/Private/UncookedOnlyUtils.cpp
+++ b/Engine/Plugins/Experimental/UAF/UAF/Source/UAFUncookedOnly/Private/UncookedOnlyUtils.cpp
@@ -700,17 +700,17 @@ void FUtils::SetupEventGraph(URigVMController* InController, UScriptStruct* InEv
 		FRigUnit_AnimNextUserEvent Defaults;
 		Defaults.Name = InEventName;
 		Defaults.SortOrder = InEventName.GetNumber();
-		InController->AddUnitNodeWithDefaults(InEventStruct, Defaults, FRigVMStruct::ExecuteName, FVector2D(-200.0f, 0.0f), FString(), false);
+		InController->AddUnitNodeWithDefaults(InEventStruct, Defaults, FRigVMStruct::ExecuteName, FVector2D(0.0f, 0.0f), FString(), false);
 	}
 	else if (InEventStruct == FRigVMFunction_UserDefinedEvent::StaticStruct())
 	{
 		FRigVMFunction_UserDefinedEvent Defaults;
 		Defaults.EventName = InEventName;
-		InController->AddUnitNodeWithDefaults(InEventStruct, Defaults, FRigVMStruct::ExecuteName, FVector2D(-200.0f, 0.0f), FString(), false);
+		InController->AddUnitNodeWithDefaults(InEventStruct, Defaults, FRigVMStruct::ExecuteName, FVector2D(0.0f, 0.0f), FString(), false);
 	}
 	else
 	{
-		InController->AddUnitNode(InEventStruct, FRigVMStruct::ExecuteName, FVector2D(-200.0f, 0.0f), FString(), false);
+		InController->AddUnitNode(InEventStruct, FRigVMStruct::ExecuteName, FVector2D(0.0f, 0.0f), FString(), false);
 	}
 }
 
