Rename and merge FAnimNextEditorStyle and FAnimNextStyle into FUAFStyle and fix up missing files

#rnx

diff --git a/Engine/Plugins/Experimental/UAF/UAF/Source/UAFEditor/Internal/AnimNextStyle.h b/Engine/Plugins/Experimental/UAF/UAF/Source/UAFEditor/Internal/AnimNextStyle.h
deleted file mode 100644
diff --git a/Engine/Plugins/Experimental/UAF/UAF/Source/UAFEditor/Internal/UAFStyle.h b/Engine/Plugins/Experimental/UAF/UAF/Source/UAFEditor/Internal/UAFStyle.h
new file mode 100644
--- /dev/null
+++ b/Engine/Plugins/Experimental/UAF/UAF/Source/UAFEditor/Internal/UAFStyle.h
@@ -0,0 +1,29 @@
+﻿// Copyright Epic Games, Inc. All Rights Reserved.
+
+#pragma once
+
+#include "Styling/SlateStyle.h"
+
+class UAFEDITOR_API FUAFStyle final : public FSlateStyleSet
+{
+public:
+	static FUAFStyle& Get();
+
+protected:
+	static void Register();
+	static void Unregister();
+
+private:
+	FUAFStyle();
+
+	// Resolve path relative to plugin's resources directory
+	static FString InResources(const FString& RelativePath);
+
+	void SetClassIcons();
+
+	void SetVariablesOutlinerStyles();
+	
+	void SetTraitStyles();
+	
+	void SetAssetWizardStyles();
+};
diff --git a/Engine/Plugins/Experimental/UAF/UAF/Source/UAFEditor/Private/AnimNextEditorModule.cpp b/Engine/Plugins/Experimental/UAF/UAF/Source/UAFEditor/Private/AnimNextEditorModule.cpp
--- a/Engine/Plugins/Experimental/UAF/UAF/Source/UAFEditor/Private/AnimNextEditorModule.cpp
+++ b/Engine/Plugins/Experimental/UAF/UAF/Source/UAFEditor/Private/AnimNextEditorModule.cpp
@@ -34,7 +34,6 @@
 #include "Param/ObjectPropertyLocatorEditor.h"
 #include "AnimNextAssetWorkspaceAssetUserData.h"
 #include "AnimNextEditorContext.h"
-#include "AnimNextStyle.h"
 #include "AssetCompilationHandler.h"
 #include "ContentBrowserModule.h"
 #include "ContextObjectStore.h"
@@ -72,8 +71,8 @@
 #include "RewindDebugger/AnimNextModuleTrack.h"
 #include "RewindDebugger/AnimNextTraceModule.h"
 #include "Variables/VariableReferencePropertyCustomization.h"
-#include "AnimNextEditorStyle.h"
 #include "AnimNextSharedVariableNode.h"
+#include "UAFStyle.h"
 #include "Common/AnimNextAssetFindReplaceVariables.h"
 #include "Toolkits/AssetEditorToolkitMenuContext.h"
 #include "Variables/Outliner/VariablesOutlinerCommands.h"
@@ -101,7 +100,7 @@ void FAnimNextEditorModule::StartupModule()
 	FRigVMCommands::Register();
 	FAnimNextRigVMAssetCommands::Register();
 	FVariablesOutlinerCommands::Register();
-	FAnimNextEditorStyle::Get();
+	FUAFStyle::Get();
 
 	// Register settings for user editing
 	ISettingsModule& SettingsModule = FModuleManager::Get().LoadModuleChecked<ISettingsModule>("Settings");
diff --git a/Engine/Plugins/Experimental/UAF/UAF/Source/UAFEditor/Private/AnimNextEditorStyle.h b/Engine/Plugins/Experimental/UAF/UAF/Source/UAFEditor/Private/AnimNextEditorStyle.h
deleted file mode 100644
diff --git a/Engine/Plugins/Experimental/UAF/UAF/Source/UAFEditor/Private/AnimNextStyle.cpp b/Engine/Plugins/Experimental/UAF/UAF/Source/UAFEditor/Private/AnimNextStyle.cpp
deleted file mode 100644
diff --git a/Engine/Plugins/Experimental/UAF/UAF/Source/UAFEditor/Private/UAFStyle.cpp b/Engine/Plugins/Experimental/UAF/UAF/Source/UAFEditor/Private/UAFStyle.cpp
new file mode 100644
--- /dev/null
+++ b/Engine/Plugins/Experimental/UAF/UAF/Source/UAFEditor/Private/UAFStyle.cpp
@@ -0,0 +1,124 @@
+﻿// Copyright Epic Games, Inc. All Rights Reserved.
+
+#include "UAFStyle.h"
+#include "Brushes/SlateBoxBrush.h"
+#include "Styling/SlateStyleRegistry.h"
+#include "Brushes/SlateImageBrush.h"
+#include "Styling/CoreStyle.h"
+#include "Brushes/SlateRoundedBoxBrush.h"
+#include "Interfaces/IPluginManager.h"
+#include "Styling/SlateTypes.h"
+#include "Misc/Paths.h"
+#include "Styling/SlateStyleMacros.h"
+#include "Styling/StyleColors.h"
+
+#define IMAGE_PLUGIN_BRUSH_SVG( RelativePath, ... ) FSlateVectorImageBrush( FUAFStyle::InResources(RelativePath), __VA_ARGS__)
+#define PLUGIN_BOX_BRUSH_BRUSH( RelativePath, ... ) FSlateBoxBrush( FUAFStyle::InResources(RelativePath), __VA_ARGS__)
+
+FUAFStyle::FUAFStyle()
+	: FSlateStyleSet(TEXT("UAFStyle"))
+{
+	SetContentRoot(FPaths::EngineContentDir() / TEXT("Editor/Slate"));
+	SetCoreContentRoot(FPaths::EngineContentDir() / TEXT("Slate"));
+
+	const FString AnimNextPluginContentDir = FPaths::EnginePluginsDir() / TEXT("Experimental/AnimNext/Resources");
+
+	SetClassIcons();
+	SetVariablesOutlinerStyles();
+	SetTraitStyles();
+	SetAssetWizardStyles();
+
+	FSlateStyleRegistry::RegisterSlateStyle(*this);
+}
+
+void FUAFStyle::SetClassIcons()
+{
+	const FVector2D Icon16x16(16.0f, 16.0f);
+	
+	Set("ClassIcon.AnimNextSharedVariables", new IMAGE_PLUGIN_BRUSH_SVG("Icons/Archive_16.svg", Icon16x16));
+	Set("ClassIcon.AnimNextModule", new IMAGE_BRUSH_SVG("Starship/Common/Modules", Icon16x16 ));
+	Set("ClassIcon.AnimNextAnimationGraph", new IMAGE_BRUSH_SVG("Starship/Blueprints/icon_BlueprintEditor_EventGraph", Icon16x16));
+	Set("ClassIcon.AnimNextStateTree", new IMAGE_BRUSH_SVG(TEXT("UMG/TreeView"), Icon16x16));
+}
+
+void FUAFStyle::SetVariablesOutlinerStyles()
+{
+	Set("VariablesOutliner.SharedVariablesPill.Border", new PLUGIN_BOX_BRUSH_BRUSH("NamespaceBorder.png", FMargin(4.0f / 16.0f)));
+
+	const FTextBlockStyle NormalText = FAppStyle::GetWidgetStyle<FTextBlockStyle>("NormalText");
+	Set("VariablesOutliner.SharedVariablesPill.Text", FTextBlockStyle(NormalText)
+		.SetFont(DEFAULT_FONT("Regular", 8))
+		.SetColorAndOpacity(FSlateColor::UseForeground()));
+
+	const FVector2D Icon14x14(14.0f, 14.0f);
+	const FButtonStyle CloseButtonStyle = FButtonStyle()
+		.SetNormal(CORE_IMAGE_BRUSH_SVG("Starship/Common/close-small", Icon14x14, FStyleColors::Foreground))
+		.SetPressed(CORE_IMAGE_BRUSH_SVG("Starship/Common/close-small", Icon14x14, FStyleColors::Foreground))
+		.SetHovered(CORE_IMAGE_BRUSH_SVG("Starship/Common/close-small", Icon14x14, FStyleColors::ForegroundHover));
+	Set("VariablesOutliner.SharedVariablesPill.CloseButton", CloseButtonStyle);
+}
+
+void FUAFStyle::SetTraitStyles()
+{
+	const FTextBlockStyle& NormalText = FAppStyle::Get().GetWidgetStyle<FTextBlockStyle>("NormalText");
+	const FComboButtonStyle& ComboButtonStyle = FCoreStyle::Get().GetWidgetStyle<FComboButtonStyle>("ComboButton");
+	
+	// Trait details labels
+	Set("AnimNext.TraitDetails.Label", FTextBlockStyle(NormalText)
+		.SetFont(FAppStyle::GetFontStyle(TEXT("PropertyWindow.BoldFont")))
+		.SetFontSize(7));
+
+	Set("AnimNext.TraitDetails.Background", new FSlateRoundedBoxBrush(FStyleColors::Hover, 6.f));
+
+	// Add trait combo button
+	const FButtonStyle AddButtonStyle = FButtonStyle()
+		.SetNormal(FSlateRoundedBoxBrush(FStyleColors::Hover, 3.0f))
+		.SetHovered(FSlateRoundedBoxBrush(FStyleColors::Primary, 3.0f))
+		.SetPressed(FSlateRoundedBoxBrush(FStyleColors::Primary, 3.0f))
+		.SetNormalForeground(FStyleColors::Transparent)
+		.SetHoveredForeground(FStyleColors::Hover)
+		.SetPressedForeground(FStyleColors::Foreground)
+		.SetNormalPadding(FMargin(6, 2, 6, 2))
+		.SetPressedPadding(FMargin(6, 3, 6, 1));
+
+	Set("AnimNext.TraitDetails.AddButton", FComboButtonStyle(ComboButtonStyle).SetButtonStyle(AddButtonStyle));
+}
+
+void FUAFStyle::SetAssetWizardStyles()
+{
+	const FVector2D Icon16x16(16.0f, 16.0f);
+	const FVector2D Icon64x64(64.0f, 64.0f);
+	
+	Set("AssetWizard.OpenExternal", new FSlateVectorImageBrush(InResources("Icons/Wizard/OpenExternal.svg"),Icon16x16));
+	Set("AssetWizard.Template", new FSlateVectorImageBrush(InResources("Icons/Wizard/Template.svg"),Icon64x64));
+
+	Set("AssetWizard.TemplateTag.OuterBorder", new FSlateRoundedBoxBrush(FStyleColors::AccentGray, 4.0f));
+	Set("AssetWizard.TemplateTag.InnerBorder", new FSlateRoundedBoxBrush(FStyleColors::Recessed, 4.0f));
+
+	Set("AssetWizard.TemplateFolderBackground", new FSlateRoundedBoxBrush(EStyleColor::Panel, 4.0f));
+	Set("AssetWizard.TemplateFolderHoverBackground", new FSlateRoundedBoxBrush(FStyleColors::Hover, 4.0f));
+	Set("AssetWizard.TemplateFolderSelectedBackground", new FSlateRoundedBoxBrush(FStyleColors::Primary, 4.0f));
+	Set("AssetWizard.TemplateFolderHoverSelectedBackground", new FSlateRoundedBoxBrush(FStyleColors::PrimaryHover, 4.0f));
+}
+
+void FUAFStyle::Register()
+{
+	FSlateStyleRegistry::RegisterSlateStyle(Get());
+}
+
+void FUAFStyle::Unregister()
+{
+	FSlateStyleRegistry::UnRegisterSlateStyle(Get());
+}
+
+FUAFStyle& FUAFStyle::Get()
+{
+	static FUAFStyle Instance;
+	return Instance;
+}
+
+FString FUAFStyle::InResources(const FString& RelativePath)
+{
+	static const FString ResourcesDir = IPluginManager::Get().FindPlugin(TEXT("UAF"))->GetBaseDir() / TEXT("Resources");
+	return ResourcesDir / RelativePath;
+}
\ No newline at end of file
diff --git a/Engine/Plugins/Experimental/UAF/UAF/Source/UAFEditor/Private/Variables/Outliner/VariablesOutlinerAssetItem.cpp b/Engine/Plugins/Experimental/UAF/UAF/Source/UAFEditor/Private/Variables/Outliner/VariablesOutlinerAssetItem.cpp
--- a/Engine/Plugins/Experimental/UAF/UAF/Source/UAFEditor/Private/Variables/Outliner/VariablesOutlinerAssetItem.cpp
+++ b/Engine/Plugins/Experimental/UAF/UAF/Source/UAFEditor/Private/Variables/Outliner/VariablesOutlinerAssetItem.cpp
@@ -7,7 +7,7 @@
 #include "Misc/Attribute.h"
 #include "Misc/Optional.h"
 #include "SceneOutlinerStandaloneTypes.h"
-#include "AnimNextEditorStyle.h"
+#include "UAFStyle.h"
 #include "AnimNextRigVMAsset.h"
 #include "AnimNextRigVMAssetEditorData.h"
 #include "AssetToolsModule.h"
@@ -163,7 +163,7 @@ class SVariablesOutlinerAssetLabel : FSceneOutlinerCommonLabelData, public SComp
 				const FString SharedVariablesAssetName = SharedVariablesPath.GetAssetName();
 				const FLinearColor BackgroundColor = FLinearColor::MakeRandomSeededColor(GetTypeHash(SharedVariablesAssetName)).Desaturate(0.4f);
 
-				const FButtonStyle& CloseButtonStyle = FAnimNextEditorStyle::Get().GetWidgetStyle<FButtonStyle>("VariablesOutliner.SharedVariablesPill.CloseButton");
+				const FButtonStyle& CloseButtonStyle = FUAFStyle::Get().GetWidgetStyle<FButtonStyle>("VariablesOutliner.SharedVariablesPill.CloseButton");
 				TSharedPtr<SBorder> SharedVariablesBorderWidget;
 				ImplementedSharedVariablesHBox->AddSlot()
 				.VAlign(VAlign_Center)
@@ -171,7 +171,7 @@ class SVariablesOutlinerAssetLabel : FSceneOutlinerCommonLabelData, public SComp
 				.AutoWidth()
 				[
 					SAssignNew(SharedVariablesBorderWidget, SBorder)
-					.BorderImage(FAnimNextEditorStyle::Get().GetBrush("VariablesOutliner.SharedVariablesPill.Border"))
+					.BorderImage(FUAFStyle::Get().GetBrush("VariablesOutliner.SharedVariablesPill.Border"))
 					.BorderBackgroundColor(BackgroundColor)
 					.ToolTipText( FText::FromString(SharedVariablesAssetName))
 					.VAlign(VAlign_Center)
@@ -184,7 +184,7 @@ class SVariablesOutlinerAssetLabel : FSceneOutlinerCommonLabelData, public SComp
 						.VAlign(VAlign_Center)
 						[
 							SNew(STextBlock)
-							.TextStyle(FAnimNextEditorStyle::Get(), "VariablesOutliner.SharedVariablesPill.Text")
+							.TextStyle(FUAFStyle::Get(), "VariablesOutliner.SharedVariablesPill.Text")
 							.Text(FText::FromString(SharedVariablesAssetName))
 							.HighlightText(SceneOutliner.GetFilterHighlightText())
 							.OverflowPolicy(ETextOverflowPolicy::MiddleEllipsis)
diff --git a/Engine/Plugins/Experimental/UAF/UAF/Source/UAFEditor/Private/Variables/Outliner/VariablesOutlinerMode.cpp b/Engine/Plugins/Experimental/UAF/UAF/Source/UAFEditor/Private/Variables/Outliner/VariablesOutlinerMode.cpp
--- a/Engine/Plugins/Experimental/UAF/UAF/Source/UAFEditor/Private/Variables/Outliner/VariablesOutlinerMode.cpp
+++ b/Engine/Plugins/Experimental/UAF/UAF/Source/UAFEditor/Private/Variables/Outliner/VariablesOutlinerMode.cpp
@@ -10,7 +10,7 @@
 #include "AnimNextRigVMAssetEditorData.h"
 #include "IWorkspaceEditor.h"
 #include "AnimNextEditorModule.h"
-#include "AnimNextEditorStyle.h"
+#include "UAFStyle.h"
 #include "AssetRegistry/AssetRegistryModule.h"
 #include "ScopedTransaction.h"
 #include "ToolMenus.h"
@@ -794,7 +794,7 @@ void FVariablesOutlinerMode::PopulateNewVariableToolMenuEntries(UToolMenu* InMen
 					SubMenuBuilder.AddWidget(Widget.ToSharedRef(), FText::GetEmpty(), true);
 				}),
 				false,
-				FSlateIcon(FAnimNextEditorStyle::Get().GetStyleSetName(), "ClassIcon.AnimNextSharedVariables")
+				FSlateIcon(FUAFStyle::Get().GetStyleSetName(), "ClassIcon.AnimNextSharedVariables")
 			);
 
 
diff --git a/Engine/Plugins/Experimental/UAF/UAF/Source/UAFEditor/Private/Wizard/STemplatePicker.cpp b/Engine/Plugins/Experimental/UAF/UAF/Source/UAFEditor/Private/Wizard/STemplatePicker.cpp
--- a/Engine/Plugins/Experimental/UAF/UAF/Source/UAFEditor/Private/Wizard/STemplatePicker.cpp
+++ b/Engine/Plugins/Experimental/UAF/UAF/Source/UAFEditor/Private/Wizard/STemplatePicker.cpp
@@ -2,7 +2,7 @@
 
 #include "STemplatePicker.h"
 
-#include "AnimNextStyle.h"
+#include "UAFStyle.h"
 #include "AssetRegistry/AssetRegistryModule.h"
 #include "AssetThumbnail.h" 
 #include "DetailsViewArgs.h"
@@ -249,7 +249,7 @@ namespace UE::UAF::Editor
 									.Padding(FMargin(4, 0, 0, 0))
 									[
 										SNew(SImage)
-										.Image(FAnimNextStyle::Get().GetBrush("AssetWizard.OpenExternal"))
+										.Image(FUAFStyle::Get().GetBrush("AssetWizard.OpenExternal"))
 										.ColorAndOpacity(FLinearColor(0.2f, 0.4f, 1.0f))
 									]
 								]
@@ -544,12 +544,12 @@ namespace UE::UAF::Editor
 				.AutoWidth()
 				[
 					SNew(SBorder)
-					.BorderImage(FAnimNextStyle::Get().GetBrush("AssetWizard.TemplateTag.OuterBorder"))
+					.BorderImage(FUAFStyle::Get().GetBrush("AssetWizard.TemplateTag.OuterBorder"))
 					.BorderBackgroundColor(FLinearColor::White)
 					.Padding(1.0f)
 					[
 						SNew(SBorder)
-						.BorderImage(FAnimNextStyle::Get().GetBrush("AssetWizard.TemplateTag.InnerBorder"))
+						.BorderImage(FUAFStyle::Get().GetBrush("AssetWizard.TemplateTag.InnerBorder"))
 						.Padding(12.0f, 2.0f)
 						[
 							SNew(STextBlock)
@@ -619,7 +619,7 @@ namespace UE::UAF::Editor
 								.VAlign(VAlign_Center)
 								[
 									SNew(SImage)
-									.Image(FAnimNextStyle::Get().GetBrush("AssetWizard.Template"))
+									.Image(FUAFStyle::Get().GetBrush("AssetWizard.Template"))
 								]
 							]
 						]
diff --git a/Engine/Plugins/Experimental/UAF/UAFAnimGraph/Source/UAFAnimGraphEditor/Private/Factory/AnimNextFactoryParamsDetails.cpp b/Engine/Plugins/Experimental/UAF/UAFAnimGraph/Source/UAFAnimGraphEditor/Private/Factory/AnimNextFactoryParamsDetails.cpp
--- a/Engine/Plugins/Experimental/UAF/UAFAnimGraph/Source/UAFAnimGraphEditor/Private/Factory/AnimNextFactoryParamsDetails.cpp
+++ b/Engine/Plugins/Experimental/UAF/UAFAnimGraph/Source/UAFAnimGraphEditor/Private/Factory/AnimNextFactoryParamsDetails.cpp
@@ -2,7 +2,7 @@
 
 #include "AnimNextFactoryParamsDetails.h"
 
-#include "AnimNextStyle.h"
+#include "UAFStyle.h"
 #include "DetailLayoutBuilder.h"
 #include "DetailWidgetRow.h"
 #include "IDetailChildrenBuilder.h"
@@ -298,10 +298,10 @@ void FAnimNextFactoryParamsDetails::CustomizeChildren(TSharedRef<IPropertyHandle
 					[
 						SNew(SBorder)
 						.Padding(FMargin(6.0f, 2.0f))
-						.BorderImage(FAnimNextStyle::Get().GetBrush("AnimNext.TraitDetails.Background"))
+						.BorderImage(FUAFStyle::Get().GetBrush("AnimNext.TraitDetails.Background"))
 						[
 							SNew(STextBlock)
-							.TextStyle(FAnimNextStyle::Get(), "AnimNext.TraitDetails.Label")
+							.TextStyle(FUAFStyle::Get(), "AnimNext.TraitDetails.Label")
 							.ColorAndOpacity(FStyleColors::Foreground)
 							.Text(CommonType.GetValue()->GetDisplayNameText())
 						]
@@ -367,10 +367,10 @@ void FAnimNextFactoryParamsDetails::CustomizeChildren(TSharedRef<IPropertyHandle
 				[
 					SNew(SBorder)
 					.Padding(FMargin(6.0f, 2.0f))
-					.BorderImage(FAnimNextStyle::Get().GetBrush("AnimNext.TraitDetails.Background"))
+					.BorderImage(FUAFStyle::Get().GetBrush("AnimNext.TraitDetails.Background"))
 					[
 						SNew(STextBlock)
-						.TextStyle(FAnimNextStyle::Get(), "AnimNext.TraitDetails.Label")
+						.TextStyle(FUAFStyle::Get(), "AnimNext.TraitDetails.Label")
 						.ColorAndOpacity(FStyleColors::Foreground)
 						.Text(DefaultStruct.GetScriptStruct()->GetDisplayNameText())
 					]
@@ -412,7 +412,7 @@ void FAnimNextFactoryParamsDetails::CustomizeChildren(TSharedRef<IPropertyHandle
 			SNew(SComboButton)
 			.ToolTipText(LOCTEXT("AddTraitTooltip", "Add a trait to modify how this asset is played"))
 			.HasDownArrow(false)
-			.ComboButtonStyle(FAnimNextStyle::Get(), "AnimNext.TraitDetails.AddButton")
+			.ComboButtonStyle(FUAFStyle::Get(), "AnimNext.TraitDetails.AddButton")
 			.ButtonContent()
 			[
 				SNew(SHorizontalBox)
@@ -430,7 +430,7 @@ void FAnimNextFactoryParamsDetails::CustomizeChildren(TSharedRef<IPropertyHandle
 				.Padding(0.0f, 2.0f)
 				[
 					SNew(STextBlock)
-					.TextStyle(FAnimNextStyle::Get(), "AnimNext.TraitDetails.Label")
+					.TextStyle(FUAFStyle::Get(), "AnimNext.TraitDetails.Label")
 					.ColorAndOpacity(FStyleColors::Foreground)
 					.Text(LOCTEXT("AddTraitButton", "Add"))
 				]
