Fix module unit test due to changes to module factory

#rnx

diff --git a/Engine/Plugins/Experimental/UAF/UAF/Source/UAFEditor/Private/Module/AnimNextModuleFactory.cpp b/Engine/Plugins/Experimental/UAF/UAF/Source/UAFEditor/Private/Module/AnimNextModuleFactory.cpp
--- a/Engine/Plugins/Experimental/UAF/UAF/Source/UAFEditor/Private/Module/AnimNextModuleFactory.cpp
+++ b/Engine/Plugins/Experimental/UAF/UAF/Source/UAFEditor/Private/Module/AnimNextModuleFactory.cpp
@@ -35,7 +35,7 @@ UObject* UAnimNextModuleFactory::FactoryCreateNew(UClass* Class, UObject* InPare
 	EditorData->bUsesExternalPackages = false;
 	EditorData->Initialize(/*bRecompileVM*/false);
 
-	UAnimNextEventGraphEntry* GraphEntry = EditorData->AddEventGraph(TEXT("EventGraph"), FRigUnit_AnimNextInitializeEvent::StaticStruct());
+	UAnimNextEventGraphEntry* GraphEntry = EditorData->AddEventGraph(UAnimNextModule_EditorData::DefaultEventGraphName, FRigUnit_AnimNextInitializeEvent::StaticStruct());
 
 	URigVMController* Controller = EditorData->GetController(GraphEntry->Graph);
 	Controller->AddUnitNode(FRigUnit_AnimNextPrePhysicsEvent::StaticStruct(), FRigVMStruct::ExecuteName, FVector2D(0.0f, 200.0f), FString(), false);
diff --git a/Engine/Plugins/Experimental/UAF/UAF/Source/UAFTestSuite/Private/AnimNextEditorTest.cpp b/Engine/Plugins/Experimental/UAF/UAF/Source/UAFTestSuite/Private/AnimNextEditorTest.cpp
--- a/Engine/Plugins/Experimental/UAF/UAF/Source/UAFTestSuite/Private/AnimNextEditorTest.cpp
+++ b/Engine/Plugins/Experimental/UAF/UAF/Source/UAFTestSuite/Private/AnimNextEditorTest.cpp
@@ -10,6 +10,7 @@
 #include "Entries/AnimNextEventGraphEntry.h"
 #include "Module/AnimNextModule.h"
 #include "Module/AnimNextModuleFactory.h"
+#include "Module/AnimNextModule_EditorData.h"
 #include "Module/RigUnit_AnimNextModuleEvents.h"
 #if WITH_EDITOR
 #include "ScopedTransaction.h"
@@ -161,17 +162,17 @@ bool FEditor_Graphs::RunTest(const FString& InParameters)
 			UAnimNextEventGraphEntry* EventGraphEntry = nullptr;
 			if(TestSetting.bExpectExistingEventGraph)
 			{
-				EventGraphEntry = Cast<UAnimNextEventGraphEntry>(EditorData->FindEntry(TEXT("PrePhysics")));
+				EventGraphEntry = Cast<UAnimNextEventGraphEntry>(EditorData->FindEntry(UAnimNextModule_EditorData::DefaultEventGraphName));
 				UE_RETURN_ON_ERROR(EventGraphEntry != nullptr, TEXT("Could not find existing event graph."));
 			}
 			else
 			{
-				EventGraphEntry = EditorData->AddEventGraph(TEXT("PrePhysics"), FRigUnit_AnimNextPrePhysicsEvent::StaticStruct());
+				EventGraphEntry = EditorData->AddEventGraph(UAnimNextModule_EditorData::DefaultEventGraphName, FRigUnit_AnimNextPrePhysicsEvent::StaticStruct());
 				UE_RETURN_ON_ERROR(EventGraphEntry != nullptr, TEXT("Could not add event graph."));
 			}
 
 			URigVMGraph* RigVMGraph = EventGraphEntry->GetRigVMGraph();
-			UE_RETURN_ON_ERROR(RigVMGraph->GetNodes().Num() == 1, TEXT("Unexpected number of nodes in new event graph."));
+			UE_RETURN_ON_ERROR(RigVMGraph->GetNodes().Num() == 2, TEXT("Unexpected number of nodes in new event graph.")); // Initialize & PrePhysics
 
 			{
 				FScopedTransaction Transaction(FText::GetEmpty());
@@ -181,13 +182,13 @@ bool FEditor_Graphs::RunTest(const FString& InParameters)
 
 			GEditor->UndoTransaction();
 
-			UAnimNextEventGraphEntry* FoundEventGraphEntry = Cast<UAnimNextEventGraphEntry>(EditorData->FindEntry(TEXT("PrePhysics")));
+			UAnimNextEventGraphEntry* FoundEventGraphEntry = Cast<UAnimNextEventGraphEntry>(EditorData->FindEntry(UAnimNextModule_EditorData::DefaultEventGraphName));
 			UE_RETURN_ON_ERROR(FoundEventGraphEntry != nullptr, "FEditor_Graphs -> Could not find event graph post-undo.");
 		}
 		else
 		{
 			AddExpectedError(TEXT("Cannot add an event graph to this asset - entry is not allowed"));
-			UAnimNextEventGraphEntry* EventGraphEntry = EditorData->AddEventGraph(TEXT("PrePhysics"), FRigUnit_AnimNextPrePhysicsEvent::StaticStruct());
+			UAnimNextEventGraphEntry* EventGraphEntry = EditorData->AddEventGraph(UAnimNextModule_EditorData::DefaultEventGraphName, FRigUnit_AnimNextPrePhysicsEvent::StaticStruct());
 		}
 	}
 
diff --git a/Engine/Plugins/Experimental/UAF/UAF/Source/UAFUncookedOnly/Internal/Module/AnimNextModule_EditorData.h b/Engine/Plugins/Experimental/UAF/UAF/Source/UAFUncookedOnly/Internal/Module/AnimNextModule_EditorData.h
--- a/Engine/Plugins/Experimental/UAF/UAF/Source/UAFUncookedOnly/Internal/Module/AnimNextModule_EditorData.h
+++ b/Engine/Plugins/Experimental/UAF/UAF/Source/UAFUncookedOnly/Internal/Module/AnimNextModule_EditorData.h
@@ -49,6 +49,9 @@ class UAnimNextModule_EditorData : public UAnimNextSharedVariables_EditorData
 	friend class FAnimationAnimNextRuntimeTest_GraphExecuteLatent;
 	friend class FAnimationAnimNextEditorTest_GraphManifest;
 
+public:
+	static inline FName DefaultEventGraphName = TEXT("EventGraph");
+
 private:
 	// UAnimNextRigVMAssetEditorData interface
 	virtual UScriptStruct* GetExecuteContextStruct() const override { return FAnimNextExecuteContext::StaticStruct(); }
