Fix variables unit test failing due to module factory changes

#rnx

diff --git a/Engine/Plugins/Experimental/UAF/UAF/Source/UAFTestSuite/Private/AnimNextVariablesTest.cpp b/Engine/Plugins/Experimental/UAF/UAF/Source/UAFTestSuite/Private/AnimNextVariablesTest.cpp
--- a/Engine/Plugins/Experimental/UAF/UAF/Source/UAFTestSuite/Private/AnimNextVariablesTest.cpp
+++ b/Engine/Plugins/Experimental/UAF/UAF/Source/UAFTestSuite/Private/AnimNextVariablesTest.cpp
@@ -243,19 +243,32 @@ bool FVariables::RunTest(const FString& InParameters)
 	UE_RETURN_ON_ERROR(ResultEntry != nullptr, TEXT("Could not create new variable in graph."));
 
 	// Get event graph
-	UAnimNextEventGraphEntry* EventGraph = Cast<UAnimNextEventGraphEntry>(EditorData->FindEntry(TEXT("PrePhysics")));
+	UAnimNextEventGraphEntry* EventGraph = Cast<UAnimNextEventGraphEntry>(EditorData->FindEntry(UAnimNextModule_EditorData::DefaultEventGraphName));
+	URigVMGraph* RigVMGraph = nullptr;
+	URigVMNode* PrePhysicsEventNode = nullptr;
+	
 	if(EventGraph == nullptr)
 	{
-		EventGraph = EditorData->AddEventGraph(TEXT("PrePhysics"), FRigUnit_AnimNextPrePhysicsEvent::StaticStruct());
+		EventGraph = EditorData->AddEventGraph(UAnimNextModule_EditorData::DefaultEventGraphName, FRigUnit_AnimNextPrePhysicsEvent::StaticStruct());
+		UE_RETURN_ON_ERROR(EventGraph != nullptr, TEXT("Could not create new event graph in asset."));
+
+		RigVMGraph = EventGraph->GetRigVMGraph();
+		UE_RETURN_ON_ERROR(RigVMGraph->GetNodes().Num() == 1, TEXT("Unexpected number of nodes in new event graph.")); // PrePhysics
+
+		PrePhysicsEventNode = RigVMGraph->GetNodes()[0];
+		UE_RETURN_ON_ERROR(PrePhysicsEventNode->GetEventName() == "PrePhysics", TEXT("Failed to find PrePhysics event node."));
 	}
-	UE_RETURN_ON_ERROR(EventGraph != nullptr, TEXT("Could not create new event graph in asset."));
-	
-	URigVMGraph* RigVMGraph = EventGraph->GetRigVMGraph();
-	UE_RETURN_ON_ERROR(RigVMGraph->GetNodes().Num() == 1, TEXT("Unexpected number of nodes in new event graph."));
+	else
+	{
+		RigVMGraph = EventGraph->GetRigVMGraph();
+		UE_RETURN_ON_ERROR(RigVMGraph->GetNodes().Num() == 2, TEXT("Unexpected number of nodes in new event graph."));  // Initialize & PrePhysics
 
-	URigVMNode* EventNode = RigVMGraph->GetNodes()[0];
-	check(EventNode);
-	URigVMPin* ExecutePin = EventNode->FindPin("ExecuteContext");
+		PrePhysicsEventNode = RigVMGraph->GetNodes()[1];
+		UE_RETURN_ON_ERROR(PrePhysicsEventNode->GetEventName() == "PrePhysics", TEXT("Failed to find PrePhysics event node."));
+	}
+	
+	check(PrePhysicsEventNode);
+	URigVMPin* ExecutePin = PrePhysicsEventNode->FindPin("ExecuteContext");
 	UE_RETURN_ON_ERROR(ExecutePin != nullptr, TEXT("Could find initial execute pin."));
 
 	URigVMController* Controller = EditorData->GetController(EventGraph->GetRigVMGraph());
@@ -275,7 +288,7 @@ bool FVariables::RunTest(const FString& InParameters)
 	bool bLinkResultAdded = Controller->AddLink(TestOpUnitNode->FindPin("Result"), SetResultNode->FindPin("Value"));
 	UE_RETURN_ON_ERROR(bLinkResultAdded, TEXT("Could not link variable node."));
 
-	bool bLinkExec1Added = Controller->AddLink(EventNode->FindPin(FRigVMStruct::ExecuteContextName.ToString()), TestOpUnitNode->FindPin(FRigVMStruct::ExecuteContextName.ToString()));
+	bool bLinkExec1Added = Controller->AddLink(PrePhysicsEventNode->FindPin(FRigVMStruct::ExecuteContextName.ToString()), TestOpUnitNode->FindPin(FRigVMStruct::ExecuteContextName.ToString()));
 	UE_RETURN_ON_ERROR(bLinkExec1Added, TEXT("Could not link variable node exec."));
 
 	bool bLinkExec2Added = Controller->AddLink(TestOpUnitNode->FindPin(FRigVMStruct::ExecuteContextName.ToString()), SetResultNode->FindPin(FRigVMStruct::ExecuteContextName.ToString()));
@@ -358,19 +371,32 @@ bool FVariables_UOLBindings::RunTest(const FString& InParameters)
 	HoistedAccessorEntry->SetBinding(MoveTemp(BindingC));
 
 	// Get event graph
-	UAnimNextEventGraphEntry* EventGraph = Cast<UAnimNextEventGraphEntry>(EditorData->FindEntry(TEXT("PrePhysics")));
+	UAnimNextEventGraphEntry* EventGraph = Cast<UAnimNextEventGraphEntry>(EditorData->FindEntry(UAnimNextModule_EditorData::DefaultEventGraphName));
+	URigVMGraph* RigVMGraph = nullptr;
+	URigVMNode* PrePhysicsEventNode = nullptr;
+	
 	if(EventGraph == nullptr)
 	{
-		EventGraph = EditorData->AddEventGraph(TEXT("PrePhysics"), FRigUnit_AnimNextPrePhysicsEvent::StaticStruct());
+		EventGraph = EditorData->AddEventGraph(UAnimNextModule_EditorData::DefaultEventGraphName, FRigUnit_AnimNextPrePhysicsEvent::StaticStruct());
+		UE_RETURN_ON_ERROR(EventGraph != nullptr, TEXT("Could not create new event graph in asset."));
+
+		RigVMGraph = EventGraph->GetRigVMGraph();
+		UE_RETURN_ON_ERROR(RigVMGraph->GetNodes().Num() == 1, TEXT("Unexpected number of nodes in new event graph.")); // PrePhysics
+
+		PrePhysicsEventNode = RigVMGraph->GetNodes()[0];
+		UE_RETURN_ON_ERROR(PrePhysicsEventNode->GetEventName() == "PrePhysics", TEXT("Failed to find PrePhysics event node."));
 	}
-	UE_RETURN_ON_ERROR(EventGraph != nullptr, TEXT("Could not create new event graph in asset."));
-	
-	URigVMGraph* RigVMGraph = EventGraph->GetRigVMGraph();
-	UE_RETURN_ON_ERROR(RigVMGraph->GetNodes().Num() == 1, TEXT("Unexpected number of nodes in new event graph."));
+	else
+	{
+		RigVMGraph = EventGraph->GetRigVMGraph();
+		UE_RETURN_ON_ERROR(RigVMGraph->GetNodes().Num() == 2, TEXT("Unexpected number of nodes in new event graph."));  // Initialize & PrePhysics
 
-	URigVMNode* EventNode = RigVMGraph->GetNodes()[0];
-	check(EventNode);
-	URigVMPin* ExecutePin = EventNode->FindPin("ExecuteContext");
+		PrePhysicsEventNode = RigVMGraph->GetNodes()[1];
+		UE_RETURN_ON_ERROR(PrePhysicsEventNode->GetEventName() == "PrePhysics", TEXT("Failed to find PrePhysics event node."));
+	}
+	
+	check(PrePhysicsEventNode);
+	URigVMPin* ExecutePin = PrePhysicsEventNode->FindPin("ExecuteContext");
 	UE_RETURN_ON_ERROR(ExecutePin != nullptr, TEXT("Could find initial execute pin."));
 
 	URigVMController* Controller = EditorData->GetController(EventGraph->GetRigVMGraph());
@@ -386,7 +412,7 @@ bool FVariables_UOLBindings::RunTest(const FString& InParameters)
 	bool bLinkAAdded = Controller->AddLink(PropertyEntryNode->FindPin("Value"), PrintPropertyUnitNode->FindPin("Result"));
 	UE_RETURN_ON_ERROR(bLinkAAdded, TEXT("Could not link variable node."));
 
-	bool bLinkExec1Added = Controller->AddLink(EventNode->FindPin(FRigVMStruct::ExecuteContextName.ToString()), PrintPropertyUnitNode->FindPin(FRigVMStruct::ExecuteContextName.ToString()));
+	bool bLinkExec1Added = Controller->AddLink(PrePhysicsEventNode->FindPin(FRigVMStruct::ExecuteContextName.ToString()), PrintPropertyUnitNode->FindPin(FRigVMStruct::ExecuteContextName.ToString()));
 	UE_RETURN_ON_ERROR(bLinkExec1Added, TEXT("Could not link variable node exec."));
 	
 	URigVMUnitNode* PrintObjectAccessorUnitNode = Controller->AddUnitNode(FAnimNextTests_PrintResult::StaticStruct());
