Run UnrealCodeFixup to add #include UE_INLINE_GENERATED_CPP_BY_NAME

diff --git a/Engine/Plugins/Animation/IKRig/Source/IKRig/Private/Retargeter/RetargetOps/StretchChainOp.cpp b/Engine/Plugins/Animation/IKRig/Source/IKRig/Private/Retargeter/RetargetOps/StretchChainOp.cpp
--- a/Engine/Plugins/Animation/IKRig/Source/IKRig/Private/Retargeter/RetargetOps/StretchChainOp.cpp
+++ b/Engine/Plugins/Animation/IKRig/Source/IKRig/Private/Retargeter/RetargetOps/StretchChainOp.cpp
@@ -1,10 +1,12 @@
-﻿// Copyright Epic Games, Inc. All Rights Reserved.
+// Copyright Epic Games, Inc. All Rights Reserved.
 
 #include "Retargeter/RetargetOps/StretchChainOp.h"
 
 #include "Engine/SkeletalMesh.h"
 #include "Retargeter/IKRetargetProcessor.h"
 
+#include UE_INLINE_GENERATED_CPP_BY_NAME(StretchChainOp)
+
 #define LOCTEXT_NAMESPACE "StretchChainOp"
 
 bool FChainStretcher::Initialize(
diff --git a/Engine/Plugins/Cameras/GameplayCameras/Source/GameplayCameras/Private/Nodes/Input/DrivenControlRotationCameraNode.cpp b/Engine/Plugins/Cameras/GameplayCameras/Source/GameplayCameras/Private/Nodes/Input/DrivenControlRotationCameraNode.cpp
--- a/Engine/Plugins/Cameras/GameplayCameras/Source/GameplayCameras/Private/Nodes/Input/DrivenControlRotationCameraNode.cpp
+++ b/Engine/Plugins/Cameras/GameplayCameras/Source/GameplayCameras/Private/Nodes/Input/DrivenControlRotationCameraNode.cpp
@@ -12,6 +12,8 @@
 #include "Debug/CameraDebugRenderer.h"
 #include "GameFramework/PlayerController.h"
 
+#include UE_INLINE_GENERATED_CPP_BY_NAME(DrivenControlRotationCameraNode)
+
 namespace UE::Cameras
 {
 
diff --git a/Engine/Plugins/Cameras/GameplayCameras/Source/GameplayCameras/Private/Nodes/Input/InputAccumulator2DCameraNode.cpp b/Engine/Plugins/Cameras/GameplayCameras/Source/GameplayCameras/Private/Nodes/Input/InputAccumulator2DCameraNode.cpp
--- a/Engine/Plugins/Cameras/GameplayCameras/Source/GameplayCameras/Private/Nodes/Input/InputAccumulator2DCameraNode.cpp
+++ b/Engine/Plugins/Cameras/GameplayCameras/Source/GameplayCameras/Private/Nodes/Input/InputAccumulator2DCameraNode.cpp
@@ -4,6 +4,8 @@
 
 #include "Nodes/Input/CameraRigInput2DSlot.h"
 
+#include UE_INLINE_GENERATED_CPP_BY_NAME(InputAccumulator2DCameraNode)
+
 #define LOCTEXT_NAMESPACE "InputAccumulator2DCameraNode"
 
 namespace UE::Cameras
diff --git a/Engine/Plugins/Cameras/GameplayCameras/Source/GameplayCameras/Private/Nodes/Input/RawInputAxisBinding2DCameraNode.cpp b/Engine/Plugins/Cameras/GameplayCameras/Source/GameplayCameras/Private/Nodes/Input/RawInputAxisBinding2DCameraNode.cpp
--- a/Engine/Plugins/Cameras/GameplayCameras/Source/GameplayCameras/Private/Nodes/Input/RawInputAxisBinding2DCameraNode.cpp
+++ b/Engine/Plugins/Cameras/GameplayCameras/Source/GameplayCameras/Private/Nodes/Input/RawInputAxisBinding2DCameraNode.cpp
@@ -7,6 +7,8 @@
 #include "EnhancedInputComponent.h"
 #include "Nodes/Input/InputAxisBindingHelpers.h"
 
+#include UE_INLINE_GENERATED_CPP_BY_NAME(RawInputAxisBinding2DCameraNode)
+
 namespace UE::Cameras
 {
 
diff --git a/Engine/Plugins/Experimental/ChaosMover/Source/ChaosMover/Private/Backends/NetworkChaosMoverData.cpp b/Engine/Plugins/Experimental/ChaosMover/Source/ChaosMover/Private/Backends/NetworkChaosMoverData.cpp
--- a/Engine/Plugins/Experimental/ChaosMover/Source/ChaosMover/Private/Backends/NetworkChaosMoverData.cpp
+++ b/Engine/Plugins/Experimental/ChaosMover/Source/ChaosMover/Private/Backends/NetworkChaosMoverData.cpp
@@ -5,6 +5,8 @@
 #include "ChaosMover/Backends/ChaosMoverBackend.h"
 #include "ChaosMover/ChaosMoverSimulation.h"
 
+#include UE_INLINE_GENERATED_CPP_BY_NAME(NetworkChaosMoverData)
+
 void FNetworkChaosMoverInputData::ApplyData(UActorComponent* NetworkComponent) const
 {
 	if (UChaosMoverBackendComponent* BackendComp = Cast<UChaosMoverBackendComponent>(NetworkComponent))
@@ -137,4 +139,4 @@ const FString FNetworkChaosMoverStateData::DebugData()
 	FAnsiStringBuilderBase StringBuilder;
 	SyncState.ToString(StringBuilder);
 	return FString::Printf(TEXT("FNetworkChaosMoverStateData:\n%hs"), StringBuilder.ToString());
-}
\ No newline at end of file
+}
diff --git a/Engine/Plugins/Experimental/ChaosMover/Source/ChaosMover/Private/ChaosMoverSimulationTypes.cpp b/Engine/Plugins/Experimental/ChaosMover/Source/ChaosMover/Private/ChaosMoverSimulationTypes.cpp
--- a/Engine/Plugins/Experimental/ChaosMover/Source/ChaosMover/Private/ChaosMoverSimulationTypes.cpp
+++ b/Engine/Plugins/Experimental/ChaosMover/Source/ChaosMover/Private/ChaosMoverSimulationTypes.cpp
@@ -3,6 +3,8 @@
 #include "ChaosMover/ChaosMoverSimulationTypes.h"
 #include "ChaosMover/ChaosMoverSimulation.h"
 
+#include UE_INLINE_GENERATED_CPP_BY_NAME(ChaosMoverSimulationTypes)
+
 void FChaosMoverSimulationDefaultInputs::Reset()
 {
 	CollisionResponseParams = FCollisionResponseParams();
@@ -22,4 +24,4 @@ FMoverDataCollection& UE::ChaosMover::GetDebugSimData(UChaosMoverSimulation* Sim
 {
 	check(Simulation);
 	return Simulation->GetDebugSimData();
-}
\ No newline at end of file
+}
diff --git a/Engine/Plugins/Experimental/ChaosMover/Source/ChaosMover/Private/Character/ChaosCharacterInputs.cpp b/Engine/Plugins/Experimental/ChaosMover/Source/ChaosMover/Private/Character/ChaosCharacterInputs.cpp
--- a/Engine/Plugins/Experimental/ChaosMover/Source/ChaosMover/Private/Character/ChaosCharacterInputs.cpp
+++ b/Engine/Plugins/Experimental/ChaosMover/Source/ChaosMover/Private/Character/ChaosCharacterInputs.cpp
@@ -2,6 +2,8 @@
 
 #include "ChaosMover/Character/ChaosCharacterInputs.h"
 
+#include UE_INLINE_GENERATED_CPP_BY_NAME(ChaosCharacterInputs)
+
 bool FChaosMoverLaunchInputs::NetSerialize(FArchive& Ar, UPackageMap* Map, bool& bOutSuccess)
 {
 	Super::NetSerialize(Ar, Map, bOutSuccess);
@@ -166,4 +168,4 @@ void FChaosMovementSettingsOverridesRemover::Interpolate(const FMoverDataStructB
 
 void FChaosMovementSettingsOverridesRemover::Merge(const FMoverDataStructBase& From)
 {
-}
\ No newline at end of file
+}
diff --git a/Engine/Plugins/Experimental/ChaosMover/Source/ChaosMover/Private/Character/ChaosCharacterMoverComponent.cpp b/Engine/Plugins/Experimental/ChaosMover/Source/ChaosMover/Private/Character/ChaosCharacterMoverComponent.cpp
--- a/Engine/Plugins/Experimental/ChaosMover/Source/ChaosMover/Private/Character/ChaosCharacterMoverComponent.cpp
+++ b/Engine/Plugins/Experimental/ChaosMover/Source/ChaosMover/Private/Character/ChaosCharacterMoverComponent.cpp
@@ -11,6 +11,8 @@
 #include "DefaultMovementSet/CharacterMoverSimulationTypes.h"
 #include "GameFramework/Pawn.h"
 
+#include UE_INLINE_GENERATED_CPP_BY_NAME(ChaosCharacterMoverComponent)
+
 UChaosCharacterMoverComponent::UChaosCharacterMoverComponent()
 {
 	// Default movement modes
@@ -180,4 +182,4 @@ void UChaosCharacterMoverComponent::CancelMovementSettingsOverrides(FName ModeNa
 {
 	ModeToOverrideSettings = ModeName;
 	bCancelMovementOverrides = true;
-}
\ No newline at end of file
+}
diff --git a/Engine/Plugins/Experimental/ChaosMover/Source/ChaosMover/Private/Character/Effects/ChaosCharacterApplyVelocityEffect.cpp b/Engine/Plugins/Experimental/ChaosMover/Source/ChaosMover/Private/Character/Effects/ChaosCharacterApplyVelocityEffect.cpp
--- a/Engine/Plugins/Experimental/ChaosMover/Source/ChaosMover/Private/Character/Effects/ChaosCharacterApplyVelocityEffect.cpp
+++ b/Engine/Plugins/Experimental/ChaosMover/Source/ChaosMover/Private/Character/Effects/ChaosCharacterApplyVelocityEffect.cpp
@@ -6,6 +6,8 @@
 #include "ChaosMover/ChaosMoverLog.h"
 #include "ChaosMover/ChaosMoverSimulation.h"
 
+#include UE_INLINE_GENERATED_CPP_BY_NAME(ChaosCharacterApplyVelocityEffect)
+
 bool FChaosCharacterApplyVelocityEffect::ApplyMovementEffect_Async(FApplyMovementEffectParams_Async& ApplyEffectParams, FMoverSyncState& OutputState)
 {
 	check(ApplyEffectParams.StartState);
@@ -91,4 +93,4 @@ FString FChaosCharacterApplyVelocityEffect::ToSimpleString() const
 void FChaosCharacterApplyVelocityEffect::AddReferencedObjects(FReferenceCollector& Collector)
 {
 	Super::AddReferencedObjects(Collector);
-}
\ No newline at end of file
+}
diff --git a/Engine/Plugins/Experimental/ChaosMover/Source/ChaosMover/Private/Character/Modifiers/ChaosStanceModifier.cpp b/Engine/Plugins/Experimental/ChaosMover/Source/ChaosMover/Private/Character/Modifiers/ChaosStanceModifier.cpp
--- a/Engine/Plugins/Experimental/ChaosMover/Source/ChaosMover/Private/Character/Modifiers/ChaosStanceModifier.cpp
+++ b/Engine/Plugins/Experimental/ChaosMover/Source/ChaosMover/Private/Character/Modifiers/ChaosStanceModifier.cpp
@@ -13,6 +13,8 @@
 #include "PhysicsEngine/PhysicsObjectExternalInterface.h"
 #include "Engine/World.h"
 
+#include UE_INLINE_GENERATED_CPP_BY_NAME(ChaosStanceModifier)
+
 FChaosStanceModifier::FChaosStanceModifier()
 	: ActiveStance(EStanceMode::Invalid)
 {
@@ -247,4 +249,4 @@ FString FChaosStanceModifier::ToSimpleString() const
 void FChaosStanceModifier::AddReferencedObjects(class FReferenceCollector& Collector)
 {
 	Super::AddReferencedObjects(Collector);
-}
\ No newline at end of file
+}
diff --git a/Engine/Plugins/Experimental/ChaosMover/Source/ChaosMover/Private/Character/Transitions/ChaosCharacterLaunchCheck.cpp b/Engine/Plugins/Experimental/ChaosMover/Source/ChaosMover/Private/Character/Transitions/ChaosCharacterLaunchCheck.cpp
--- a/Engine/Plugins/Experimental/ChaosMover/Source/ChaosMover/Private/Character/Transitions/ChaosCharacterLaunchCheck.cpp
+++ b/Engine/Plugins/Experimental/ChaosMover/Source/ChaosMover/Private/Character/Transitions/ChaosCharacterLaunchCheck.cpp
@@ -11,6 +11,8 @@
 #include "Chaos/PhysicsObjectInternalInterface.h"
 #include "Chaos/PhysicsObject.h"
 
+#include UE_INLINE_GENERATED_CPP_BY_NAME(ChaosCharacterLaunchCheck)
+
 
 UChaosCharacterLaunchCheck::UChaosCharacterLaunchCheck(const FObjectInitializer& ObjectInitializer)
 	: Super(ObjectInitializer)
diff --git a/Engine/Plugins/Experimental/ChaosMover/Source/ChaosMover/Private/PathedMovement/ChaosPathedMovementDebugDrawComponent.cpp b/Engine/Plugins/Experimental/ChaosMover/Source/ChaosMover/Private/PathedMovement/ChaosPathedMovementDebugDrawComponent.cpp
--- a/Engine/Plugins/Experimental/ChaosMover/Source/ChaosMover/Private/PathedMovement/ChaosPathedMovementDebugDrawComponent.cpp
+++ b/Engine/Plugins/Experimental/ChaosMover/Source/ChaosMover/Private/PathedMovement/ChaosPathedMovementDebugDrawComponent.cpp
@@ -11,6 +11,8 @@
 #include "MoverComponent.h"
 #include "SceneInterface.h"
 
+#include UE_INLINE_GENERATED_CPP_BY_NAME(ChaosPathedMovementDebugDrawComponent)
+
 TAutoConsoleVariable<int32> CVarChaosPathedMovementTotalDebugDrawSteps(
 	TEXT("ChaosMover.PathedMovement.DebugDraw.TotalNumSteps"),
 	64,
diff --git a/Engine/Plugins/Experimental/ChaosMover/Source/ChaosMover/Private/PathedMovement/ChaosPointMovementPathPattern.cpp b/Engine/Plugins/Experimental/ChaosMover/Source/ChaosMover/Private/PathedMovement/ChaosPointMovementPathPattern.cpp
--- a/Engine/Plugins/Experimental/ChaosMover/Source/ChaosMover/Private/PathedMovement/ChaosPointMovementPathPattern.cpp
+++ b/Engine/Plugins/Experimental/ChaosMover/Source/ChaosMover/Private/PathedMovement/ChaosPointMovementPathPattern.cpp
@@ -9,6 +9,8 @@
 #include "ChaosMover/PathedMovement/ChaosPathedMovementMode.h"
 #include "ChaosMover/PathedMovement/ChaosPathedMovementDebugDrawComponent.h"
 
+#include UE_INLINE_GENERATED_CPP_BY_NAME(ChaosPointMovementPathPattern)
+
 void UChaosPointMovementPathPattern::InitializePattern(UChaosMoverSimulation* InSimulation)
 {
 	Super::InitializePattern(InSimulation);
diff --git a/Engine/Plugins/Experimental/ChaosMover/Source/ChaosMover/Private/PathedMovement/ChaosSplineMovementPathPattern.cpp b/Engine/Plugins/Experimental/ChaosMover/Source/ChaosMover/Private/PathedMovement/ChaosSplineMovementPathPattern.cpp
--- a/Engine/Plugins/Experimental/ChaosMover/Source/ChaosMover/Private/PathedMovement/ChaosSplineMovementPathPattern.cpp
+++ b/Engine/Plugins/Experimental/ChaosMover/Source/ChaosMover/Private/PathedMovement/ChaosSplineMovementPathPattern.cpp
@@ -7,6 +7,8 @@
 
 #include "Components/SplineComponent.h"
 
+#include UE_INLINE_GENERATED_CPP_BY_NAME(ChaosSplineMovementPathPattern)
+
 void UChaosSplineMovementPathPattern::InitializePattern(UChaosMoverSimulation* InSimulation)
 {
 	Super::InitializePattern(InSimulation);
@@ -52,4 +54,4 @@ FTransform UChaosSplineMovementPathPattern::CalcUnmaskedTargetRelativeTransform(
 		return TargetToSpline;
 	}
 	return FTransform::Identity;
-}
\ No newline at end of file
+}
diff --git a/Engine/Plugins/Experimental/ChaosMover/Source/ChaosMover/Private/Utilities/ChaosGroundMovementUtils.cpp b/Engine/Plugins/Experimental/ChaosMover/Source/ChaosMover/Private/Utilities/ChaosGroundMovementUtils.cpp
--- a/Engine/Plugins/Experimental/ChaosMover/Source/ChaosMover/Private/Utilities/ChaosGroundMovementUtils.cpp
+++ b/Engine/Plugins/Experimental/ChaosMover/Source/ChaosMover/Private/Utilities/ChaosGroundMovementUtils.cpp
@@ -5,6 +5,8 @@
 #include "Chaos/PhysicsObjectInternalInterface.h"
 #include "MoveLibrary/FloorQueryUtils.h"
 
+#include UE_INLINE_GENERATED_CPP_BY_NAME(ChaosGroundMovementUtils)
+
 FVector UChaosGroundMovementUtils::ComputeLocalGroundVelocity_Internal(const FVector& Position, const FFloorCheckResult& FloorResult)
 {
 	FVector GroundVelocity = FVector::ZeroVector;
diff --git a/Engine/Plugins/Experimental/LearningAgents/Source/LearningAgentsTrainingEditor/Private/LearningAgentsImitationTrainerEditor.cpp b/Engine/Plugins/Experimental/LearningAgents/Source/LearningAgentsTrainingEditor/Private/LearningAgentsImitationTrainerEditor.cpp
--- a/Engine/Plugins/Experimental/LearningAgents/Source/LearningAgentsTrainingEditor/Private/LearningAgentsImitationTrainerEditor.cpp
+++ b/Engine/Plugins/Experimental/LearningAgents/Source/LearningAgentsTrainingEditor/Private/LearningAgentsImitationTrainerEditor.cpp
@@ -4,6 +4,8 @@
 #include "LearningLog.h"
 #include "Containers/Ticker.h"
 
+#include UE_INLINE_GENERATED_CPP_BY_NAME(LearningAgentsImitationTrainerEditor)
+
 ALearningAgentsImitationTrainerEditor::ALearningAgentsImitationTrainerEditor()
 {
 	LearningAgentsManager = CreateDefaultSubobject<ULearningAgentsManager>(TEXT("LearningAgentsManager"));
diff --git a/Engine/Plugins/Experimental/Mover/Source/Mover/Private/DefaultMovementSet/Modes/SpringWalkingMode.cpp b/Engine/Plugins/Experimental/Mover/Source/Mover/Private/DefaultMovementSet/Modes/SpringWalkingMode.cpp
--- a/Engine/Plugins/Experimental/Mover/Source/Mover/Private/DefaultMovementSet/Modes/SpringWalkingMode.cpp
+++ b/Engine/Plugins/Experimental/Mover/Source/Mover/Private/DefaultMovementSet/Modes/SpringWalkingMode.cpp
@@ -7,6 +7,8 @@
 #include "DefaultMovementSet/Settings/CommonLegacyMovementSettings.h"
 #include "Runtime/Engine/Internal/Animation/AnimationMath.h"
 
+#include UE_INLINE_GENERATED_CPP_BY_NAME(SpringWalkingMode)
+
 void USpringWalkingMode::GenerateMove_Implementation(const FMoverTickStartData& StartState, const FMoverTimeStep& TimeStep,
                                                    FProposedMove& OutProposedMove) const
 {
diff --git a/Engine/Plugins/Experimental/UAF/UAF/Source/UAF/Private/AnimNextConfig.cpp b/Engine/Plugins/Experimental/UAF/UAF/Source/UAF/Private/AnimNextConfig.cpp
--- a/Engine/Plugins/Experimental/UAF/UAF/Source/UAF/Private/AnimNextConfig.cpp
+++ b/Engine/Plugins/Experimental/UAF/UAF/Source/UAF/Private/AnimNextConfig.cpp
@@ -1,9 +1,11 @@
-﻿// Copyright Epic Games, Inc. All Rights Reserved.
+// Copyright Epic Games, Inc. All Rights Reserved.
 
 #include "AnimNextConfig.h"
 
 #include "Module/RigUnit_AnimNextModuleEvents.h"
 
+#include UE_INLINE_GENERATED_CPP_BY_NAME(AnimNextConfig)
+
 #if WITH_EDITOR
 
 void UAnimNextConfig::PostEditChangeProperty(FPropertyChangedEvent& PropertyChangedEvent)
diff --git a/Engine/Plugins/Experimental/UAF/UAF/Source/UAF/Private/AnimNextRigVMAsset.cpp b/Engine/Plugins/Experimental/UAF/UAF/Source/UAF/Private/AnimNextRigVMAsset.cpp
--- a/Engine/Plugins/Experimental/UAF/UAF/Source/UAF/Private/AnimNextRigVMAsset.cpp
+++ b/Engine/Plugins/Experimental/UAF/UAF/Source/UAF/Private/AnimNextRigVMAsset.cpp
@@ -12,6 +12,8 @@
 #include "Modules/ModuleManager.h"
 #include "UObject/AssetRegistryTagsContext.h"
 
+#include UE_INLINE_GENERATED_CPP_BY_NAME(AnimNextRigVMAsset)
+
 namespace UE::UAF::Private
 {
 #if WITH_EDITOR
@@ -206,4 +208,4 @@ UAnimNextRigVMAsset::FOnCompileJobEvent& UAnimNextRigVMAsset::OnCompileJobFinish
 	return UE::UAF::Private::GOnCompileJobFinished;
 }
 
-#endif
\ No newline at end of file
+#endif
diff --git a/Engine/Plugins/Experimental/UAF/UAF/Source/UAF/Private/Component/AnimNextComponent.cpp b/Engine/Plugins/Experimental/UAF/UAF/Source/UAF/Private/Component/AnimNextComponent.cpp
--- a/Engine/Plugins/Experimental/UAF/UAF/Source/UAF/Private/Component/AnimNextComponent.cpp
+++ b/Engine/Plugins/Experimental/UAF/UAF/Source/UAF/Private/Component/AnimNextComponent.cpp
@@ -8,6 +8,8 @@
 #include "Module/ModuleTaskContext.h"
 #include "Variables/AnimNextVariableReference.h"
 
+#include UE_INLINE_GENERATED_CPP_BY_NAME(AnimNextComponent)
+
 void UAnimNextComponent::OnRegister()
 {
 	using namespace UE::UAF;
diff --git a/Engine/Plugins/Experimental/UAF/UAF/Source/UAF/Private/Graph/RigUnit_GetActorTransform.cpp b/Engine/Plugins/Experimental/UAF/UAF/Source/UAF/Private/Graph/RigUnit_GetActorTransform.cpp
--- a/Engine/Plugins/Experimental/UAF/UAF/Source/UAF/Private/Graph/RigUnit_GetActorTransform.cpp
+++ b/Engine/Plugins/Experimental/UAF/UAF/Source/UAF/Private/Graph/RigUnit_GetActorTransform.cpp
@@ -1,10 +1,12 @@
-﻿// Copyright Epic Games, Inc. All Rights Reserved.
+// Copyright Epic Games, Inc. All Rights Reserved.
 
 #include "RigUnit_GetActorTransform.h"
 
 #include "Component/AnimNextComponent.h"
 #include "Module/AnimNextModuleInstance.h"
 
+#include UE_INLINE_GENERATED_CPP_BY_NAME(RigUnit_GetActorTransform)
+
 FAnimNextActorTransformComponent::FAnimNextActorTransformComponent()
 {
 	FAnimNextModuleInstance& ModuleInstance = GetModuleInstance();
@@ -58,4 +60,4 @@ FRigUnit_GetActorTransform_Execute()
 	FAnimNextModuleInstance& ModuleInstance = ModuleContextData.GetModuleInstance();
 	const FAnimNextActorTransformComponent& TransformComponent = ModuleInstance.GetComponent<FAnimNextActorTransformComponent>();
 	Transform = TransformComponent.GetActorTransform();
-}
\ No newline at end of file
+}
diff --git a/Engine/Plugins/Experimental/UAF/UAF/Source/UAF/Private/Graph/RigUnit_IsAnimBPDriven.cpp b/Engine/Plugins/Experimental/UAF/UAF/Source/UAF/Private/Graph/RigUnit_IsAnimBPDriven.cpp
--- a/Engine/Plugins/Experimental/UAF/UAF/Source/UAF/Private/Graph/RigUnit_IsAnimBPDriven.cpp
+++ b/Engine/Plugins/Experimental/UAF/UAF/Source/UAF/Private/Graph/RigUnit_IsAnimBPDriven.cpp
@@ -3,6 +3,8 @@
 #include "RigUnit_IsAnimBPDriven.h"
 #include "Components/SkeletalMeshComponent.h"
 
+#include UE_INLINE_GENERATED_CPP_BY_NAME(RigUnit_IsAnimBPDriven)
+
 FRigUnit_IsAnimBPDriven_Execute()
 {
 	Result = MeshComponent != nullptr && MeshComponent->bEnableAnimation;
diff --git a/Engine/Plugins/Experimental/UAF/UAF/Source/UAF/Private/Graph/RigUnit_IsRecentlyRendered.cpp b/Engine/Plugins/Experimental/UAF/UAF/Source/UAF/Private/Graph/RigUnit_IsRecentlyRendered.cpp
--- a/Engine/Plugins/Experimental/UAF/UAF/Source/UAF/Private/Graph/RigUnit_IsRecentlyRendered.cpp
+++ b/Engine/Plugins/Experimental/UAF/UAF/Source/UAF/Private/Graph/RigUnit_IsRecentlyRendered.cpp
@@ -3,6 +3,8 @@
 #include "RigUnit_IsRecentlyRendered.h"
 #include "SkeletalRenderPublic.h"
 
+#include UE_INLINE_GENERATED_CPP_BY_NAME(RigUnit_IsRecentlyRendered)
+
 FRigUnit_IsRecentlyRendered_Execute()
 {
 	if (const USkinnedMeshComponent* Component = MeshComponent)
diff --git a/Engine/Plugins/Experimental/UAF/UAF/Source/UAF/Private/Graph/RigVMTrait_AnimNextPublicVariables.cpp b/Engine/Plugins/Experimental/UAF/UAF/Source/UAF/Private/Graph/RigVMTrait_AnimNextPublicVariables.cpp
--- a/Engine/Plugins/Experimental/UAF/UAF/Source/UAF/Private/Graph/RigVMTrait_AnimNextPublicVariables.cpp
+++ b/Engine/Plugins/Experimental/UAF/UAF/Source/UAF/Private/Graph/RigVMTrait_AnimNextPublicVariables.cpp
@@ -4,6 +4,8 @@
 #include "AnimNextRigVMAsset.h"
 #include "Logging/StructuredLog.h"
 
+#include UE_INLINE_GENERATED_CPP_BY_NAME(RigVMTrait_AnimNextPublicVariables)
+
 #if WITH_EDITOR
 FString (*FRigVMTrait_AnimNextPublicVariables::GetDisplayNameFunc)(const FRigVMTrait_AnimNextPublicVariables& InTrait);
 void (*FRigVMTrait_AnimNextPublicVariables::GetProgrammaticPinsFunc)(const FRigVMTrait_AnimNextPublicVariables& InTrait, URigVMController* InController, int32 InParentPinIndex, const FString& InDefaultValue, FRigVMPinInfoArray& OutPinArray);
diff --git a/Engine/Plugins/Experimental/UAF/UAF/Source/UAF/Private/Module/AnimNextActorComponentReferenceComponent.cpp b/Engine/Plugins/Experimental/UAF/UAF/Source/UAF/Private/Module/AnimNextActorComponentReferenceComponent.cpp
--- a/Engine/Plugins/Experimental/UAF/UAF/Source/UAF/Private/Module/AnimNextActorComponentReferenceComponent.cpp
+++ b/Engine/Plugins/Experimental/UAF/UAF/Source/UAF/Private/Module/AnimNextActorComponentReferenceComponent.cpp
@@ -5,6 +5,8 @@
 #include "Component/AnimNextComponent.h"
 #include "Module/AnimNextModuleInstance.h"
 
+#include UE_INLINE_GENERATED_CPP_BY_NAME(AnimNextActorComponentReferenceComponent)
+
 void FAnimNextActorComponentReferenceComponent::OnInitializeHelper(UScriptStruct* InScriptStruct)
 {
 	check(InScriptStruct->IsChildOf(FAnimNextActorComponentReferenceComponent::StaticStruct()));
diff --git a/Engine/Plugins/Experimental/UAF/UAF/Source/UAF/Private/Module/AnimNextModuleContextData.cpp b/Engine/Plugins/Experimental/UAF/UAF/Source/UAF/Private/Module/AnimNextModuleContextData.cpp
--- a/Engine/Plugins/Experimental/UAF/UAF/Source/UAF/Private/Module/AnimNextModuleContextData.cpp
+++ b/Engine/Plugins/Experimental/UAF/UAF/Source/UAF/Private/Module/AnimNextModuleContextData.cpp
@@ -1,8 +1,10 @@
-﻿// Copyright Epic Games, Inc. All Rights Reserved.
+// Copyright Epic Games, Inc. All Rights Reserved.
 
 #include "Module/AnimNextModuleContextData.h"
 #include "Module/AnimNextModuleInstance.h"
 
+#include UE_INLINE_GENERATED_CPP_BY_NAME(AnimNextModuleContextData)
+
 FAnimNextModuleContextData::FAnimNextModuleContextData(FAnimNextModuleInstance* InModuleInstance)
 	: ModuleInstance(InModuleInstance)
 	, Instance(InModuleInstance)
@@ -18,4 +20,4 @@ FAnimNextModuleContextData::FAnimNextModuleContextData(FAnimNextModuleInstance*
 UObject* FAnimNextModuleContextData::GetObject() const
 {
 	return ModuleInstance ? ModuleInstance->GetObject() : nullptr;
-}
\ No newline at end of file
+}
diff --git a/Engine/Plugins/Experimental/UAF/UAF/Source/UAF/Private/Module/AnimNextModuleInstance.cpp b/Engine/Plugins/Experimental/UAF/UAF/Source/UAF/Private/Module/AnimNextModuleInstance.cpp
--- a/Engine/Plugins/Experimental/UAF/UAF/Source/UAF/Private/Module/AnimNextModuleInstance.cpp
+++ b/Engine/Plugins/Experimental/UAF/UAF/Source/UAF/Private/Module/AnimNextModuleInstance.cpp
@@ -20,6 +20,8 @@
 #include "RewindDebugger/AnimNextTrace.h"
 #include "Variables/AnimNextVariableReference.h"
 
+#include UE_INLINE_GENERATED_CPP_BY_NAME(AnimNextModuleInstance)
+
 DEFINE_STAT(STAT_AnimNext_InitializeInstance);
 
 FAnimNextModuleInstance::FAnimNextModuleInstance() = default;
@@ -714,4 +716,4 @@ void FAnimNextModuleInstance::QueueTaskOnOtherModule(const UE::UAF::FModuleHandl
 	}
 
 	OtherModuleInstance->QueueTask(InEventName, MoveTemp(InTaskFunction), InLocation);
-}
\ No newline at end of file
+}
diff --git a/Engine/Plugins/Experimental/UAF/UAF/Source/UAF/Private/Module/AnimNextSkeletalMeshComponentReferenceComponent.cpp b/Engine/Plugins/Experimental/UAF/UAF/Source/UAF/Private/Module/AnimNextSkeletalMeshComponentReferenceComponent.cpp
--- a/Engine/Plugins/Experimental/UAF/UAF/Source/UAF/Private/Module/AnimNextSkeletalMeshComponentReferenceComponent.cpp
+++ b/Engine/Plugins/Experimental/UAF/UAF/Source/UAF/Private/Module/AnimNextSkeletalMeshComponentReferenceComponent.cpp
@@ -1,9 +1,11 @@
-﻿// Copyright Epic Games, Inc. All Rights Reserved.
+// Copyright Epic Games, Inc. All Rights Reserved.
 
 #include "Module/AnimNextSkeletalMeshComponentReferenceComponent.h"
 
 #include "Components/SkeletalMeshComponent.h"
 
+#include UE_INLINE_GENERATED_CPP_BY_NAME(AnimNextSkeletalMeshComponentReferenceComponent)
+
 FAnimNextSkeletalMeshComponentReferenceComponent::FAnimNextSkeletalMeshComponentReferenceComponent()
 {
 	ComponentType = USkeletalMeshComponent::StaticClass();
diff --git a/Engine/Plugins/Experimental/UAF/UAF/Source/UAF/Private/Module/RigUnit_AnimNextModuleEventDependencies.cpp b/Engine/Plugins/Experimental/UAF/UAF/Source/UAF/Private/Module/RigUnit_AnimNextModuleEventDependencies.cpp
--- a/Engine/Plugins/Experimental/UAF/UAF/Source/UAF/Private/Module/RigUnit_AnimNextModuleEventDependencies.cpp
+++ b/Engine/Plugins/Experimental/UAF/UAF/Source/UAF/Private/Module/RigUnit_AnimNextModuleEventDependencies.cpp
@@ -14,6 +14,8 @@
 #include "Widgets/Layout/SBox.h"
 #endif
 
+#include UE_INLINE_GENERATED_CPP_BY_NAME(RigUnit_AnimNextModuleEventDependencies)
+
 namespace UE::UAF
 {
 
@@ -248,4 +250,4 @@ FRigUnit_AnimNextModuleRemoveDependencies_Execute()
 	{
 		InDependency.OnRemoveDependency(InContext);
 	});
-}
\ No newline at end of file
+}
diff --git a/Engine/Plugins/Experimental/UAF/UAF/Source/UAF/Private/Module/RigUnit_AnimNextModuleEvents.cpp b/Engine/Plugins/Experimental/UAF/UAF/Source/UAF/Private/Module/RigUnit_AnimNextModuleEvents.cpp
--- a/Engine/Plugins/Experimental/UAF/UAF/Source/UAF/Private/Module/RigUnit_AnimNextModuleEvents.cpp
+++ b/Engine/Plugins/Experimental/UAF/UAF/Source/UAF/Private/Module/RigUnit_AnimNextModuleEvents.cpp
@@ -6,6 +6,8 @@
 #include "Engine/World.h"
 #include "Module/AnimNextModuleInstance.h"
 
+#include UE_INLINE_GENERATED_CPP_BY_NAME(RigUnit_AnimNextModuleEvents)
+
 namespace UE::UAF::CVars
 {
 	// Expose control over whether or not the AnimNext binding execution should happen before others in the tick group
diff --git a/Engine/Plugins/Experimental/UAF/UAF/Source/UAF/Private/Module/RigVMTrait_ModuleEventDependency_ActorComponentClassPrimaryTickFunction.cpp b/Engine/Plugins/Experimental/UAF/UAF/Source/UAF/Private/Module/RigVMTrait_ModuleEventDependency_ActorComponentClassPrimaryTickFunction.cpp
--- a/Engine/Plugins/Experimental/UAF/UAF/Source/UAF/Private/Module/RigVMTrait_ModuleEventDependency_ActorComponentClassPrimaryTickFunction.cpp
+++ b/Engine/Plugins/Experimental/UAF/UAF/Source/UAF/Private/Module/RigVMTrait_ModuleEventDependency_ActorComponentClassPrimaryTickFunction.cpp
@@ -2,6 +2,8 @@
 
 #include "RigVMTrait_ModuleEventDependency_ActorComponentClassPrimaryTickFunction.h"
 
+#include UE_INLINE_GENERATED_CPP_BY_NAME(RigVMTrait_ModuleEventDependency_ActorComponentClassPrimaryTickFunction)
+
 #if WITH_EDITOR
 FString FRigVMTrait_ModuleEventDependency_ActorComponentClassPrimaryTickFunction::GetDisplayName() const
 {
@@ -67,4 +69,4 @@ void FRigVMTrait_ModuleEventDependency_ActorComponentClassPrimaryTickFunction::O
 	{
 		TargetComponent->PrimaryComponentTick.RemovePrerequisite(AnimNextComponent, InContext.TickFunction);
 	}
-}
\ No newline at end of file
+}
diff --git a/Engine/Plugins/Experimental/UAF/UAF/Source/UAF/Private/Module/RigVMTrait_ModuleEventDependency_ActorComponentPrimaryTickFunction.cpp b/Engine/Plugins/Experimental/UAF/UAF/Source/UAF/Private/Module/RigVMTrait_ModuleEventDependency_ActorComponentPrimaryTickFunction.cpp
--- a/Engine/Plugins/Experimental/UAF/UAF/Source/UAF/Private/Module/RigVMTrait_ModuleEventDependency_ActorComponentPrimaryTickFunction.cpp
+++ b/Engine/Plugins/Experimental/UAF/UAF/Source/UAF/Private/Module/RigVMTrait_ModuleEventDependency_ActorComponentPrimaryTickFunction.cpp
@@ -2,6 +2,8 @@
 
 #include "RigVMTrait_ModuleEventDependency_ActorComponentPrimaryTickFunction.h"
 
+#include UE_INLINE_GENERATED_CPP_BY_NAME(RigVMTrait_ModuleEventDependency_ActorComponentPrimaryTickFunction)
+
 #if WITH_EDITOR
 FString FRigVMTrait_ModuleEventDependency_ActorComponentPrimaryTickFunction::GetDisplayName() const
 {
@@ -53,4 +55,4 @@ void FRigVMTrait_ModuleEventDependency_ActorComponentPrimaryTickFunction::OnRemo
 	{
 		Component->PrimaryComponentTick.RemovePrerequisite(AnimNextComponent, InContext.TickFunction);
 	}
-}
\ No newline at end of file
+}
diff --git a/Engine/Plugins/Experimental/UAF/UAF/Source/UAF/Private/Module/UAFModuleInstanceComponent.cpp b/Engine/Plugins/Experimental/UAF/UAF/Source/UAF/Private/Module/UAFModuleInstanceComponent.cpp
--- a/Engine/Plugins/Experimental/UAF/UAF/Source/UAF/Private/Module/UAFModuleInstanceComponent.cpp
+++ b/Engine/Plugins/Experimental/UAF/UAF/Source/UAF/Private/Module/UAFModuleInstanceComponent.cpp
@@ -1,8 +1,10 @@
-﻿// Copyright Epic Games, Inc. All Rights Reserved.
+// Copyright Epic Games, Inc. All Rights Reserved.
 
 #include "Module/UAFModuleInstanceComponent.h"
 #include "Module/AnimNextModuleInstance.h"
 
+#include UE_INLINE_GENERATED_CPP_BY_NAME(UAFModuleInstanceComponent)
+
 FAnimNextModuleInstance* FUAFModuleInstanceComponent::GetModuleInstancePtr()
 {
 	return static_cast<FAnimNextModuleInstance*>(GetAssetInstancePtr());
diff --git a/Engine/Plugins/Experimental/UAF/UAF/Source/UAF/Private/Param/AnimNextActorLocatorFragment.cpp b/Engine/Plugins/Experimental/UAF/UAF/Source/UAF/Private/Param/AnimNextActorLocatorFragment.cpp
--- a/Engine/Plugins/Experimental/UAF/UAF/Source/UAF/Private/Param/AnimNextActorLocatorFragment.cpp
+++ b/Engine/Plugins/Experimental/UAF/UAF/Source/UAF/Private/Param/AnimNextActorLocatorFragment.cpp
@@ -1,4 +1,4 @@
-﻿// Copyright Epic Games, Inc. All Rights Reserved.
+// Copyright Epic Games, Inc. All Rights Reserved.
 
 #include "Param/AnimNextActorLocatorFragment.h"
 #include "GameFramework/Actor.h"
@@ -9,6 +9,8 @@
 #include "UniversalObjectLocatorStringParams.h"
 #include "UniversalObjectLocatorInitializeResult.h"
 
+#include UE_INLINE_GENERATED_CPP_BY_NAME(AnimNextActorLocatorFragment)
+
 UE::UniversalObjectLocator::TFragmentTypeHandle<FAnimNextActorLocatorFragment> FAnimNextActorLocatorFragment::FragmentType;
 
 UE::UniversalObjectLocator::FResolveResult FAnimNextActorLocatorFragment::Resolve(const UE::UniversalObjectLocator::FResolveParams& InParams) const
diff --git a/Engine/Plugins/Experimental/UAF/UAF/Source/UAF/Private/Param/AnimNextComponentLocatorFragment.cpp b/Engine/Plugins/Experimental/UAF/UAF/Source/UAF/Private/Param/AnimNextComponentLocatorFragment.cpp
--- a/Engine/Plugins/Experimental/UAF/UAF/Source/UAF/Private/Param/AnimNextComponentLocatorFragment.cpp
+++ b/Engine/Plugins/Experimental/UAF/UAF/Source/UAF/Private/Param/AnimNextComponentLocatorFragment.cpp
@@ -1,4 +1,4 @@
-﻿// Copyright Epic Games, Inc. All Rights Reserved.
+// Copyright Epic Games, Inc. All Rights Reserved.
 
 #include "Param/AnimNextComponentLocatorFragment.h"
 #include "Component/AnimNextComponent.h"
@@ -8,6 +8,8 @@
 #include "UniversalObjectLocatorStringParams.h"
 #include "UniversalObjectLocatorInitializeResult.h"
 
+#include UE_INLINE_GENERATED_CPP_BY_NAME(AnimNextComponentLocatorFragment)
+
 UE::UniversalObjectLocator::TFragmentTypeHandle<FAnimNextComponentLocatorFragment> FAnimNextComponentLocatorFragment::FragmentType;
 
 FAnimNextComponentLocatorFragment::FAnimNextComponentLocatorFragment(TSubclassOf<UActorComponent> InClass)
diff --git a/Engine/Plugins/Experimental/UAF/UAF/Source/UAF/Private/Param/AnimNextObjectCastLocatorFragment.cpp b/Engine/Plugins/Experimental/UAF/UAF/Source/UAF/Private/Param/AnimNextObjectCastLocatorFragment.cpp
--- a/Engine/Plugins/Experimental/UAF/UAF/Source/UAF/Private/Param/AnimNextObjectCastLocatorFragment.cpp
+++ b/Engine/Plugins/Experimental/UAF/UAF/Source/UAF/Private/Param/AnimNextObjectCastLocatorFragment.cpp
@@ -1,4 +1,4 @@
-﻿// Copyright Epic Games, Inc. All Rights Reserved.
+// Copyright Epic Games, Inc. All Rights Reserved.
 
 #include "Param/AnimNextObjectCastLocatorFragment.h"
 #include "UniversalObjectLocatorFragmentTypeHandle.h"
@@ -7,6 +7,8 @@
 #include "UniversalObjectLocatorInitializeParams.h"
 #include "UniversalObjectLocatorInitializeResult.h"
 
+#include UE_INLINE_GENERATED_CPP_BY_NAME(AnimNextObjectCastLocatorFragment)
+
 UE::UniversalObjectLocator::TFragmentTypeHandle<FAnimNextObjectCastLocatorFragment> FAnimNextObjectCastLocatorFragment::FragmentType;
 
 FAnimNextObjectCastLocatorFragment::FAnimNextObjectCastLocatorFragment(UClass* InClass)
diff --git a/Engine/Plugins/Experimental/UAF/UAF/Source/UAF/Private/Param/AnimNextObjectFunctionLocatorFragment.cpp b/Engine/Plugins/Experimental/UAF/UAF/Source/UAF/Private/Param/AnimNextObjectFunctionLocatorFragment.cpp
--- a/Engine/Plugins/Experimental/UAF/UAF/Source/UAF/Private/Param/AnimNextObjectFunctionLocatorFragment.cpp
+++ b/Engine/Plugins/Experimental/UAF/UAF/Source/UAF/Private/Param/AnimNextObjectFunctionLocatorFragment.cpp
@@ -8,6 +8,8 @@
 #include "UniversalObjectLocatorInitializeResult.h"
 #include "Param/ParamUtils.h"
 
+#include UE_INLINE_GENERATED_CPP_BY_NAME(AnimNextObjectFunctionLocatorFragment)
+
 UE::UniversalObjectLocator::TFragmentTypeHandle<FAnimNextObjectFunctionLocatorFragment> FAnimNextObjectFunctionLocatorFragment::FragmentType;
 
 FAnimNextObjectFunctionLocatorFragment::FAnimNextObjectFunctionLocatorFragment(UFunction* InFunction)
diff --git a/Engine/Plugins/Experimental/UAF/UAF/Source/UAF/Private/Param/AnimNextObjectPropertyLocatorFragment.cpp b/Engine/Plugins/Experimental/UAF/UAF/Source/UAF/Private/Param/AnimNextObjectPropertyLocatorFragment.cpp
--- a/Engine/Plugins/Experimental/UAF/UAF/Source/UAF/Private/Param/AnimNextObjectPropertyLocatorFragment.cpp
+++ b/Engine/Plugins/Experimental/UAF/UAF/Source/UAF/Private/Param/AnimNextObjectPropertyLocatorFragment.cpp
@@ -8,6 +8,8 @@
 #include "UniversalObjectLocatorInitializeResult.h"
 #include "String/ParseTokens.h"
 
+#include UE_INLINE_GENERATED_CPP_BY_NAME(AnimNextObjectPropertyLocatorFragment)
+
 #define LOCTEXT_NAMESPACE "UAFObjectPropertyLocatorFragment"
 
 UE::UniversalObjectLocator::TFragmentTypeHandle<FAnimNextObjectPropertyLocatorFragment> FAnimNextObjectPropertyLocatorFragment::FragmentType;
@@ -159,4 +161,4 @@ UClass* FAnimNextObjectPropertyLocatorFragment::GetRootClass() const
 	return nullptr;
 }
 
-#undef LOCTEXT_NAMESPACE
\ No newline at end of file
+#undef LOCTEXT_NAMESPACE
diff --git a/Engine/Plugins/Experimental/UAF/UAF/Source/UAF/Private/Param/AnimNextWorldLibrary.cpp b/Engine/Plugins/Experimental/UAF/UAF/Source/UAF/Private/Param/AnimNextWorldLibrary.cpp
--- a/Engine/Plugins/Experimental/UAF/UAF/Source/UAF/Private/Param/AnimNextWorldLibrary.cpp
+++ b/Engine/Plugins/Experimental/UAF/UAF/Source/UAF/Private/Param/AnimNextWorldLibrary.cpp
@@ -3,6 +3,8 @@
 #include "AnimNextWorldLibrary.h"
 #include "Engine/World.h"
 
+#include UE_INLINE_GENERATED_CPP_BY_NAME(AnimNextWorldLibrary)
+
 double UAnimNextWorldLibrary::GetTimeSeconds(UWorld* InWorld)
 {
 	return InWorld->GetTimeSeconds();
diff --git a/Engine/Plugins/Experimental/UAF/UAF/Source/UAF/Private/Param/ParamType.cpp b/Engine/Plugins/Experimental/UAF/UAF/Source/UAF/Private/Param/ParamType.cpp
--- a/Engine/Plugins/Experimental/UAF/UAF/Source/UAF/Private/Param/ParamType.cpp
+++ b/Engine/Plugins/Experimental/UAF/UAF/Source/UAF/Private/Param/ParamType.cpp
@@ -8,6 +8,8 @@
 #include "RigVMCore/RigVMTemplate.h"
 #include "UObject/TextProperty.h"
 
+#include UE_INLINE_GENERATED_CPP_BY_NAME(ParamType)
+
 FRigVMTemplateArgumentType FAnimNextParamType::ToRigVMTemplateArgument() const
 {
 	auto GetTypeInternal = [this]()
@@ -711,4 +713,4 @@ bool FAnimNextParamType::IsObjectType() const
 	default:
 		return false;
 	}
-}
\ No newline at end of file
+}
diff --git a/Engine/Plugins/Experimental/UAF/UAF/Source/UAF/Private/Param/RigVMDispatch_GetLayerParameter.cpp b/Engine/Plugins/Experimental/UAF/UAF/Source/UAF/Private/Param/RigVMDispatch_GetLayerParameter.cpp
--- a/Engine/Plugins/Experimental/UAF/UAF/Source/UAF/Private/Param/RigVMDispatch_GetLayerParameter.cpp
+++ b/Engine/Plugins/Experimental/UAF/UAF/Source/UAF/Private/Param/RigVMDispatch_GetLayerParameter.cpp
@@ -3,6 +3,8 @@
 #include "Param/RigVMDispatch_GetLayerParameter.h"
 #include "RigVMCore/RigVMStruct.h"
 
+#include UE_INLINE_GENERATED_CPP_BY_NAME(RigVMDispatch_GetLayerParameter)
+
 const FName FRigVMDispatch_GetLayerParameter::ValueName = TEXT("Value");
 const FName FRigVMDispatch_GetLayerParameter::TypeHandleName = TEXT("Type");
 const FName FRigVMDispatch_GetLayerParameter::ParameterName = TEXT("Parameter");
diff --git a/Engine/Plugins/Experimental/UAF/UAF/Source/UAF/Private/Param/RigVMDispatch_GetParameter.cpp b/Engine/Plugins/Experimental/UAF/UAF/Source/UAF/Private/Param/RigVMDispatch_GetParameter.cpp
--- a/Engine/Plugins/Experimental/UAF/UAF/Source/UAF/Private/Param/RigVMDispatch_GetParameter.cpp
+++ b/Engine/Plugins/Experimental/UAF/UAF/Source/UAF/Private/Param/RigVMDispatch_GetParameter.cpp
@@ -3,6 +3,8 @@
 #include "Param/RigVMDispatch_GetParameter.h"
 #include "RigVMCore/RigVMStruct.h"
 
+#include UE_INLINE_GENERATED_CPP_BY_NAME(RigVMDispatch_GetParameter)
+
 const FName FRigVMDispatch_GetParameter::ValueName = TEXT("Value");
 const FName FRigVMDispatch_GetParameter::TypeHandleName = TEXT("Type");
 const FName FRigVMDispatch_GetParameter::ParameterName = TEXT("Parameter");
diff --git a/Engine/Plugins/Experimental/UAF/UAF/Source/UAF/Private/Param/RigVMDispatch_GetScopedParameter.cpp b/Engine/Plugins/Experimental/UAF/UAF/Source/UAF/Private/Param/RigVMDispatch_GetScopedParameter.cpp
--- a/Engine/Plugins/Experimental/UAF/UAF/Source/UAF/Private/Param/RigVMDispatch_GetScopedParameter.cpp
+++ b/Engine/Plugins/Experimental/UAF/UAF/Source/UAF/Private/Param/RigVMDispatch_GetScopedParameter.cpp
@@ -3,6 +3,8 @@
 #include "Param/RigVMDispatch_GetScopedParameter.h"
 #include "RigVMCore/RigVMStruct.h"
 
+#include UE_INLINE_GENERATED_CPP_BY_NAME(RigVMDispatch_GetScopedParameter)
+
 const FName FRigVMDispatch_GetScopedParameter::ParameterName = TEXT("Parameter");
 const FName FRigVMDispatch_GetScopedParameter::ValueName = TEXT("Value");
 const FName FRigVMDispatch_GetScopedParameter::ParameterIdName = TEXT("ParameterId");
diff --git a/Engine/Plugins/Experimental/UAF/UAF/Source/UAF/Private/Param/RigVMDispatch_SetLayerParameter.cpp b/Engine/Plugins/Experimental/UAF/UAF/Source/UAF/Private/Param/RigVMDispatch_SetLayerParameter.cpp
--- a/Engine/Plugins/Experimental/UAF/UAF/Source/UAF/Private/Param/RigVMDispatch_SetLayerParameter.cpp
+++ b/Engine/Plugins/Experimental/UAF/UAF/Source/UAF/Private/Param/RigVMDispatch_SetLayerParameter.cpp
@@ -4,6 +4,8 @@
 #include "RigVMCore/RigVMStruct.h"
 #include "RigVMCore/RigVM.h"
 
+#include UE_INLINE_GENERATED_CPP_BY_NAME(RigVMDispatch_SetLayerParameter)
+
 const FName FRigVMDispatch_SetLayerParameter::ExecuteContextName = TEXT("ExecuteContext");
 const FName FRigVMDispatch_SetLayerParameter::ValueName = TEXT("Value");
 const FName FRigVMDispatch_SetLayerParameter::TypeHandleName = TEXT("Type");
diff --git a/Engine/Plugins/Experimental/UAF/UAF/Source/UAF/Private/TraitCore/TraitEvent.cpp b/Engine/Plugins/Experimental/UAF/UAF/Source/UAF/Private/TraitCore/TraitEvent.cpp
--- a/Engine/Plugins/Experimental/UAF/UAF/Source/UAF/Private/TraitCore/TraitEvent.cpp
+++ b/Engine/Plugins/Experimental/UAF/UAF/Source/UAF/Private/TraitCore/TraitEvent.cpp
@@ -2,6 +2,8 @@
 
 #include "TraitCore/TraitEvent.h"
 
+#include UE_INLINE_GENERATED_CPP_BY_NAME(TraitEvent)
+
 bool FAnimNextTraitEvent::DecrementLifetime(UE::UAF::FTraitEventList& OutputEventList)
 {
 	const bool bExpired = Lifetime.Decrement();
diff --git a/Engine/Plugins/Experimental/UAF/UAF/Source/UAF/Private/UAFAssetInstance.cpp b/Engine/Plugins/Experimental/UAF/UAF/Source/UAF/Private/UAFAssetInstance.cpp
--- a/Engine/Plugins/Experimental/UAF/UAF/Source/UAF/Private/UAFAssetInstance.cpp
+++ b/Engine/Plugins/Experimental/UAF/UAF/Source/UAF/Private/UAFAssetInstance.cpp
@@ -9,6 +9,8 @@
 #include "Variables/VariableOverridesCollection.h"
 #include "UAFAssetInstanceComponent.h"
 
+#include UE_INLINE_GENERATED_CPP_BY_NAME(UAFAssetInstance)
+
 void FUAFAssetInstance::InitializeVariables(const TSharedPtr<const UE::UAF::FVariableOverridesCollection>& InOverrides)
 {
 	Variables.Initialize(*this, InOverrides);
@@ -92,4 +94,4 @@ FUAFAssetInstanceComponent* FUAFAssetInstance::TryGetComponent(const UScriptStru
 		return InstanceComponent->GetMutablePtr<FUAFAssetInstanceComponent>();
 	}
 	return nullptr;
-}
\ No newline at end of file
+}
diff --git a/Engine/Plugins/Experimental/UAF/UAF/Source/UAF/Private/UAFAssetInstanceComponent.cpp b/Engine/Plugins/Experimental/UAF/UAF/Source/UAF/Private/UAFAssetInstanceComponent.cpp
--- a/Engine/Plugins/Experimental/UAF/UAF/Source/UAF/Private/UAFAssetInstanceComponent.cpp
+++ b/Engine/Plugins/Experimental/UAF/UAF/Source/UAF/Private/UAFAssetInstanceComponent.cpp
@@ -1,7 +1,9 @@
-﻿// Copyright Epic Games, Inc. All Rights Reserved.
+// Copyright Epic Games, Inc. All Rights Reserved.
 
 #include "UAFAssetInstanceComponent.h"
 
+#include UE_INLINE_GENERATED_CPP_BY_NAME(UAFAssetInstanceComponent)
+
 namespace UE::UAF::Private
 {
 	static thread_local FUAFAssetInstance* GCurrentAssetInstance = nullptr;
@@ -20,4 +22,4 @@ FUAFAssetInstanceComponent::FScopedConstructorHelper::~FScopedConstructorHelper(
 FUAFAssetInstanceComponent::FUAFAssetInstanceComponent()
 {
 	Instance = UE::UAF::Private::GCurrentAssetInstance;
-}
\ No newline at end of file
+}
diff --git a/Engine/Plugins/Experimental/UAF/UAF/Source/UAF/Private/UAFRigVMComponent.cpp b/Engine/Plugins/Experimental/UAF/UAF/Source/UAF/Private/UAFRigVMComponent.cpp
--- a/Engine/Plugins/Experimental/UAF/UAF/Source/UAF/Private/UAFRigVMComponent.cpp
+++ b/Engine/Plugins/Experimental/UAF/UAF/Source/UAF/Private/UAFRigVMComponent.cpp
@@ -1,4 +1,4 @@
-﻿// Copyright Epic Games, Inc. All Rights Reserved.
+// Copyright Epic Games, Inc. All Rights Reserved.
 
 #include "UAFRigVMComponent.h"
 
@@ -6,6 +6,8 @@
 #include "UAFAssetInstance.h"
 #include "Variables/UAFInstanceVariableContainer.h"
 
+#include UE_INLINE_GENERATED_CPP_BY_NAME(UAFRigVMComponent)
+
 FUAFRigVMComponent::FUAFRigVMComponent()
 {
 	FUAFAssetInstance* AssetInstance = GetAssetInstancePtr();
@@ -46,4 +48,4 @@ FUAFRigVMComponent::FUAFRigVMComponent()
 
 	// Now initialize the 'instance', cache memory handles etc. in the context
 	VM->InitializeInstance(ExtendedExecuteContext);
-}
\ No newline at end of file
+}
diff --git a/Engine/Plugins/Experimental/UAF/UAF/Source/UAF/Private/Variables/AnimNextSharedVariables.cpp b/Engine/Plugins/Experimental/UAF/UAF/Source/UAF/Private/Variables/AnimNextSharedVariables.cpp
--- a/Engine/Plugins/Experimental/UAF/UAF/Source/UAF/Private/Variables/AnimNextSharedVariables.cpp
+++ b/Engine/Plugins/Experimental/UAF/UAF/Source/UAF/Private/Variables/AnimNextSharedVariables.cpp
@@ -1,8 +1,10 @@
-﻿// Copyright Epic Games, Inc. All Rights Reserved.
+// Copyright Epic Games, Inc. All Rights Reserved.
 
 #include "Variables/AnimNextSharedVariables.h"
 #include "AnimNextExecuteContext.h"
 
+#include UE_INLINE_GENERATED_CPP_BY_NAME(AnimNextSharedVariables)
+
 #if WITH_EDITOR	
 #include "Engine/ExternalAssetDependencyGatherer.h"
 REGISTER_ASSETDEPENDENCY_GATHERER(FExternalAssetDependencyGatherer, UAnimNextSharedVariables);
diff --git a/Engine/Plugins/Experimental/UAF/UAF/Source/UAF/Private/Variables/AnimNextSoftVariableReference.cpp b/Engine/Plugins/Experimental/UAF/UAF/Source/UAF/Private/Variables/AnimNextSoftVariableReference.cpp
--- a/Engine/Plugins/Experimental/UAF/UAF/Source/UAF/Private/Variables/AnimNextSoftVariableReference.cpp
+++ b/Engine/Plugins/Experimental/UAF/UAF/Source/UAF/Private/Variables/AnimNextSoftVariableReference.cpp
@@ -4,6 +4,8 @@
 #include "AnimNextRigVMAsset.h"
 #include "Variables/AnimNextVariableReference.h"
 
+#include UE_INLINE_GENERATED_CPP_BY_NAME(AnimNextSoftVariableReference)
+
 FAnimNextSoftVariableReference::FAnimNextSoftVariableReference(const FAnimNextVariableReference& InVariableReference)
 	: Name(InVariableReference.GetName())
 	, SoftObjectPath(InVariableReference.GetObject())
diff --git a/Engine/Plugins/Experimental/UAF/UAF/Source/UAF/Private/Variables/AnimNextVariableOverridesCollection.cpp b/Engine/Plugins/Experimental/UAF/UAF/Source/UAF/Private/Variables/AnimNextVariableOverridesCollection.cpp
--- a/Engine/Plugins/Experimental/UAF/UAF/Source/UAF/Private/Variables/AnimNextVariableOverridesCollection.cpp
+++ b/Engine/Plugins/Experimental/UAF/UAF/Source/UAF/Private/Variables/AnimNextVariableOverridesCollection.cpp
@@ -1,9 +1,11 @@
-﻿// Copyright Epic Games, Inc. All Rights Reserved.
+// Copyright Epic Games, Inc. All Rights Reserved.
 
 #include "Variables/AnimNextVariableOverridesCollection.h"
 #include "AnimNextRigVMAsset.h"
 #include "VariableOverridesCollection.h"
 
+#include UE_INLINE_GENERATED_CPP_BY_NAME(AnimNextVariableOverridesCollection)
+
 void FAnimNextVariableOverridesCollection::AddStructReferencedObjects(FReferenceCollector& Collector)
 {
 	using namespace UE::UAF;
diff --git a/Engine/Plugins/Experimental/UAF/UAF/Source/UAF/Private/Variables/AnimNextVariableReference.cpp b/Engine/Plugins/Experimental/UAF/UAF/Source/UAF/Private/Variables/AnimNextVariableReference.cpp
--- a/Engine/Plugins/Experimental/UAF/UAF/Source/UAF/Private/Variables/AnimNextVariableReference.cpp
+++ b/Engine/Plugins/Experimental/UAF/UAF/Source/UAF/Private/Variables/AnimNextVariableReference.cpp
@@ -1,9 +1,11 @@
-﻿// Copyright Epic Games, Inc. All Rights Reserved.
+// Copyright Epic Games, Inc. All Rights Reserved.
 
 #include "Variables/AnimNextVariableReference.h"
 #include "AnimNextRigVMAsset.h"
 #include "Variables/AnimNextSoftVariableReference.h"
 
+#include UE_INLINE_GENERATED_CPP_BY_NAME(AnimNextVariableReference)
+
 FAnimNextVariableReference::FAnimNextVariableReference(const FAnimNextSoftVariableReference& InSoftReference)
 	: Name(InSoftReference.GetName())
 	, Object(InSoftReference.GetSoftObjectPath().TryLoad())
@@ -56,4 +58,4 @@ const FProperty* FAnimNextVariableReference::ResolveProperty() const
 		}
 	}
 	return CachedProperty.Get();
-}
\ No newline at end of file
+}
diff --git a/Engine/Plugins/Experimental/UAF/UAF/Source/UAF/Private/Variables/RigUnit_CopyModuleProxyVariables.cpp b/Engine/Plugins/Experimental/UAF/UAF/Source/UAF/Private/Variables/RigUnit_CopyModuleProxyVariables.cpp
--- a/Engine/Plugins/Experimental/UAF/UAF/Source/UAF/Private/Variables/RigUnit_CopyModuleProxyVariables.cpp
+++ b/Engine/Plugins/Experimental/UAF/UAF/Source/UAF/Private/Variables/RigUnit_CopyModuleProxyVariables.cpp
@@ -1,10 +1,12 @@
-﻿// Copyright Epic Games, Inc. All Rights Reserved.
+// Copyright Epic Games, Inc. All Rights Reserved.
 
 #include "Variables/RigUnit_CopyModuleProxyVariables.h"
 #include "Module/AnimNextModuleInstance.h"
 
+#include UE_INLINE_GENERATED_CPP_BY_NAME(RigUnit_CopyModuleProxyVariables)
+
 FRigUnit_CopyModuleProxyVariables_Execute()
 {
 	const FAnimNextModuleContextData& ContextData = ExecuteContext.GetContextData<FAnimNextModuleContextData>();
 	ContextData.GetModuleInstance().CopyProxyVariables();
-}
\ No newline at end of file
+}
diff --git a/Engine/Plugins/Experimental/UAF/UAF/Source/UAF/Private/Variables/RigUnit_OverrideVariables.cpp b/Engine/Plugins/Experimental/UAF/UAF/Source/UAF/Private/Variables/RigUnit_OverrideVariables.cpp
--- a/Engine/Plugins/Experimental/UAF/UAF/Source/UAF/Private/Variables/RigUnit_OverrideVariables.cpp
+++ b/Engine/Plugins/Experimental/UAF/UAF/Source/UAF/Private/Variables/RigUnit_OverrideVariables.cpp
@@ -19,6 +19,8 @@
 #include "Variables/AnimNextSharedVariables.h"
 #endif
 
+#include UE_INLINE_GENERATED_CPP_BY_NAME(RigUnit_OverrideVariables)
+
 #define LOCTEXT_NAMESPACE "AnimNextOverrideVariables"
 
 FRigUnit_OverrideVariables_Execute()
@@ -288,4 +290,4 @@ TArray<FRigVMUserWorkflow> FRigUnit_OverrideVariables::GetSupportedWorkflows(con
 	return Workflows;
 }
 
-#undef LOCTEXT_NAMESPACE
\ No newline at end of file
+#undef LOCTEXT_NAMESPACE
diff --git a/Engine/Plugins/Experimental/UAF/UAF/Source/UAF/Private/Variables/RigUnit_ResolveUniversalObjectLocator.cpp b/Engine/Plugins/Experimental/UAF/UAF/Source/UAF/Private/Variables/RigUnit_ResolveUniversalObjectLocator.cpp
--- a/Engine/Plugins/Experimental/UAF/UAF/Source/UAF/Private/Variables/RigUnit_ResolveUniversalObjectLocator.cpp
+++ b/Engine/Plugins/Experimental/UAF/UAF/Source/UAF/Private/Variables/RigUnit_ResolveUniversalObjectLocator.cpp
@@ -1,8 +1,10 @@
-﻿// Copyright Epic Games, Inc. All Rights Reserved.
+// Copyright Epic Games, Inc. All Rights Reserved.
 
 #include "Variables/RigUnit_ResolveUniversalObjectLocator.h"
 #include "Module/AnimNextModuleInstance.h"
 
+#include UE_INLINE_GENERATED_CPP_BY_NAME(RigUnit_ResolveUniversalObjectLocator)
+
 FRigUnit_ResolveUniversalObjectLocator_Execute()
 {
 	using namespace UE::UniversalObjectLocator;
@@ -13,4 +15,4 @@ FRigUnit_ResolveUniversalObjectLocator_Execute()
 	ResolveParams.Context = ContextData.GetObject();
 	FResolveResult Result = Locator.Resolve(ResolveParams);
 	Object = Result.SyncGet().Object;
-}
\ No newline at end of file
+}
diff --git a/Engine/Plugins/Experimental/UAF/UAF/Source/UAF/Private/Variables/RigVMDispatch_CallHoistedAccessorFunction.cpp b/Engine/Plugins/Experimental/UAF/UAF/Source/UAF/Private/Variables/RigVMDispatch_CallHoistedAccessorFunction.cpp
--- a/Engine/Plugins/Experimental/UAF/UAF/Source/UAF/Private/Variables/RigVMDispatch_CallHoistedAccessorFunction.cpp
+++ b/Engine/Plugins/Experimental/UAF/UAF/Source/UAF/Private/Variables/RigVMDispatch_CallHoistedAccessorFunction.cpp
@@ -1,9 +1,11 @@
-﻿// Copyright Epic Games, Inc. All Rights Reserved.
+// Copyright Epic Games, Inc. All Rights Reserved.
 
 #include "Variables/RigVMDispatch_CallHoistedAccessorFunction.h"
 #include "RigVMCore/RigVMStruct.h"
 #include "Variables/AnimNextSoftFunctionPtr.h"
 
+#include UE_INLINE_GENERATED_CPP_BY_NAME(RigVMDispatch_CallHoistedAccessorFunction)
+
 FRigVMDispatch_CallHoistedAccessorFunctionBase::FRigVMDispatch_CallHoistedAccessorFunctionBase()
 {
 	FactoryScriptStruct = StaticStruct();
diff --git a/Engine/Plugins/Experimental/UAF/UAF/Source/UAF/Private/Variables/RigVMDispatch_CallObjectAccessorFunction.cpp b/Engine/Plugins/Experimental/UAF/UAF/Source/UAF/Private/Variables/RigVMDispatch_CallObjectAccessorFunction.cpp
--- a/Engine/Plugins/Experimental/UAF/UAF/Source/UAF/Private/Variables/RigVMDispatch_CallObjectAccessorFunction.cpp
+++ b/Engine/Plugins/Experimental/UAF/UAF/Source/UAF/Private/Variables/RigVMDispatch_CallObjectAccessorFunction.cpp
@@ -1,9 +1,11 @@
-﻿// Copyright Epic Games, Inc. All Rights Reserved.
+// Copyright Epic Games, Inc. All Rights Reserved.
 
 #include "Variables/RigVMDispatch_CallObjectAccessorFunction.h"
 #include "RigVMCore/RigVMStruct.h"
 #include "Variables/AnimNextSoftFunctionPtr.h"
 
+#include UE_INLINE_GENERATED_CPP_BY_NAME(RigVMDispatch_CallObjectAccessorFunction)
+
 FRigVMDispatch_CallObjectAccessorFunctionBase::FRigVMDispatch_CallObjectAccessorFunctionBase()
 {
 	FactoryScriptStruct = StaticStruct();
diff --git a/Engine/Plugins/Experimental/UAF/UAF/Source/UAF/Private/Variables/RigVMDispatch_GetObjectProperty.cpp b/Engine/Plugins/Experimental/UAF/UAF/Source/UAF/Private/Variables/RigVMDispatch_GetObjectProperty.cpp
--- a/Engine/Plugins/Experimental/UAF/UAF/Source/UAF/Private/Variables/RigVMDispatch_GetObjectProperty.cpp
+++ b/Engine/Plugins/Experimental/UAF/UAF/Source/UAF/Private/Variables/RigVMDispatch_GetObjectProperty.cpp
@@ -1,9 +1,11 @@
-﻿// Copyright Epic Games, Inc. All Rights Reserved.
+// Copyright Epic Games, Inc. All Rights Reserved.
 
 #include "Variables/RigVMDispatch_GetObjectProperty.h"
 #include "RigVMCore/RigVMStruct.h"
 #include "Variables/AnimNextFieldPath.h"
 
+#include UE_INLINE_GENERATED_CPP_BY_NAME(RigVMDispatch_GetObjectProperty)
+
 FRigVMDispatch_GetObjectProperty::FRigVMDispatch_GetObjectProperty()
 {
 	FactoryScriptStruct = StaticStruct();
diff --git a/Engine/Plugins/Experimental/UAF/UAF/Source/UAF/Private/Variables/RigVMTrait_OverrideVariables.cpp b/Engine/Plugins/Experimental/UAF/UAF/Source/UAF/Private/Variables/RigVMTrait_OverrideVariables.cpp
--- a/Engine/Plugins/Experimental/UAF/UAF/Source/UAF/Private/Variables/RigVMTrait_OverrideVariables.cpp
+++ b/Engine/Plugins/Experimental/UAF/UAF/Source/UAF/Private/Variables/RigVMTrait_OverrideVariables.cpp
@@ -9,6 +9,8 @@
 #include "RigVMModel/RigVMController.h"
 #endif
 
+#include UE_INLINE_GENERATED_CPP_BY_NAME(RigVMTrait_OverrideVariables)
+
 #if WITH_EDITOR
 void FRigVMTrait_OverrideVariables::GetProgrammaticPins(URigVMController* InController, int32 InParentPinIndex, const URigVMPin* InTraitPin, const FString& InDefaultValue, FRigVMPinInfoArray& OutPinArray) const
 {
@@ -122,4 +124,4 @@ void FRigVMTrait_OverrideVariablesForStruct::GenerateOverrides(const FRigVMTrait
 			OutOverrides.Emplace(Struct, MoveTemp(StructOverrides));
 		}
 	}
-}
\ No newline at end of file
+}
diff --git a/Engine/Plugins/Experimental/UAF/UAF/Source/UAF/Private/Variables/UAFInstanceVariableContainer.cpp b/Engine/Plugins/Experimental/UAF/UAF/Source/UAF/Private/Variables/UAFInstanceVariableContainer.cpp
--- a/Engine/Plugins/Experimental/UAF/UAF/Source/UAF/Private/Variables/UAFInstanceVariableContainer.cpp
+++ b/Engine/Plugins/Experimental/UAF/UAF/Source/UAF/Private/Variables/UAFInstanceVariableContainer.cpp
@@ -1,4 +1,4 @@
-﻿// Copyright Epic Games, Inc. All Rights Reserved.
+// Copyright Epic Games, Inc. All Rights Reserved.
 
 #include "Variables/UAFInstanceVariableContainer.h"
 
@@ -8,6 +8,8 @@
 #include "Param/ParamType.h"
 #include "Variables/VariableOverridesCollection.h"
 
+#include UE_INLINE_GENERATED_CPP_BY_NAME(UAFInstanceVariableContainer)
+
 namespace UE::UAF::Private
 {
 	// Utility functions for type conversion, adapted from UE::StructUtils::Private in PropertyBag.cpp
@@ -1167,4 +1169,4 @@ void FUAFInstanceVariableContainer::AddStructReferencedObjects(FReferenceCollect
 		checkNoEntry();
 		break;
 	}
-}
\ No newline at end of file
+}
diff --git a/Engine/Plugins/Experimental/UAF/UAF/Source/UAF/Private/Variables/UAFInstanceVariableData.cpp b/Engine/Plugins/Experimental/UAF/UAF/Source/UAF/Private/Variables/UAFInstanceVariableData.cpp
--- a/Engine/Plugins/Experimental/UAF/UAF/Source/UAF/Private/Variables/UAFInstanceVariableData.cpp
+++ b/Engine/Plugins/Experimental/UAF/UAF/Source/UAF/Private/Variables/UAFInstanceVariableData.cpp
@@ -11,6 +11,8 @@
 #include "VariableOverridesCollection.h"
 #include "Variables/UAFInstanceVariableContainer.h"
 
+#include UE_INLINE_GENERATED_CPP_BY_NAME(UAFInstanceVariableData)
+
 void FUAFInstanceVariableData::Initialize(const FUAFAssetInstance& InInstance, const TSharedPtr<const UE::UAF::FVariableOverridesCollection>& InOverrides)
 {
 	using namespace UE::UAF;
@@ -452,4 +454,4 @@ void FUAFInstanceVariableData::AddStructReferencedObjects(FReferenceCollector& C
 	{
 		Collector.AddPropertyReferencesWithStructARO(FUAFInstanceVariableContainer::StaticStruct(), &SharedVariableSet.Get());
 	}
-}
\ No newline at end of file
+}
diff --git a/Engine/Plugins/Experimental/UAF/UAF/Source/UAF/Private/Variables/UAFInstanceVariableDataProxy.cpp b/Engine/Plugins/Experimental/UAF/UAF/Source/UAF/Private/Variables/UAFInstanceVariableDataProxy.cpp
--- a/Engine/Plugins/Experimental/UAF/UAF/Source/UAF/Private/Variables/UAFInstanceVariableDataProxy.cpp
+++ b/Engine/Plugins/Experimental/UAF/UAF/Source/UAF/Private/Variables/UAFInstanceVariableDataProxy.cpp
@@ -1,4 +1,4 @@
-﻿// Copyright Epic Games, Inc. All Rights Reserved.
+// Copyright Epic Games, Inc. All Rights Reserved.
 
 #include "Variables/UAFInstanceVariableDataProxy.h"
 
@@ -6,6 +6,8 @@
 #include "Variables/StructDataCache.h"
 #include "Variables/AnimNextVariableReference.h"
 
+#include UE_INLINE_GENERATED_CPP_BY_NAME(UAFInstanceVariableDataProxy)
+
 FUAFInstanceVariableContainerProxy::FUAFInstanceVariableContainerProxy(const FUAFInstanceVariableContainer& InVariables)
 	: FUAFInstanceVariableContainer(InVariables)
 {
@@ -201,4 +203,4 @@ void FUAFInstanceVariableDataProxy::CopyDirty()
 		DirtyFlags.SetRange(0, DirtyFlags.Num(), false);
 		bIsDirty = false;
 	}
-}
\ No newline at end of file
+}
diff --git a/Engine/Plugins/Experimental/UAF/UAF/Source/UAFEditor/Private/Common/AnimNextAssetFindReplaceVariables.cpp b/Engine/Plugins/Experimental/UAF/UAF/Source/UAFEditor/Private/Common/AnimNextAssetFindReplaceVariables.cpp
--- a/Engine/Plugins/Experimental/UAF/UAF/Source/UAFEditor/Private/Common/AnimNextAssetFindReplaceVariables.cpp
+++ b/Engine/Plugins/Experimental/UAF/UAF/Source/UAFEditor/Private/Common/AnimNextAssetFindReplaceVariables.cpp
@@ -19,6 +19,8 @@
 #include "Param/ParamUtils.h"
 #include "Variables/SVariablePickerCombo.h"
 
+#include UE_INLINE_GENERATED_CPP_BY_NAME(AnimNextAssetFindReplaceVariables)
+
 #define LOCTEXT_NAMESPACE "AnimNextAssetFindReplaceVariables"
 
 FString UAnimNextAssetFindReplaceVariables::GetFindResultStringFromAssetData(const FAssetData& InAssetData) const
diff --git a/Engine/Plugins/Experimental/UAF/UAF/Source/UAFEditor/Private/Common/AnimNextEntryAssetDefinitions.cpp b/Engine/Plugins/Experimental/UAF/UAF/Source/UAFEditor/Private/Common/AnimNextEntryAssetDefinitions.cpp
--- a/Engine/Plugins/Experimental/UAF/UAF/Source/UAFEditor/Private/Common/AnimNextEntryAssetDefinitions.cpp
+++ b/Engine/Plugins/Experimental/UAF/UAF/Source/UAFEditor/Private/Common/AnimNextEntryAssetDefinitions.cpp
@@ -1,7 +1,9 @@
-﻿// Copyright Epic Games, Inc. All Rights Reserved.
+// Copyright Epic Games, Inc. All Rights Reserved.
 
 #include "AnimNextEntryAssetDefinitions.h"
 
+#include UE_INLINE_GENERATED_CPP_BY_NAME(AnimNextEntryAssetDefinitions)
+
 #define LOCTEXT_NAMESPACE "AnimNextAssetDefinitions"
 
 FText UAssetDefinition_AnimNextVariableEntry::GetObjectDisplayNameText(UObject* Object) const
@@ -22,4 +24,4 @@ FText UAssetDefinition_AnimNextSharedVariablesEntry::GetObjectDisplayNameText(UO
 	return SharedVariables->GetDisplayName();
 }
 
-#undef LOCTEXT_NAMESPACE
\ No newline at end of file
+#undef LOCTEXT_NAMESPACE
diff --git a/Engine/Plugins/Experimental/UAF/UAF/Source/UAFEditor/Private/Common/GraphEditorSchemaActions.cpp b/Engine/Plugins/Experimental/UAF/UAF/Source/UAFEditor/Private/Common/GraphEditorSchemaActions.cpp
--- a/Engine/Plugins/Experimental/UAF/UAF/Source/UAFEditor/Private/Common/GraphEditorSchemaActions.cpp
+++ b/Engine/Plugins/Experimental/UAF/UAF/Source/UAFEditor/Private/Common/GraphEditorSchemaActions.cpp
@@ -20,6 +20,8 @@
 #include "ScopedTransaction.h"
 #include "Variables/AnimNextSharedVariables.h"
 
+#include UE_INLINE_GENERATED_CPP_BY_NAME(GraphEditorSchemaActions)
+
 #define LOCTEXT_NAMESPACE "AnimNextSchemaActions"
 
 // *** Base Schema Action ***
@@ -532,4 +534,4 @@ const FSlateBrush* FAnimNextSchemaAction_PromoteToVariable::GetIconBrush() const
 	return FAnimNextSchemaAction::GetIconBrush();
 }
 
-#undef LOCTEXT_NAMESPACE
\ No newline at end of file
+#undef LOCTEXT_NAMESPACE
diff --git a/Engine/Plugins/Experimental/UAF/UAF/Source/UAFEditor/Private/Module/AnimNextModuleAssetDefinition.cpp b/Engine/Plugins/Experimental/UAF/UAF/Source/UAFEditor/Private/Module/AnimNextModuleAssetDefinition.cpp
--- a/Engine/Plugins/Experimental/UAF/UAF/Source/UAFEditor/Private/Module/AnimNextModuleAssetDefinition.cpp
+++ b/Engine/Plugins/Experimental/UAF/UAF/Source/UAFEditor/Private/Module/AnimNextModuleAssetDefinition.cpp
@@ -11,6 +11,8 @@
 #include "UncookedOnlyUtils.h"
 #include "Workspace/AnimNextWorkspaceFactory.h"
 
+#include UE_INLINE_GENERATED_CPP_BY_NAME(AnimNextModuleAssetDefinition)
+
 #define LOCTEXT_NAMESPACE "AssetDefinition_AnimNextGraph"
 
 EAssetCommandResult UAssetDefinition_AnimNextModule::OpenAssets(const FAssetOpenArgs& OpenArgs) const
@@ -27,4 +29,4 @@ EAssetCommandResult UAssetDefinition_AnimNextModule::OpenAssets(const FAssetOpen
 	return EAssetCommandResult::Handled;
 }
 
-#undef LOCTEXT_NAMESPACE
\ No newline at end of file
+#undef LOCTEXT_NAMESPACE
diff --git a/Engine/Plugins/Experimental/UAF/UAF/Source/UAFEditor/Private/Module/AnimNextModuleFactory.cpp b/Engine/Plugins/Experimental/UAF/UAF/Source/UAFEditor/Private/Module/AnimNextModuleFactory.cpp
--- a/Engine/Plugins/Experimental/UAF/UAF/Source/UAFEditor/Private/Module/AnimNextModuleFactory.cpp
+++ b/Engine/Plugins/Experimental/UAF/UAF/Source/UAFEditor/Private/Module/AnimNextModuleFactory.cpp
@@ -7,6 +7,8 @@
 #include "Module/RigUnit_AnimNextModuleEvents.h"
 #include "Entries/AnimNextEventGraphEntry.h"
 
+#include UE_INLINE_GENERATED_CPP_BY_NAME(AnimNextModuleFactory)
+
 UAnimNextModuleFactory::UAnimNextModuleFactory()
 {
 	bCreateNew = true;
@@ -45,4 +47,4 @@ UObject* UAnimNextModuleFactory::FactoryCreateNew(UClass* Class, UObject* InPare
 	check(!EditorData->bErrorsDuringCompilation);
 
 	return NewModule;
-}
\ No newline at end of file
+}
diff --git a/Engine/Plugins/Experimental/UAF/UAF/Source/UAFEditor/Private/RewindDebugger/AnimNextModuleTrack.cpp b/Engine/Plugins/Experimental/UAF/UAF/Source/UAFEditor/Private/RewindDebugger/AnimNextModuleTrack.cpp
--- a/Engine/Plugins/Experimental/UAF/UAF/Source/UAFEditor/Private/RewindDebugger/AnimNextModuleTrack.cpp
+++ b/Engine/Plugins/Experimental/UAF/UAF/Source/UAFEditor/Private/RewindDebugger/AnimNextModuleTrack.cpp
@@ -2,6 +2,8 @@
 
 #include "AnimNextModuleTrack.h"
 
+#include UE_INLINE_GENERATED_CPP_BY_NAME(AnimNextModuleTrack)
+
 #if ANIMNEXT_TRACE_ENABLED
 #include "AnimNextProvider.h"
 #include "Editor.h"
diff --git a/Engine/Plugins/Experimental/UAF/UAF/Source/UAFEditor/Private/Variables/AnimNextSharedVariablesAssetDefinition.cpp b/Engine/Plugins/Experimental/UAF/UAF/Source/UAFEditor/Private/Variables/AnimNextSharedVariablesAssetDefinition.cpp
--- a/Engine/Plugins/Experimental/UAF/UAF/Source/UAFEditor/Private/Variables/AnimNextSharedVariablesAssetDefinition.cpp
+++ b/Engine/Plugins/Experimental/UAF/UAF/Source/UAFEditor/Private/Variables/AnimNextSharedVariablesAssetDefinition.cpp
@@ -1,4 +1,4 @@
-﻿// Copyright Epic Games, Inc. All Rights Reserved.
+// Copyright Epic Games, Inc. All Rights Reserved.
 
 #include "AnimNextSharedVariablesAssetDefinition.h"
 
@@ -11,6 +11,8 @@
 #include "UObject/SavePackage.h"
 #include "Workspace/AnimNextWorkspaceFactory.h"
 
+#include UE_INLINE_GENERATED_CPP_BY_NAME(AnimNextSharedVariablesAssetDefinition)
+
 #define LOCTEXT_NAMESPACE "AnimNextAssetDefinitions"
 
 EAssetCommandResult UAssetDefinition_AnimNextSharedVariables::OpenAssets(const FAssetOpenArgs& OpenArgs) const
@@ -150,4 +152,4 @@ static FDelayedAutoRegisterHelper AutoRegisterGraphMenuItems(EDelayedRegisterRun
 
 }
 
-#undef LOCTEXT_NAMESPACE
\ No newline at end of file
+#undef LOCTEXT_NAMESPACE
diff --git a/Engine/Plugins/Experimental/UAF/UAF/Source/UAFEditor/Private/Variables/AnimNextSharedVariablesFactory.cpp b/Engine/Plugins/Experimental/UAF/UAF/Source/UAFEditor/Private/Variables/AnimNextSharedVariablesFactory.cpp
--- a/Engine/Plugins/Experimental/UAF/UAF/Source/UAFEditor/Private/Variables/AnimNextSharedVariablesFactory.cpp
+++ b/Engine/Plugins/Experimental/UAF/UAF/Source/UAFEditor/Private/Variables/AnimNextSharedVariablesFactory.cpp
@@ -1,10 +1,12 @@
-﻿// Copyright Epic Games, Inc. All Rights Reserved.
+// Copyright Epic Games, Inc. All Rights Reserved.
 
 #include "Variables/AnimNextSharedVariablesFactory.h"
 #include "Variables/AnimNextSharedVariables.h"
 #include "Variables/AnimNextSharedVariables_EditorData.h"
 #include "Module/RigUnit_AnimNextModuleEvents.h"
 
+#include UE_INLINE_GENERATED_CPP_BY_NAME(AnimNextSharedVariablesFactory)
+
 UAnimNextSharedVariablesFactory::UAnimNextSharedVariablesFactory()
 {
 	bCreateNew = true;
@@ -38,4 +40,4 @@ UObject* UAnimNextSharedVariablesFactory::FactoryCreateNew(UClass* Class, UObjec
 	check(!EditorData->bErrorsDuringCompilation);
 
 	return NewDataInterface;
-}
\ No newline at end of file
+}
diff --git a/Engine/Plugins/Experimental/UAF/UAF/Source/UAFEditor/Private/Variables/AnimNextVariableSettings.cpp b/Engine/Plugins/Experimental/UAF/UAF/Source/UAFEditor/Private/Variables/AnimNextVariableSettings.cpp
--- a/Engine/Plugins/Experimental/UAF/UAF/Source/UAFEditor/Private/Variables/AnimNextVariableSettings.cpp
+++ b/Engine/Plugins/Experimental/UAF/UAF/Source/UAFEditor/Private/Variables/AnimNextVariableSettings.cpp
@@ -1,7 +1,9 @@
-﻿// Copyright Epic Games, Inc. All Rights Reserved.
+// Copyright Epic Games, Inc. All Rights Reserved.
 
 #include "AnimNextVariableSettings.h"
 
+#include UE_INLINE_GENERATED_CPP_BY_NAME(AnimNextVariableSettings)
+
 UAnimNextVariableSettings::UAnimNextVariableSettings()
 {
 }
@@ -24,4 +26,4 @@ FName UAnimNextVariableSettings::GetLastVariableName() const
 void UAnimNextVariableSettings::SetLastVariableName(FName InLastVariableName)
 {
 	LastVariableName = InLastVariableName;
-}
\ No newline at end of file
+}
diff --git a/Engine/Plugins/Experimental/UAF/UAF/Source/UAFEditor/Private/Workspace/AnimNextWorkspaceEditorMode.cpp b/Engine/Plugins/Experimental/UAF/UAF/Source/UAFEditor/Private/Workspace/AnimNextWorkspaceEditorMode.cpp
--- a/Engine/Plugins/Experimental/UAF/UAF/Source/UAFEditor/Private/Workspace/AnimNextWorkspaceEditorMode.cpp
+++ b/Engine/Plugins/Experimental/UAF/UAF/Source/UAFEditor/Private/Workspace/AnimNextWorkspaceEditorMode.cpp
@@ -18,6 +18,8 @@
 #include "Modules/ModuleManager.h"
 #include "Toolkits/BaseToolkit.h"
 
+#include UE_INLINE_GENERATED_CPP_BY_NAME(AnimNextWorkspaceEditorMode)
+
 #define LOCTEXT_NAMESPACE "AnimNextWorkspaceEditorMode"
 
 const FEditorModeID UAnimNextWorkspaceEditorMode::EM_AnimNextWorkspace("AnimNextWorkspace");
@@ -390,4 +392,4 @@ TSharedPtr<UE::UAF::Editor::IAssetCompilationHandler> UAnimNextWorkspaceEditorMo
 	return nullptr;
 }
 
-#undef LOCTEXT_NAMESPACE
\ No newline at end of file
+#undef LOCTEXT_NAMESPACE
diff --git a/Engine/Plugins/Experimental/UAF/UAF/Source/UAFEditor/Private/Workspace/AnimNextWorkspaceFactory.cpp b/Engine/Plugins/Experimental/UAF/UAF/Source/UAFEditor/Private/Workspace/AnimNextWorkspaceFactory.cpp
--- a/Engine/Plugins/Experimental/UAF/UAF/Source/UAFEditor/Private/Workspace/AnimNextWorkspaceFactory.cpp
+++ b/Engine/Plugins/Experimental/UAF/UAF/Source/UAFEditor/Private/Workspace/AnimNextWorkspaceFactory.cpp
@@ -4,6 +4,8 @@
 #include "AnimNextWorkspaceSchema.h"
 #include "AssetTypeCategories.h"
 
+#include UE_INLINE_GENERATED_CPP_BY_NAME(AnimNextWorkspaceFactory)
+
 #define LOCTEXT_NAMESPACE "AnimNextWorkspaceFactory"
 
 UAnimNextWorkspaceFactory::UAnimNextWorkspaceFactory()
@@ -32,4 +34,4 @@ const TArray<FText>& UAnimNextWorkspaceFactory::GetMenuCategorySubMenus() const
 	return Categories;
 }
 
-#undef LOCTEXT_NAMESPACE
\ No newline at end of file
+#undef LOCTEXT_NAMESPACE
diff --git a/Engine/Plugins/Experimental/UAF/UAF/Source/UAFEditor/Private/Workspace/AnimNextWorkspaceSchema.cpp b/Engine/Plugins/Experimental/UAF/UAF/Source/UAFEditor/Private/Workspace/AnimNextWorkspaceSchema.cpp
--- a/Engine/Plugins/Experimental/UAF/UAF/Source/UAFEditor/Private/Workspace/AnimNextWorkspaceSchema.cpp
+++ b/Engine/Plugins/Experimental/UAF/UAF/Source/UAFEditor/Private/Workspace/AnimNextWorkspaceSchema.cpp
@@ -8,6 +8,8 @@
 #include "IWorkspaceEditor.h"
 #include "StructUtils/InstancedStruct.h"
 
+#include UE_INLINE_GENERATED_CPP_BY_NAME(AnimNextWorkspaceSchema)
+
 #define LOCTEXT_NAMESPACE "AnimNextWorkspaceSchema"
 
 FText UAnimNextWorkspaceSchema::GetDisplayName() const
@@ -54,4 +56,4 @@ void UAnimNextWorkspaceSchema::OnLoadWorkspaceState(TSharedRef< UE::Workspace::I
 	EditorMode->PropagateAutoCompile(InWorkspaceEditor, EditorMode->State.bAutoCompile);
 }
 
-#undef LOCTEXT_NAMESPACE
\ No newline at end of file
+#undef LOCTEXT_NAMESPACE
diff --git a/Engine/Plugins/Experimental/UAF/UAF/Source/UAFTestSuite/Private/AnimNextTest.cpp b/Engine/Plugins/Experimental/UAF/UAF/Source/UAFTestSuite/Private/AnimNextTest.cpp
--- a/Engine/Plugins/Experimental/UAF/UAF/Source/UAFTestSuite/Private/AnimNextTest.cpp
+++ b/Engine/Plugins/Experimental/UAF/UAF/Source/UAFTestSuite/Private/AnimNextTest.cpp
@@ -12,6 +12,8 @@
 #include "Editor/Transactor.h"
 #endif
 
+#include UE_INLINE_GENERATED_CPP_BY_NAME(AnimNextTest)
+
 #if WITH_DEV_AUTOMATION_TESTS
 
 namespace UE::UAF::Tests
@@ -197,4 +199,4 @@ bool FDataRegistryTransformsTest::RunTest(const FString& InParameters)
 
 }
 
-#endif // WITH_DEV_AUTOMATION_TESTS
\ No newline at end of file
+#endif // WITH_DEV_AUTOMATION_TESTS
diff --git a/Engine/Plugins/Experimental/UAF/UAF/Source/UAFTestSuite/Private/AnimNextVariablesTest.cpp b/Engine/Plugins/Experimental/UAF/UAF/Source/UAFTestSuite/Private/AnimNextVariablesTest.cpp
--- a/Engine/Plugins/Experimental/UAF/UAF/Source/UAFTestSuite/Private/AnimNextVariablesTest.cpp
+++ b/Engine/Plugins/Experimental/UAF/UAF/Source/UAFTestSuite/Private/AnimNextVariablesTest.cpp
@@ -24,6 +24,8 @@
 
 // AnimNext Parameters Tests
 
+#include UE_INLINE_GENERATED_CPP_BY_NAME(AnimNextVariablesTest)
+
 #if WITH_DEV_AUTOMATION_TESTS
 
 namespace UE::UAF::Tests
diff --git a/Engine/Plugins/Experimental/UAF/UAF/Source/UAFUncookedOnly/Private/AnimNextAssetWorkspaceAssetUserData.cpp b/Engine/Plugins/Experimental/UAF/UAF/Source/UAFUncookedOnly/Private/AnimNextAssetWorkspaceAssetUserData.cpp
--- a/Engine/Plugins/Experimental/UAF/UAF/Source/UAFUncookedOnly/Private/AnimNextAssetWorkspaceAssetUserData.cpp
+++ b/Engine/Plugins/Experimental/UAF/UAF/Source/UAFUncookedOnly/Private/AnimNextAssetWorkspaceAssetUserData.cpp
@@ -9,6 +9,8 @@
 #include "Module/AnimNextModule_EditorData.h"
 #include "UObject/AssetRegistryTagsContext.h"
 
+#include UE_INLINE_GENERATED_CPP_BY_NAME(AnimNextAssetWorkspaceAssetUserData)
+
 void UAnimNextAssetWorkspaceAssetUserData::GetAssetRegistryTags(FAssetRegistryTagsContext Context) const
 {
 	Super::GetAssetRegistryTags(Context);
diff --git a/Engine/Plugins/Experimental/UAF/UAF/Source/UAFUncookedOnly/Private/AnimNextControllerBase.cpp b/Engine/Plugins/Experimental/UAF/UAF/Source/UAFUncookedOnly/Private/AnimNextControllerBase.cpp
--- a/Engine/Plugins/Experimental/UAF/UAF/Source/UAFUncookedOnly/Private/AnimNextControllerBase.cpp
+++ b/Engine/Plugins/Experimental/UAF/UAF/Source/UAFUncookedOnly/Private/AnimNextControllerBase.cpp
@@ -5,6 +5,8 @@
 #include "RigVMPythonUtils.h"
 #include "Variables/AnimNextSharedVariables.h"
 
+#include UE_INLINE_GENERATED_CPP_BY_NAME(AnimNextControllerBase)
+
 UAnimNextSharedVariableNode* UAnimNextControllerBase::AddAssetSharedVariableNode(const UAnimNextSharedVariables* InAsset, const FName& InVariableName, const FString& InCPPType, UObject* InCPPTypeObject, bool bIsGetter, const FString& InDefaultValue, const FVector2D& InPosition, const FString& InNodeName, bool bSetupUndoRedo, bool bPrintPythonCommand)
 {
 	if (InAsset == nullptr)
@@ -190,4 +192,4 @@ URigVMVariableNode* UAnimNextControllerBase::ReplaceSharedVariableNodeWithVariab
 	}
 
 	return VariableNode;
-}
\ No newline at end of file
+}
diff --git a/Engine/Plugins/Experimental/UAF/UAF/Source/UAFUncookedOnly/Private/AnimNextEdGraph.cpp b/Engine/Plugins/Experimental/UAF/UAF/Source/UAFUncookedOnly/Private/AnimNextEdGraph.cpp
--- a/Engine/Plugins/Experimental/UAF/UAF/Source/UAFUncookedOnly/Private/AnimNextEdGraph.cpp
+++ b/Engine/Plugins/Experimental/UAF/UAF/Source/UAFUncookedOnly/Private/AnimNextEdGraph.cpp
@@ -3,6 +3,8 @@
 #include "AnimNextEdGraph.h"
 #include "Module/AnimNextModule_EditorData.h"
 
+#include UE_INLINE_GENERATED_CPP_BY_NAME(AnimNextEdGraph)
+
 void UAnimNextEdGraph::PostLoad()
 {
 	Super::PostLoad();
diff --git a/Engine/Plugins/Experimental/UAF/UAF/Source/UAFUncookedOnly/Private/AnimNextEdGraphNode.cpp b/Engine/Plugins/Experimental/UAF/UAF/Source/UAFUncookedOnly/Private/AnimNextEdGraphNode.cpp
--- a/Engine/Plugins/Experimental/UAF/UAF/Source/UAFUncookedOnly/Private/AnimNextEdGraphNode.cpp
+++ b/Engine/Plugins/Experimental/UAF/UAF/Source/UAFUncookedOnly/Private/AnimNextEdGraphNode.cpp
@@ -1,3 +1,5 @@
 ﻿// Copyright Epic Games, Inc. All Rights Reserved.
 
 #include "AnimNextEdGraphNode.h"
+
+#include UE_INLINE_GENERATED_CPP_BY_NAME(AnimNextEdGraphNode)
diff --git a/Engine/Plugins/Experimental/UAF/UAF/Source/UAFUncookedOnly/Private/AnimNextEdGraphSchema.cpp b/Engine/Plugins/Experimental/UAF/UAF/Source/UAFUncookedOnly/Private/AnimNextEdGraphSchema.cpp
--- a/Engine/Plugins/Experimental/UAF/UAF/Source/UAFUncookedOnly/Private/AnimNextEdGraphSchema.cpp
+++ b/Engine/Plugins/Experimental/UAF/UAF/Source/UAFUncookedOnly/Private/AnimNextEdGraphSchema.cpp
@@ -5,6 +5,8 @@
 #include "AnimNextRigVMAsset.h"
 #include "Entries/AnimNextRigVMAssetEntry.h"
 
+#include UE_INLINE_GENERATED_CPP_BY_NAME(AnimNextEdGraphSchema)
+
 #define LOCTEXT_NAMESPACE "AnimNextEdGraphSchema"
 
 void UAnimNextEdGraphSchema::GetGraphDisplayInformation(const UEdGraph& Graph, /*out*/ FGraphDisplayInfo& DisplayInfo) const
@@ -18,4 +20,4 @@ void UAnimNextEdGraphSchema::GetGraphDisplayInformation(const UEdGraph& Graph, /
 	}
 }
 
-#undef LOCTEXT_NAMESPACE
\ No newline at end of file
+#undef LOCTEXT_NAMESPACE
diff --git a/Engine/Plugins/Experimental/UAF/UAF/Source/UAFUncookedOnly/Private/AnimNextRigVMAssetEditorData.cpp b/Engine/Plugins/Experimental/UAF/UAF/Source/UAFUncookedOnly/Private/AnimNextRigVMAssetEditorData.cpp
--- a/Engine/Plugins/Experimental/UAF/UAF/Source/UAFUncookedOnly/Private/AnimNextRigVMAssetEditorData.cpp
+++ b/Engine/Plugins/Experimental/UAF/UAF/Source/UAFUncookedOnly/Private/AnimNextRigVMAssetEditorData.cpp
@@ -47,6 +47,8 @@
 #include "PackageSourceControlHelper.h"
 #endif
 
+#include UE_INLINE_GENERATED_CPP_BY_NAME(AnimNextRigVMAssetEditorData)
+
 #define LOCTEXT_NAMESPACE "AnimNextRigVMAssetEditorData"
 
 void UAnimNextRigVMAssetEditorData::BroadcastModified(EAnimNextEditorDataNotifType InType, UObject* InSubject)
diff --git a/Engine/Plugins/Experimental/UAF/UAF/Source/UAFUncookedOnly/Private/AnimNextRigVMAssetSchema.cpp b/Engine/Plugins/Experimental/UAF/UAF/Source/UAFUncookedOnly/Private/AnimNextRigVMAssetSchema.cpp
--- a/Engine/Plugins/Experimental/UAF/UAF/Source/UAFUncookedOnly/Private/AnimNextRigVMAssetSchema.cpp
+++ b/Engine/Plugins/Experimental/UAF/UAF/Source/UAFUncookedOnly/Private/AnimNextRigVMAssetSchema.cpp
@@ -1,8 +1,10 @@
-﻿// Copyright Epic Games, Inc. All Rights Reserved.
+// Copyright Epic Games, Inc. All Rights Reserved.
 
 #include "AnimNextRigVMAssetSchema.h"
 #include "AnimNextExecuteContext.h"
 
+#include UE_INLINE_GENERATED_CPP_BY_NAME(AnimNextRigVMAssetSchema)
+
 UAnimNextRigVMAssetSchema::UAnimNextRigVMAssetSchema()
 {
 	SetExecuteContextStruct(FAnimNextExecuteContext::StaticStruct());
diff --git a/Engine/Plugins/Experimental/UAF/UAF/Source/UAFUncookedOnly/Private/AnimNextSharedVariableNode.cpp b/Engine/Plugins/Experimental/UAF/UAF/Source/UAFUncookedOnly/Private/AnimNextSharedVariableNode.cpp
--- a/Engine/Plugins/Experimental/UAF/UAF/Source/UAFUncookedOnly/Private/AnimNextSharedVariableNode.cpp
+++ b/Engine/Plugins/Experimental/UAF/UAF/Source/UAFUncookedOnly/Private/AnimNextSharedVariableNode.cpp
@@ -1,8 +1,10 @@
-﻿// Copyright Epic Games, Inc. All Rights Reserved.
+// Copyright Epic Games, Inc. All Rights Reserved.
 
 #include "AnimNextSharedVariableNode.h"
 #include "Variables/AnimNextSharedVariables.h"
 
+#include UE_INLINE_GENERATED_CPP_BY_NAME(AnimNextSharedVariableNode)
+
 FString UAnimNextSharedVariableNode::GetNodeSubTitle() const
 {
 	switch (Type)
diff --git a/Engine/Plugins/Experimental/UAF/UAF/Source/UAFUncookedOnly/Private/Entries/AnimNextCategoryEntry.cpp b/Engine/Plugins/Experimental/UAF/UAF/Source/UAFUncookedOnly/Private/Entries/AnimNextCategoryEntry.cpp
--- a/Engine/Plugins/Experimental/UAF/UAF/Source/UAFUncookedOnly/Private/Entries/AnimNextCategoryEntry.cpp
+++ b/Engine/Plugins/Experimental/UAF/UAF/Source/UAFUncookedOnly/Private/Entries/AnimNextCategoryEntry.cpp
@@ -3,6 +3,8 @@
 #include "Entries/AnimNextCategoryEntry.h"
 #include "AnimNextRigVMAssetEditorData.h"
 
+#include UE_INLINE_GENERATED_CPP_BY_NAME(AnimNextCategoryEntry)
+
 void UAnimNextCategoryEntry::Initialize(UAnimNextRigVMAssetEditorData* InEditorData)
 {
 	Super::Initialize(InEditorData);
diff --git a/Engine/Plugins/Experimental/UAF/UAF/Source/UAFUncookedOnly/Private/Entries/AnimNextEventGraphEntry.cpp b/Engine/Plugins/Experimental/UAF/UAF/Source/UAFUncookedOnly/Private/Entries/AnimNextEventGraphEntry.cpp
--- a/Engine/Plugins/Experimental/UAF/UAF/Source/UAFUncookedOnly/Private/Entries/AnimNextEventGraphEntry.cpp
+++ b/Engine/Plugins/Experimental/UAF/UAF/Source/UAFUncookedOnly/Private/Entries/AnimNextEventGraphEntry.cpp
@@ -7,6 +7,8 @@
 #include "AnimNextEdGraph.h"
 #include "AnimNextRigVMAssetEditorData.h"
 
+#include UE_INLINE_GENERATED_CPP_BY_NAME(AnimNextEventGraphEntry)
+
 FText UAnimNextEventGraphEntry::GetDisplayName() const
 {
 	return FText::FromName(GraphName);
@@ -52,4 +54,4 @@ void UAnimNextEventGraphEntry::SetRigVMGraph(URigVMGraph* InGraph)
 void UAnimNextEventGraphEntry::SetEdGraph(URigVMEdGraph* InGraph)
 {
 	EdGraph = CastChecked<UAnimNextEdGraph>(InGraph, ECastCheckedType::NullAllowed);
-}
\ No newline at end of file
+}
diff --git a/Engine/Plugins/Experimental/UAF/UAF/Source/UAFUncookedOnly/Private/Entries/AnimNextRigVMAssetEntry.cpp b/Engine/Plugins/Experimental/UAF/UAF/Source/UAFUncookedOnly/Private/Entries/AnimNextRigVMAssetEntry.cpp
--- a/Engine/Plugins/Experimental/UAF/UAF/Source/UAFUncookedOnly/Private/Entries/AnimNextRigVMAssetEntry.cpp
+++ b/Engine/Plugins/Experimental/UAF/UAF/Source/UAFUncookedOnly/Private/Entries/AnimNextRigVMAssetEntry.cpp
@@ -4,6 +4,8 @@
 #include "AnimNextRigVMAssetEditorData.h"
 #include "Misc/TransactionObjectEvent.h"
 
+#include UE_INLINE_GENERATED_CPP_BY_NAME(AnimNextRigVMAssetEntry)
+
 #define LOCTEXT_NAMESPACE "AnimNextRigVMAssetEntry"
 
 void UAnimNextRigVMAssetEntry::Initialize(UAnimNextRigVMAssetEditorData* InEditorData)
@@ -55,4 +57,4 @@ void UAnimNextRigVMAssetEntry::BroadcastModified(EAnimNextEditorDataNotifType In
 	}
 }
 
-#undef LOCTEXT_NAMESPACE
\ No newline at end of file
+#undef LOCTEXT_NAMESPACE
diff --git a/Engine/Plugins/Experimental/UAF/UAF/Source/UAFUncookedOnly/Private/Entries/AnimNextSharedVariablesEntry.cpp b/Engine/Plugins/Experimental/UAF/UAF/Source/UAFUncookedOnly/Private/Entries/AnimNextSharedVariablesEntry.cpp
--- a/Engine/Plugins/Experimental/UAF/UAF/Source/UAFUncookedOnly/Private/Entries/AnimNextSharedVariablesEntry.cpp
+++ b/Engine/Plugins/Experimental/UAF/UAF/Source/UAFUncookedOnly/Private/Entries/AnimNextSharedVariablesEntry.cpp
@@ -1,4 +1,4 @@
-﻿// Copyright Epic Games, Inc. All Rights Reserved.
+// Copyright Epic Games, Inc. All Rights Reserved.
 
 #include "Entries/AnimNextSharedVariablesEntry.h"
 
@@ -13,6 +13,8 @@
 #include "Param/ParamType.h"
 #include "Variables/AnimNextSharedVariables.h"
 
+#include UE_INLINE_GENERATED_CPP_BY_NAME(AnimNextSharedVariablesEntry)
+
 #define LOCTEXT_NAMESPACE "AnimNextSharedVariablesEntry"
 
 void UAnimNextSharedVariablesEntry::Initialize(UAnimNextRigVMAssetEditorData* InEditorData)
@@ -150,4 +152,4 @@ void UAnimNextSharedVariablesEntry::HandlePatchComplete()
 }
 #endif
 
-#undef LOCTEXT_NAMESPACE
\ No newline at end of file
+#undef LOCTEXT_NAMESPACE
diff --git a/Engine/Plugins/Experimental/UAF/UAF/Source/UAFUncookedOnly/Private/Entries/AnimNextVariableEntry.cpp b/Engine/Plugins/Experimental/UAF/UAF/Source/UAFUncookedOnly/Private/Entries/AnimNextVariableEntry.cpp
--- a/Engine/Plugins/Experimental/UAF/UAF/Source/UAFUncookedOnly/Private/Entries/AnimNextVariableEntry.cpp
+++ b/Engine/Plugins/Experimental/UAF/UAF/Source/UAFUncookedOnly/Private/Entries/AnimNextVariableEntry.cpp
@@ -8,6 +8,8 @@
 #include "Module/AnimNextModule.h"
 #include "Variables/AnimNextUniversalObjectLocatorBindingData.h"
 
+#include UE_INLINE_GENERATED_CPP_BY_NAME(AnimNextVariableEntry)
+
 const FLazyName IAnimNextRigVMVariableInterface::ValueName("Value");
 
 UAnimNextVariableEntry::UAnimNextVariableEntry()
@@ -284,4 +286,4 @@ void UAnimNextVariableEntry::PostEditChangeProperty(FPropertyChangedEvent& Prope
 	}
 }
 
-#endif
\ No newline at end of file
+#endif
diff --git a/Engine/Plugins/Experimental/UAF/UAF/Source/UAFUncookedOnly/Private/Module/AnimNextModule_EditorData.cpp b/Engine/Plugins/Experimental/UAF/UAF/Source/UAFUncookedOnly/Private/Module/AnimNextModule_EditorData.cpp
--- a/Engine/Plugins/Experimental/UAF/UAF/Source/UAFUncookedOnly/Private/Module/AnimNextModule_EditorData.cpp
+++ b/Engine/Plugins/Experimental/UAF/UAF/Source/UAFUncookedOnly/Private/Module/AnimNextModule_EditorData.cpp
@@ -15,6 +15,8 @@
 #include "UObject/LinkerLoad.h"
 #include "Variables/AnimNextUniversalObjectLocatorBindingData.h"
 
+#include UE_INLINE_GENERATED_CPP_BY_NAME(AnimNextModule_EditorData)
+
 TConstArrayView<TSubclassOf<UAnimNextRigVMAssetEntry>> UAnimNextModule_EditorData::GetEntryClasses() const
 {
 	static const TSubclassOf<UAnimNextRigVMAssetEntry> Classes[] =
@@ -101,4 +103,4 @@ void UAnimNextModule_EditorData::CustomizeNewAssetEntry(UAnimNextRigVMAssetEntry
 	}
 
 	VariableEntry->SetBindingType(FAnimNextUniversalObjectLocatorBindingData::StaticStruct(), false);
-}
\ No newline at end of file
+}
diff --git a/Engine/Plugins/Experimental/UAF/UAF/Source/UAFUncookedOnly/Private/Variables/AnimNextSharedVariables_EditorData.cpp b/Engine/Plugins/Experimental/UAF/UAF/Source/UAFUncookedOnly/Private/Variables/AnimNextSharedVariables_EditorData.cpp
--- a/Engine/Plugins/Experimental/UAF/UAF/Source/UAFUncookedOnly/Private/Variables/AnimNextSharedVariables_EditorData.cpp
+++ b/Engine/Plugins/Experimental/UAF/UAF/Source/UAFUncookedOnly/Private/Variables/AnimNextSharedVariables_EditorData.cpp
@@ -1,10 +1,12 @@
-﻿// Copyright Epic Games, Inc. All Rights Reserved.
+// Copyright Epic Games, Inc. All Rights Reserved.
 
 #include "Variables/AnimNextSharedVariables_EditorData.h"
 
 #include "Entries/AnimNextSharedVariablesEntry.h"
 #include "Entries/AnimNextVariableEntry.h"
 
+#include UE_INLINE_GENERATED_CPP_BY_NAME(AnimNextSharedVariables_EditorData)
+
 TConstArrayView<TSubclassOf<UAnimNextRigVMAssetEntry>> UAnimNextSharedVariables_EditorData::GetEntryClasses() const
 {
 	static const TSubclassOf<UAnimNextRigVMAssetEntry> Classes[] =
@@ -32,4 +34,4 @@ void UAnimNextSharedVariables_EditorData::CustomizeNewAssetEntry(UAnimNextRigVMA
 
 	// Force all variables in 'pure' shared variables assets to be public
 	VariableEntry->SetExportAccessSpecifier(EAnimNextExportAccessSpecifier::Public, false);
-}
\ No newline at end of file
+}
diff --git a/Engine/Plugins/Experimental/UAF/UAF/Source/UAFUncookedOnly/Private/Variables/AnimNextUniversalObjectLocatorBindingData.cpp b/Engine/Plugins/Experimental/UAF/UAF/Source/UAFUncookedOnly/Private/Variables/AnimNextUniversalObjectLocatorBindingData.cpp
--- a/Engine/Plugins/Experimental/UAF/UAF/Source/UAFUncookedOnly/Private/Variables/AnimNextUniversalObjectLocatorBindingData.cpp
+++ b/Engine/Plugins/Experimental/UAF/UAF/Source/UAFUncookedOnly/Private/Variables/AnimNextUniversalObjectLocatorBindingData.cpp
@@ -1,8 +1,10 @@
-﻿// Copyright Epic Games, Inc. All Rights Reserved.
+// Copyright Epic Games, Inc. All Rights Reserved.
 
 #include "Variables/AnimNextUniversalObjectLocatorBindingData.h"
 #include "Kismet2/BlueprintEditorUtils.h"
 
+#include UE_INLINE_GENERATED_CPP_BY_NAME(AnimNextUniversalObjectLocatorBindingData)
+
 #if WITH_EDITORONLY_DATA
 
 bool FAnimNextUniversalObjectLocatorBindingData::IsThreadSafe() const
@@ -31,4 +33,4 @@ bool FAnimNextUniversalObjectLocatorBindingData::IsThreadSafe() const
 	*/
 }
 
-#endif
\ No newline at end of file
+#endif
diff --git a/Engine/Plugins/Experimental/UAF/UAF/Source/UAFUncookedOnly/Private/Variables/K2Node_AnimNextComponentSetVariable.cpp b/Engine/Plugins/Experimental/UAF/UAF/Source/UAFUncookedOnly/Private/Variables/K2Node_AnimNextComponentSetVariable.cpp
--- a/Engine/Plugins/Experimental/UAF/UAF/Source/UAFUncookedOnly/Private/Variables/K2Node_AnimNextComponentSetVariable.cpp
+++ b/Engine/Plugins/Experimental/UAF/UAF/Source/UAFUncookedOnly/Private/Variables/K2Node_AnimNextComponentSetVariable.cpp
@@ -1,4 +1,4 @@
-﻿// Copyright Epic Games, Inc. All Rights Reserved.
+// Copyright Epic Games, Inc. All Rights Reserved.
 
 #include "Variables/K2Node_AnimNextComponentSetVariable.h"
 
@@ -14,6 +14,8 @@
 #include "Param/ParamUtils.h"
 #include "Variables/AnimNextVariableReference.h"
 
+#include UE_INLINE_GENERATED_CPP_BY_NAME(K2Node_AnimNextComponentSetVariable)
+
 #define LOCTEXT_NAMESPACE "K2Node_AnimNextComponentSetVariable"
 
 UK2Node_AnimNextComponentSetVariable::UK2Node_AnimNextComponentSetVariable()
diff --git a/Engine/Plugins/Experimental/UAF/UAFAnimGraph/Source/UAFAnimGraph/Private/AnimNextAnimGraphSettings.cpp b/Engine/Plugins/Experimental/UAF/UAFAnimGraph/Source/UAFAnimGraph/Private/AnimNextAnimGraphSettings.cpp
--- a/Engine/Plugins/Experimental/UAF/UAFAnimGraph/Source/UAFAnimGraph/Private/AnimNextAnimGraphSettings.cpp
+++ b/Engine/Plugins/Experimental/UAF/UAFAnimGraph/Source/UAFAnimGraph/Private/AnimNextAnimGraphSettings.cpp
@@ -3,6 +3,8 @@
 #include "AnimNextAnimGraphSettings.h"
 #include "Factory/AnimGraphFactory.h"
 
+#include UE_INLINE_GENERATED_CPP_BY_NAME(AnimNextAnimGraphSettings)
+
 TArray<UClass*> UAnimNextAnimGraphSettings::GetAllowedAssetClasses()
 {
 	TArray<UClass*> Array;
diff --git a/Engine/Plugins/Experimental/UAF/UAFAnimGraph/Source/UAFAnimGraph/Private/EvaluationVM/SerializableEvaluationProgram.cpp b/Engine/Plugins/Experimental/UAF/UAFAnimGraph/Source/UAFAnimGraph/Private/EvaluationVM/SerializableEvaluationProgram.cpp
--- a/Engine/Plugins/Experimental/UAF/UAFAnimGraph/Source/UAFAnimGraph/Private/EvaluationVM/SerializableEvaluationProgram.cpp
+++ b/Engine/Plugins/Experimental/UAF/UAFAnimGraph/Source/UAFAnimGraph/Private/EvaluationVM/SerializableEvaluationProgram.cpp
@@ -3,6 +3,8 @@
 #include "EvaluationVM/SerializableEvaluationProgram.h"
 #include "EvaluationVM/EvaluationProgram.h"
 
+#include UE_INLINE_GENERATED_CPP_BY_NAME(SerializableEvaluationProgram)
+
 FSerializableEvaluationProgram::FSerializableEvaluationProgram(const UE::UAF::FEvaluationProgram& Other)
 {
 	for (const TSharedPtr<FAnimNextEvaluationTask>& Task : Other.Tasks)
@@ -10,4 +12,4 @@ FSerializableEvaluationProgram::FSerializableEvaluationProgram(const UE::UAF::FE
 		Tasks.Add(FInstancedStruct(Task->GetStruct()));
 		Task->GetStruct()->CopyScriptStruct(Tasks.Last().GetMutableMemory(), &*Task);
 	}
-}
\ No newline at end of file
+}
diff --git a/Engine/Plugins/Experimental/UAF/UAFAnimGraph/Source/UAFAnimGraph/Private/EvaluationVM/Tasks/ApplyAdditiveKeyframe.cpp b/Engine/Plugins/Experimental/UAF/UAFAnimGraph/Source/UAFAnimGraph/Private/EvaluationVM/Tasks/ApplyAdditiveKeyframe.cpp
--- a/Engine/Plugins/Experimental/UAF/UAFAnimGraph/Source/UAFAnimGraph/Private/EvaluationVM/Tasks/ApplyAdditiveKeyframe.cpp
+++ b/Engine/Plugins/Experimental/UAF/UAFAnimGraph/Source/UAFAnimGraph/Private/EvaluationVM/Tasks/ApplyAdditiveKeyframe.cpp
@@ -6,6 +6,8 @@
 #include "EvaluationVM/KeyframeState.h"
 #include "TransformArrayOperations.h"
 
+#include UE_INLINE_GENERATED_CPP_BY_NAME(ApplyAdditiveKeyframe)
+
 FAnimNextApplyAdditiveKeyframeTask FAnimNextApplyAdditiveKeyframeTask::Make(float BlendWeight)
 {
 	FAnimNextApplyAdditiveKeyframeTask Task;
diff --git a/Engine/Plugins/Experimental/UAF/UAFAnimGraph/Source/UAFAnimGraph/Private/EvaluationVM/Tasks/BlendKeyframes.cpp b/Engine/Plugins/Experimental/UAF/UAFAnimGraph/Source/UAFAnimGraph/Private/EvaluationVM/Tasks/BlendKeyframes.cpp
--- a/Engine/Plugins/Experimental/UAF/UAFAnimGraph/Source/UAFAnimGraph/Private/EvaluationVM/Tasks/BlendKeyframes.cpp
+++ b/Engine/Plugins/Experimental/UAF/UAFAnimGraph/Source/UAFAnimGraph/Private/EvaluationVM/Tasks/BlendKeyframes.cpp
@@ -7,6 +7,8 @@
 #include "TransformArrayOperations.h"
 #include "Animation/InputScaleBias.h"
 
+#include UE_INLINE_GENERATED_CPP_BY_NAME(BlendKeyframes)
+
 FAnimNextBlendTwoKeyframesTask FAnimNextBlendTwoKeyframesTask::Make(float InterpolationAlpha)
 {
 	FAnimNextBlendTwoKeyframesTask Task;
diff --git a/Engine/Plugins/Experimental/UAF/UAFAnimGraph/Source/UAFAnimGraph/Private/EvaluationVM/Tasks/BlendKeyframesPerBone.cpp b/Engine/Plugins/Experimental/UAF/UAFAnimGraph/Source/UAFAnimGraph/Private/EvaluationVM/Tasks/BlendKeyframesPerBone.cpp
--- a/Engine/Plugins/Experimental/UAF/UAFAnimGraph/Source/UAFAnimGraph/Private/EvaluationVM/Tasks/BlendKeyframesPerBone.cpp
+++ b/Engine/Plugins/Experimental/UAF/UAFAnimGraph/Source/UAFAnimGraph/Private/EvaluationVM/Tasks/BlendKeyframesPerBone.cpp
@@ -12,6 +12,8 @@
 #include "EvaluationVM/KeyframeState.h"
 #include "TransformArrayOperations.h"
 
+#include UE_INLINE_GENERATED_CPP_BY_NAME(BlendKeyframesPerBone)
+
 FAnimNextBlendOverwriteKeyframePerBoneWithScaleTask FAnimNextBlendOverwriteKeyframePerBoneWithScaleTask::Make(const IBlendProfileInterface* BlendProfile, const FBlendSampleData& BlendData, float ScaleFactor)
 {
 	FAnimNextBlendOverwriteKeyframePerBoneWithScaleTask Task;
@@ -441,4 +443,4 @@ void FAnimNextBlendKeyframePerBoneWithScaleTask::Execute(UE::UAF::FEvaluationVM&
 	}
 
 	VM.PushValue(KEYFRAME_STACK_NAME, MoveTemp(KeyframeA));
-}
\ No newline at end of file
+}
diff --git a/Engine/Plugins/Experimental/UAF/UAFAnimGraph/Source/UAFAnimGraph/Private/EvaluationVM/Tasks/ConvertRotationsToLocalSpace.cpp b/Engine/Plugins/Experimental/UAF/UAFAnimGraph/Source/UAFAnimGraph/Private/EvaluationVM/Tasks/ConvertRotationsToLocalSpace.cpp
--- a/Engine/Plugins/Experimental/UAF/UAFAnimGraph/Source/UAFAnimGraph/Private/EvaluationVM/Tasks/ConvertRotationsToLocalSpace.cpp
+++ b/Engine/Plugins/Experimental/UAF/UAFAnimGraph/Source/UAFAnimGraph/Private/EvaluationVM/Tasks/ConvertRotationsToLocalSpace.cpp
@@ -9,6 +9,8 @@
 #include "ConvertRotationSpace.ispc.generated.h"
 #endif
 
+#include UE_INLINE_GENERATED_CPP_BY_NAME(ConvertRotationsToLocalSpace)
+
 FAnimNextConvertRotationsToLocalSpaceTask FAnimNextConvertRotationsToLocalSpaceTask::Make(const int32 NumPoses)
 {
 	FAnimNextConvertRotationsToLocalSpaceTask Task;
@@ -55,4 +57,4 @@ void FAnimNextConvertRotationsToLocalSpaceTask::Execute(UE::UAF::FEvaluationVM&
 		}
 #endif // UE_ANIM_CONVERT_ROTATION_SPACE_TASKS_ISPC
 	}
-}
\ No newline at end of file
+}
diff --git a/Engine/Plugins/Experimental/UAF/UAFAnimGraph/Source/UAFAnimGraph/Private/EvaluationVM/Tasks/ConvertRotationsToMeshSpace.cpp b/Engine/Plugins/Experimental/UAF/UAFAnimGraph/Source/UAFAnimGraph/Private/EvaluationVM/Tasks/ConvertRotationsToMeshSpace.cpp
--- a/Engine/Plugins/Experimental/UAF/UAFAnimGraph/Source/UAFAnimGraph/Private/EvaluationVM/Tasks/ConvertRotationsToMeshSpace.cpp
+++ b/Engine/Plugins/Experimental/UAF/UAFAnimGraph/Source/UAFAnimGraph/Private/EvaluationVM/Tasks/ConvertRotationsToMeshSpace.cpp
@@ -9,6 +9,8 @@
 #include "ConvertRotationSpace.ispc.generated.h"
 #endif
 
+#include UE_INLINE_GENERATED_CPP_BY_NAME(ConvertRotationsToMeshSpace)
+
 FAnimNextConvertRotationsToMeshSpaceTask FAnimNextConvertRotationsToMeshSpaceTask::Make(const int32 NumPoses)
 {
 	FAnimNextConvertRotationsToMeshSpaceTask Task;
@@ -55,4 +57,4 @@ void FAnimNextConvertRotationsToMeshSpaceTask::Execute(UE::UAF::FEvaluationVM& V
 		}
 #endif // UE_ANIM_CONVERT_ROTATION_SPACE_TASKS_ISPC
 	}
-}
\ No newline at end of file
+}
diff --git a/Engine/Plugins/Experimental/UAF/UAFAnimGraph/Source/UAFAnimGraph/Private/EvaluationVM/Tasks/DeadBlending.cpp b/Engine/Plugins/Experimental/UAF/UAFAnimGraph/Source/UAFAnimGraph/Private/EvaluationVM/Tasks/DeadBlending.cpp
--- a/Engine/Plugins/Experimental/UAF/UAFAnimGraph/Source/UAFAnimGraph/Private/EvaluationVM/Tasks/DeadBlending.cpp
+++ b/Engine/Plugins/Experimental/UAF/UAFAnimGraph/Source/UAFAnimGraph/Private/EvaluationVM/Tasks/DeadBlending.cpp
@@ -8,6 +8,8 @@
 
 #include "Traits/Inertialization.h"
 
+#include UE_INLINE_GENERATED_CPP_BY_NAME(DeadBlending)
+
 #ifndef UE_ANIM_NEXT_DEAD_BLENDING_ISPC
 #define UE_ANIM_NEXT_DEAD_BLENDING_ISPC INTEL_ISPC
 //#define UE_ANIM_NEXT_DEAD_BLENDING_ISPC 0
@@ -573,4 +575,4 @@ void FAnimNextDeadBlendingApplyTask::Execute(UE::UAF::FEvaluationVM& VM) const
 	}
 }
 
-#undef UE_ANIM_NEXT_DEAD_BLENDING_ISPC
\ No newline at end of file
+#undef UE_ANIM_NEXT_DEAD_BLENDING_ISPC
diff --git a/Engine/Plugins/Experimental/UAF/UAFAnimGraph/Source/UAFAnimGraph/Private/EvaluationVM/Tasks/ExecuteProgram.cpp b/Engine/Plugins/Experimental/UAF/UAFAnimGraph/Source/UAFAnimGraph/Private/EvaluationVM/Tasks/ExecuteProgram.cpp
--- a/Engine/Plugins/Experimental/UAF/UAFAnimGraph/Source/UAFAnimGraph/Private/EvaluationVM/Tasks/ExecuteProgram.cpp
+++ b/Engine/Plugins/Experimental/UAF/UAFAnimGraph/Source/UAFAnimGraph/Private/EvaluationVM/Tasks/ExecuteProgram.cpp
@@ -5,6 +5,8 @@
 #include "EvaluationVM/EvaluationProgram.h"
 #include "EvaluationVM/EvaluationVM.h"
 
+#include UE_INLINE_GENERATED_CPP_BY_NAME(ExecuteProgram)
+
 FAnimNextExecuteProgramTask FAnimNextExecuteProgramTask::Make(TSharedPtr<const UE::UAF::FEvaluationProgram> Program)
 {
 	FAnimNextExecuteProgramTask Task;
diff --git a/Engine/Plugins/Experimental/UAF/UAFAnimGraph/Source/UAFAnimGraph/Private/EvaluationVM/Tasks/NormalizeRotations.cpp b/Engine/Plugins/Experimental/UAF/UAFAnimGraph/Source/UAFAnimGraph/Private/EvaluationVM/Tasks/NormalizeRotations.cpp
--- a/Engine/Plugins/Experimental/UAF/UAFAnimGraph/Source/UAFAnimGraph/Private/EvaluationVM/Tasks/NormalizeRotations.cpp
+++ b/Engine/Plugins/Experimental/UAF/UAFAnimGraph/Source/UAFAnimGraph/Private/EvaluationVM/Tasks/NormalizeRotations.cpp
@@ -6,6 +6,8 @@
 #include "EvaluationVM/KeyframeState.h"
 #include "TransformArrayOperations.h"
 
+#include UE_INLINE_GENERATED_CPP_BY_NAME(NormalizeRotations)
+
 void FAnimNextNormalizeKeyframeRotationsTask::Execute(UE::UAF::FEvaluationVM& VM) const
 {
 	using namespace UE::UAF;
diff --git a/Engine/Plugins/Experimental/UAF/UAFAnimGraph/Source/UAFAnimGraph/Private/EvaluationVM/Tasks/PushAnimSequenceKeyframe.cpp b/Engine/Plugins/Experimental/UAF/UAFAnimGraph/Source/UAFAnimGraph/Private/EvaluationVM/Tasks/PushAnimSequenceKeyframe.cpp
--- a/Engine/Plugins/Experimental/UAF/UAFAnimGraph/Source/UAFAnimGraph/Private/EvaluationVM/Tasks/PushAnimSequenceKeyframe.cpp
+++ b/Engine/Plugins/Experimental/UAF/UAFAnimGraph/Source/UAFAnimGraph/Private/EvaluationVM/Tasks/PushAnimSequenceKeyframe.cpp
@@ -9,6 +9,8 @@
 #include "EvaluationVM/EvaluationVM.h"
 #include "EvaluationVM/KeyframeState.h"
 
+#include UE_INLINE_GENERATED_CPP_BY_NAME(PushAnimSequenceKeyframe)
+
 FAnimNextAnimSequenceKeyframeTask FAnimNextAnimSequenceKeyframeTask::MakeFromSampleTime(TWeakObjectPtr<const UAnimSequence> AnimSequence, double SampleTime, bool bInterpolate)
 {
 	FAnimNextAnimSequenceKeyframeTask Task;
diff --git a/Engine/Plugins/Experimental/UAF/UAFAnimGraph/Source/UAFAnimGraph/Private/EvaluationVM/Tasks/PushPose.cpp b/Engine/Plugins/Experimental/UAF/UAFAnimGraph/Source/UAFAnimGraph/Private/EvaluationVM/Tasks/PushPose.cpp
--- a/Engine/Plugins/Experimental/UAF/UAFAnimGraph/Source/UAFAnimGraph/Private/EvaluationVM/Tasks/PushPose.cpp
+++ b/Engine/Plugins/Experimental/UAF/UAFAnimGraph/Source/UAFAnimGraph/Private/EvaluationVM/Tasks/PushPose.cpp
@@ -3,6 +3,8 @@
 #include "EvaluationVM/Tasks/PushPose.h"
 #include "EvaluationVM/EvaluationVM.h"
 
+#include UE_INLINE_GENERATED_CPP_BY_NAME(PushPose)
+
 DEFINE_STAT(STAT_AnimNext_Task_PushPose);
 
 FAnimNextPushPoseTask::FAnimNextPushPoseTask(const FAnimNextGraphLODPose* InGraphPose)
@@ -34,4 +36,4 @@ void FAnimNextPushPoseTask::Execute(UE::UAF::FEvaluationVM& VM) const
 	Keyframe.Attributes.CopyFrom(GraphPose->Attributes);
 
 	VM.PushValue(KEYFRAME_STACK_NAME, MakeUnique<FKeyframeState>(MoveTemp(Keyframe)));
-}
\ No newline at end of file
+}
diff --git a/Engine/Plugins/Experimental/UAF/UAFAnimGraph/Source/UAFAnimGraph/Private/EvaluationVM/Tasks/PushReferenceKeyframe.cpp b/Engine/Plugins/Experimental/UAF/UAFAnimGraph/Source/UAFAnimGraph/Private/EvaluationVM/Tasks/PushReferenceKeyframe.cpp
--- a/Engine/Plugins/Experimental/UAF/UAFAnimGraph/Source/UAFAnimGraph/Private/EvaluationVM/Tasks/PushReferenceKeyframe.cpp
+++ b/Engine/Plugins/Experimental/UAF/UAFAnimGraph/Source/UAFAnimGraph/Private/EvaluationVM/Tasks/PushReferenceKeyframe.cpp
@@ -4,6 +4,8 @@
 
 #include "EvaluationVM/EvaluationVM.h"
 
+#include UE_INLINE_GENERATED_CPP_BY_NAME(PushReferenceKeyframe)
+
 FAnimNextPushReferenceKeyframeTask FAnimNextPushReferenceKeyframeTask::MakeFromSkeleton()
 {
 	return FAnimNextPushReferenceKeyframeTask();
diff --git a/Engine/Plugins/Experimental/UAF/UAFAnimGraph/Source/UAFAnimGraph/Private/EvaluationVM/Tasks/StoreKeyframe.cpp b/Engine/Plugins/Experimental/UAF/UAFAnimGraph/Source/UAFAnimGraph/Private/EvaluationVM/Tasks/StoreKeyframe.cpp
--- a/Engine/Plugins/Experimental/UAF/UAFAnimGraph/Source/UAFAnimGraph/Private/EvaluationVM/Tasks/StoreKeyframe.cpp
+++ b/Engine/Plugins/Experimental/UAF/UAFAnimGraph/Source/UAFAnimGraph/Private/EvaluationVM/Tasks/StoreKeyframe.cpp
@@ -6,6 +6,8 @@
 #include "EvaluationVM/KeyframeState.h"
 #include "TransformArrayOperations.h"
 
+#include UE_INLINE_GENERATED_CPP_BY_NAME(StoreKeyframe)
+
 FAnimNextSwapTransformsTask FAnimNextSwapTransformsTask::Make(UE::UAF::FTransformArraySoAHeap* A, UE::UAF::FTransformArraySoAHeap* B)
 {
 	FAnimNextSwapTransformsTask Task;
diff --git a/Engine/Plugins/Experimental/UAF/UAFAnimGraph/Source/UAFAnimGraph/Private/Factory/AnimNextAnimGraphBuilder.cpp b/Engine/Plugins/Experimental/UAF/UAFAnimGraph/Source/UAFAnimGraph/Private/Factory/AnimNextAnimGraphBuilder.cpp
--- a/Engine/Plugins/Experimental/UAF/UAFAnimGraph/Source/UAFAnimGraph/Private/Factory/AnimNextAnimGraphBuilder.cpp
+++ b/Engine/Plugins/Experimental/UAF/UAFAnimGraph/Source/UAFAnimGraph/Private/Factory/AnimNextAnimGraphBuilder.cpp
@@ -1,9 +1,11 @@
-﻿// Copyright Epic Games, Inc. All Rights Reserved.
+// Copyright Epic Games, Inc. All Rights Reserved.
 
 #include "Factory/AnimNextAnimGraphBuilder.h"
 #include "Graph/AnimNextAnimationGraph.h"
 #include "TraitCore/TraitWriter.h"
 
+#include UE_INLINE_GENERATED_CPP_BY_NAME(AnimNextAnimGraphBuilder)
+
 uint64 FAnimNextAnimGraphBuilder::GetKey() const
 {
 	if (CachedKey == 0)
diff --git a/Engine/Plugins/Experimental/UAF/UAFAnimGraph/Source/UAFAnimGraph/Private/Factory/AnimNextFactoryParams.cpp b/Engine/Plugins/Experimental/UAF/UAFAnimGraph/Source/UAFAnimGraph/Private/Factory/AnimNextFactoryParams.cpp
--- a/Engine/Plugins/Experimental/UAF/UAFAnimGraph/Source/UAFAnimGraph/Private/Factory/AnimNextFactoryParams.cpp
+++ b/Engine/Plugins/Experimental/UAF/UAFAnimGraph/Source/UAFAnimGraph/Private/Factory/AnimNextFactoryParams.cpp
@@ -1,8 +1,10 @@
-﻿// Copyright Epic Games, Inc. All Rights Reserved.
+// Copyright Epic Games, Inc. All Rights Reserved.
 
 #include "Factory/AnimNextFactoryParams.h"
 #include "Graph/AnimNextGraphInstance.h"
 
+#include UE_INLINE_GENERATED_CPP_BY_NAME(AnimNextFactoryParams)
+
 void FAnimNextFactoryParams::InitializeInstance(FUAFAssetInstance& InInstance) const
 {
 	if (Builder.Stacks.Num() > 0)
@@ -28,4 +30,4 @@ void FAnimNextFactoryParams::Reset()
 {
 	Builder.Reset();
 	InitializeTasks.Empty();
-}
\ No newline at end of file
+}
diff --git a/Engine/Plugins/Experimental/UAF/UAFAnimGraph/Source/UAFAnimGraph/Private/Factory/AnimNextSimpleAnimGraphBuilder.cpp b/Engine/Plugins/Experimental/UAF/UAFAnimGraph/Source/UAFAnimGraph/Private/Factory/AnimNextSimpleAnimGraphBuilder.cpp
--- a/Engine/Plugins/Experimental/UAF/UAFAnimGraph/Source/UAFAnimGraph/Private/Factory/AnimNextSimpleAnimGraphBuilder.cpp
+++ b/Engine/Plugins/Experimental/UAF/UAFAnimGraph/Source/UAFAnimGraph/Private/Factory/AnimNextSimpleAnimGraphBuilder.cpp
@@ -1,4 +1,4 @@
-﻿// Copyright Epic Games, Inc. All Rights Reserved.
+// Copyright Epic Games, Inc. All Rights Reserved.
 
 #include "Factory/AnimNextSimpleAnimGraphBuilder.h"
 
@@ -11,6 +11,8 @@
 #include "TraitCore/TraitWriter.h"
 #include "Traits/ReferencePoseTrait.h"
 
+#include UE_INLINE_GENERATED_CPP_BY_NAME(AnimNextSimpleAnimGraphBuilder)
+
 bool FAnimNextSimpleAnimGraphBuilder::Build(UE::UAF::FAnimGraphBuilderContext& InContext) const
 {
 	using namespace UE::UAF;
