UAF: Rename user facing strings for AnimNext Module -> UAF System

diff --git a/Engine/Plugins/Experimental/UAF/UAF/Source/UAF/Internal/Module/AnimNextModule.h b/Engine/Plugins/Experimental/UAF/UAF/Source/UAF/Internal/Module/AnimNextModule.h
--- a/Engine/Plugins/Experimental/UAF/UAF/Source/UAF/Internal/Module/AnimNextModule.h
+++ b/Engine/Plugins/Experimental/UAF/UAF/Source/UAF/Internal/Module/AnimNextModule.h
@@ -45,7 +45,7 @@ namespace UE::UAF::Editor
 }
 
 // Root asset represented by a component when instantiated
-UCLASS(MinimalAPI, BlueprintType)
+UCLASS(MinimalAPI, BlueprintType, DisplayName="UAF System")
 class UAnimNextModule : public UAnimNextSharedVariables
 {
 	GENERATED_BODY()
diff --git a/Engine/Plugins/Experimental/UAF/UAF/Source/UAF/Public/Component/AnimNextComponent.h b/Engine/Plugins/Experimental/UAF/UAF/Source/UAF/Public/Component/AnimNextComponent.h
--- a/Engine/Plugins/Experimental/UAF/UAF/Source/UAF/Public/Component/AnimNextComponent.h
+++ b/Engine/Plugins/Experimental/UAF/UAF/Source/UAF/Public/Component/AnimNextComponent.h
@@ -201,8 +201,8 @@ private:
 	friend class UAnimNextComponentWorldSubsystem;
 	friend class UE::UAF::Editor::FAssetWizard;
 
-	// The AnimNext module that this component will run
-	UPROPERTY(EditAnywhere, Category="Module")
+	// The UAF system that this component will run
+	UPROPERTY(EditAnywhere, Category="System", DisplayName="System")
 	TObjectPtr<UAnimNextModule> Module = nullptr;
 
 	UPROPERTY(Transient)
@@ -211,11 +211,11 @@ private:
 	// Handle to the registered module
 	UE::UAF::FModuleHandle ModuleHandle;
 
-	// How to initialize the module
-	UPROPERTY(EditAnywhere, Category="Module")
+	// How to initialize the system
+	UPROPERTY(EditAnywhere, Category="System")
 	EAnimNextModuleInitMethod InitMethod = EAnimNextModuleInitMethod::InitializeAndPauseInEditor;
 
-	/** When checked, the module's debug drawing instructions are drawn in the viewport */
+	/** When checked, the system's debug drawing instructions are drawn in the viewport */
 	UPROPERTY(EditAnywhere, Category = "Rendering")
 	uint8 bShowDebugDrawing : 1 = false;
 };
diff --git a/Engine/Plugins/Experimental/UAF/UAF/Source/UAFEditor/Private/Common/AnimNextGraphItemDetails.cpp b/Engine/Plugins/Experimental/UAF/UAF/Source/UAFEditor/Private/Common/AnimNextGraphItemDetails.cpp
--- a/Engine/Plugins/Experimental/UAF/UAF/Source/UAFEditor/Private/Common/AnimNextGraphItemDetails.cpp
+++ b/Engine/Plugins/Experimental/UAF/UAF/Source/UAFEditor/Private/Common/AnimNextGraphItemDetails.cpp
@@ -135,7 +135,7 @@ bool FAnimNextGraphItemDetails::ValidateName(const FWorkspaceOutlinerItemExport&
 		UAnimNextRigVMAssetEntry* ExistingEntry = EditorData->FindEntry(FName(*InName.ToString()));
 		if(ExistingEntry && GraphData.GetEntry() != ExistingEntry)
 		{
-			OutErrorMessage = LOCTEXT("NameAlreadyExistsError", "Name already exists in this module");
+			OutErrorMessage = LOCTEXT("NameAlreadyExistsError", "Name already exists in this system");
 			return false;
 		}
 
diff --git a/Engine/Plugins/Experimental/UAF/UAF/Source/UAFEditor/Private/Module/AnimNextModuleAssetDefinition.h b/Engine/Plugins/Experimental/UAF/UAF/Source/UAFEditor/Private/Module/AnimNextModuleAssetDefinition.h
--- a/Engine/Plugins/Experimental/UAF/UAF/Source/UAFEditor/Private/Module/AnimNextModuleAssetDefinition.h
+++ b/Engine/Plugins/Experimental/UAF/UAF/Source/UAFEditor/Private/Module/AnimNextModuleAssetDefinition.h
@@ -15,7 +15,7 @@ class UAssetDefinition_AnimNextModule : public UAssetDefinitionDefault
 
 public:
 	// UAssetDefinition interface
-	virtual FText GetAssetDisplayName() const override { return LOCTEXT("UAFModule", "UAF Module"); }
+	virtual FText GetAssetDisplayName() const override { return LOCTEXT("UAFSystem", "UAF System"); }
 	virtual FLinearColor GetAssetColor() const override { return FLinearColor(FColor(128,96,48)); }
 	virtual TSoftClassPtr<UObject> GetAssetClass() const override { return UAnimNextModule::StaticClass(); }
 	virtual EAssetCommandResult OpenAssets(const FAssetOpenArgs& OpenArgs) const override;
diff --git a/Engine/Plugins/Experimental/UAF/UAF/Source/UAFEditor/Private/RewindDebugger/AnimNextModuleTrack.cpp b/Engine/Plugins/Experimental/UAF/UAF/Source/UAFEditor/Private/RewindDebugger/AnimNextModuleTrack.cpp
--- a/Engine/Plugins/Experimental/UAF/UAF/Source/UAFEditor/Private/RewindDebugger/AnimNextModuleTrack.cpp
+++ b/Engine/Plugins/Experimental/UAF/UAF/Source/UAFEditor/Private/RewindDebugger/AnimNextModuleTrack.cpp
@@ -49,7 +49,7 @@ FText FAnimNextModuleTrack::GetDisplayNameInternal() const
  		}
 	
 		// Don't cache this since it is a placeholder name, will be replaced once the real one is sent
- 		return NSLOCTEXT("RewindDebugger", "AnimNextModuleTrackName", "Module");
+ 		return NSLOCTEXT("RewindDebugger", "AnimNextSystemTrackName", "System");
  	}
 
  	return DisplayNameCache;
@@ -57,7 +57,7 @@ FText FAnimNextModuleTrack::GetDisplayNameInternal() const
 
 void FAnimNextModuleTrackCreator::GetTrackTypesInternal(TArray<RewindDebugger::FRewindDebuggerTrackType>& Types) const
 {
-	Types.Add({AnimNextModulesName, LOCTEXT("AnimNextModule", "AnimNextModules")});
+	Types.Add({AnimNextModulesName, LOCTEXT("UAFSystems", "UAF Systems")});
 }
 
 TSharedPtr<RewindDebugger::FRewindDebuggerTrack> FAnimNextModuleTrackCreator::CreateTrackInternal(const RewindDebugger::FObjectId& ObjectId) const
diff --git a/Engine/Plugins/Experimental/UAF/UAF/Source/UAFUncookedOnly/Private/UncookedOnlyUtils.cpp b/Engine/Plugins/Experimental/UAF/UAF/Source/UAFUncookedOnly/Private/UncookedOnlyUtils.cpp
--- a/Engine/Plugins/Experimental/UAF/UAF/Source/UAFUncookedOnly/Private/UncookedOnlyUtils.cpp
+++ b/Engine/Plugins/Experimental/UAF/UAF/Source/UAFUncookedOnly/Private/UncookedOnlyUtils.cpp
@@ -373,19 +373,19 @@ void FUtils::CompileVariableBindingsInternal(const FRigVMCompileSettings& InSett
 		URigVMNode* CopyProxyVariablesNode = Controller->AddUnitNode(FRigUnit_CopyModuleProxyVariables::StaticStruct(), FRigVMStruct::ExecuteName, FVector2D(200, 0.0f), FString(), false);
 		if(CopyProxyVariablesNode == nullptr)
 		{
-			InSettings.ReportError(TEXT("Could not spawn Copy Module Proxy Variables node"));
+			InSettings.ReportError(TEXT("Could not spawn Copy System Proxy Variables node"));
 			return;
 		}
 		URigVMPin* CopyProxyVariablesExecPin = CopyProxyVariablesNode->FindPin(FRigVMStruct::ExecuteContextName.ToString());
 		if(ExecPin == nullptr)
 		{
-			InSettings.ReportError(TEXT("Could not find execute pin on Copy Module Proxy Variables node"));
+			InSettings.ReportError(TEXT("Could not find execute pin on Copy System Proxy Variables node"));
 			return;
 		}
 		bool bLinkAdded = Controller->AddLink(ExecuteBindingsExecPin, CopyProxyVariablesExecPin, false);
 		if(!bLinkAdded)
 		{
-			InSettings.ReportError(TEXT("Could not link Copy Module Proxy Variables node"));
+			InSettings.ReportError(TEXT("Could not link Copy System Proxy Variables node"));
 			return;
 		}
 		ExecPin = CopyProxyVariablesExecPin;
diff --git a/Engine/Plugins/Experimental/UAF/UAFAnimGraph/Source/UAFAnimGraph/Private/AnimNextAnimGraphModule.cpp b/Engine/Plugins/Experimental/UAF/UAFAnimGraph/Source/UAFAnimGraph/Private/AnimNextAnimGraphModule.cpp
--- a/Engine/Plugins/Experimental/UAF/UAFAnimGraph/Source/UAFAnimGraph/Private/AnimNextAnimGraphModule.cpp
+++ b/Engine/Plugins/Experimental/UAF/UAFAnimGraph/Source/UAFAnimGraph/Private/AnimNextAnimGraphModule.cpp
@@ -142,8 +142,8 @@ void FAnimNextAnimGraphModule::StartupModule()
 			ECVF_Default
 		));
 		ConsoleCommands.Add(IConsoleManager::Get().RegisterConsoleCommand(
-			TEXT("UAF.Module"),
-			TEXT("Dumps statistics about modules to the log."),
+			TEXT("UAF.Systems"),
+			TEXT("Dumps statistics about systems to the log."),
 			FConsoleCommandWithArgsDelegate::CreateRaw(this, &FAnimNextAnimGraphModule::ListAnimationGraphs),
 			ECVF_Default
 		));
diff --git a/Engine/Plugins/Experimental/UAF/UAFAnimGraph/Source/UAFAnimGraphEditor/Private/RewindDebugger/EvaluationProgramTrack.cpp b/Engine/Plugins/Experimental/UAF/UAFAnimGraph/Source/UAFAnimGraphEditor/Private/RewindDebugger/EvaluationProgramTrack.cpp
--- a/Engine/Plugins/Experimental/UAF/UAFAnimGraph/Source/UAFAnimGraphEditor/Private/RewindDebugger/EvaluationProgramTrack.cpp
+++ b/Engine/Plugins/Experimental/UAF/UAFAnimGraph/Source/UAFAnimGraphEditor/Private/RewindDebugger/EvaluationProgramTrack.cpp
@@ -36,7 +36,7 @@ FText FEvaluationProgramTrack::GetDisplayNameInternal() const
 
 void FEvaluationProgramTrackCreator::GetTrackTypesInternal(TArray<RewindDebugger::FRewindDebuggerTrackType>& Types) const
 {
-	Types.Add({AnimNextModulesName, LOCTEXT("AnimNextModule", "AnimNextModules")});
+	Types.Add({AnimNextModulesName, LOCTEXT("UAFSystems", "UAF Systems")});
 }
 
 TSharedPtr<RewindDebugger::FRewindDebuggerTrack> FEvaluationProgramTrackCreator::CreateTrackInternal(const RewindDebugger::FObjectId& InObjectId) const
