Add icons for set collection/binding and label collection/binding

diff --git a/Engine/Plugins/Experimental/UAF/UAF/Source/UAFEditor/Private/UAFStyle.cpp b/Engine/Plugins/Experimental/UAF/UAF/Source/UAFEditor/Private/UAFStyle.cpp
--- a/Engine/Plugins/Experimental/UAF/UAF/Source/UAFEditor/Private/UAFStyle.cpp
+++ b/Engine/Plugins/Experimental/UAF/UAF/Source/UAFEditor/Private/UAFStyle.cpp
@@ -21,8 +21,6 @@ FUAFStyle::FUAFStyle()
 	SetContentRoot(FPaths::EngineContentDir() / TEXT("Editor/Slate"));
 	SetCoreContentRoot(FPaths::EngineContentDir() / TEXT("Slate"));
 
-	const FString AnimNextPluginContentDir = FPaths::EnginePluginsDir() / TEXT("Experimental/AnimNext/Resources");
-
 	SetClassIcons();
 	SetVariablesOutlinerStyles();
 	SetTraitStyles();
@@ -39,6 +37,10 @@ void FUAFStyle::SetClassIcons()
 	Set("ClassIcon.AnimNextModule", new IMAGE_BRUSH_SVG("Starship/Common/Modules", Icon16x16 ));
 	Set("ClassIcon.AnimNextAnimationGraph", new IMAGE_BRUSH_SVG("Starship/Blueprints/icon_BlueprintEditor_EventGraph", Icon16x16));
 	Set("ClassIcon.AnimNextStateTree", new IMAGE_BRUSH_SVG(TEXT("UMG/TreeView"), Icon16x16));
+	Set("ClassIcon.AbstractSkeletonLabelCollection", new IMAGE_PLUGIN_BRUSH_SVG("Icons/LabelCollection_16.svg", Icon16x16));
+	Set("ClassIcon.AbstractSkeletonLabelBinding", new IMAGE_PLUGIN_BRUSH_SVG("Icons/LabelBinding_16.svg", Icon16x16));
+	Set("ClassIcon.AbstractSkeletonSetCollection", new IMAGE_PLUGIN_BRUSH_SVG("Icons/SetCollection_16.svg", Icon16x16));
+	Set("ClassIcon.AbstractSkeletonSetBinding", new IMAGE_PLUGIN_BRUSH_SVG("Icons/SetBinding_16.svg", Icon16x16));
 }
 
 void FUAFStyle::SetVariablesOutlinerStyles()
