UAF Function fixes.
- Fix crash immediate undo AddFunction UAF ST graph due to Model not generated yet on temporary
- Fix Function Selector not working on CallFunction after first reload due to not handling invalid arg types
- Fix RigVM StateTree Function Task not executing due to invalid return index (We don't set / use index for that at the moment).

Note: CallFunction does still not correctly assign after first created. Asset must be reloaded. Will make task for this but not a blocker so not addressing as priority now.

#tests PIE

diff --git a/Engine/Plugins/Experimental/UAF/UAF/Source/UAFUncookedOnly/Private/AnimNextRigVMAssetEditorData.cpp b/Engine/Plugins/Experimental/UAF/UAF/Source/UAFUncookedOnly/Private/AnimNextRigVMAssetEditorData.cpp
--- a/Engine/Plugins/Experimental/UAF/UAF/Source/UAFUncookedOnly/Private/AnimNextRigVMAssetEditorData.cpp
+++ b/Engine/Plugins/Experimental/UAF/UAF/Source/UAFUncookedOnly/Private/AnimNextRigVMAssetEditorData.cpp
@@ -810,9 +810,12 @@ void UAnimNextRigVMAssetEditorData::BuildFunctionWrapperEventVariables(FAnimNext
 			// Don't create internal variables for execution IO, visible only, or hidden pins. Those types do not need variable nodes in the graph.
 			if (Argument.Direction == ERigVMPinDirection::Input || Argument.Direction == ERigVMPinDirection::Output)
 			{
-				FRigVMGraphFunctionArgument InternallyNamedArgument = Argument;
-				InternallyNamedArgument.Name = FName(FUtils::MakeFunctionWrapperVariableName(FunctionHeader.Name, Argument.Name));
-				InContext.ProgrammaticVariables.Add(FAnimNextProgrammaticVariable::FromRigVMGraphFunctionArgument(InternallyNamedArgument));
+				if (Argument.IsCPPTypeObjectValid())
+				{
+					FRigVMGraphFunctionArgument InternallyNamedArgument = Argument;
+					InternallyNamedArgument.Name = FName(FUtils::MakeFunctionWrapperVariableName(FunctionHeader.Name, Argument.Name));
+					InContext.ProgrammaticVariables.Add(FAnimNextProgrammaticVariable::FromRigVMGraphFunctionArgument(InternallyNamedArgument));
+				}
 			}
 		}
 	}
@@ -891,7 +894,7 @@ void UAnimNextRigVMAssetEditorData::BuildFunctionWrapperEvents(FAnimNextRigVMAss
 		for (const FRigVMGraphFunctionArgument& Argument : FunctionHeader.Arguments)
 		{
 			// Execution context is captured as arg pins, skip those for internal variable gen
-			if (Argument.Direction == ERigVMPinDirection::IO)
+			if (Argument.Direction == ERigVMPinDirection::IO || !Argument.IsCPPTypeObjectValid())
 			{
 				continue;
 			}
diff --git a/Engine/Plugins/Experimental/UAF/UAFStateTree/Source/UAFStateTree/Private/Tasks/AnimNextStateTreeRigVMTaskBase.cpp b/Engine/Plugins/Experimental/UAF/UAFStateTree/Source/UAFStateTree/Private/Tasks/AnimNextStateTreeRigVMTaskBase.cpp
--- a/Engine/Plugins/Experimental/UAF/UAFStateTree/Source/UAFStateTree/Private/Tasks/AnimNextStateTreeRigVMTaskBase.cpp
+++ b/Engine/Plugins/Experimental/UAF/UAFStateTree/Source/UAFStateTree/Private/Tasks/AnimNextStateTreeRigVMTaskBase.cpp
@@ -54,7 +54,8 @@ EStateTreeRunStatus FAnimNextStateTreeRigVMTaskBase::EnterState(FStateTreeExecut
 				InstanceData->FunctionHandle = RigVMAsset->GetFunctionHandle(InternalEventName);
 			}
 
-			if (!ensure(InstanceData->FunctionHandle.IsValid() && ResultIndex != INDEX_NONE))
+			// See below, for now ignore index as we fire and forget on tasks
+			if (!ensure(InstanceData->FunctionHandle.IsValid() /*&& ResultIndex != INDEX_NONE*/))
 			{
 				return EStateTreeRunStatus::Failed;
 			}
diff --git a/Engine/Plugins/Runtime/RigVM/Source/RigVMEditor/Private/EdGraph/RigVMEdGraphConnectionDrawingPolicy.cpp b/Engine/Plugins/Runtime/RigVM/Source/RigVMEditor/Private/EdGraph/RigVMEdGraphConnectionDrawingPolicy.cpp
--- a/Engine/Plugins/Runtime/RigVM/Source/RigVMEditor/Private/EdGraph/RigVMEdGraphConnectionDrawingPolicy.cpp
+++ b/Engine/Plugins/Runtime/RigVM/Source/RigVMEditor/Private/EdGraph/RigVMEdGraphConnectionDrawingPolicy.cpp
@@ -36,7 +36,10 @@ void FRigVMEdGraphConnectionDrawingPolicy::ResetIncompatiblePinDrawState(const T
 		{
 			if(URigVMEdGraphNode* RigNode = Cast<URigVMEdGraphNode>(Pin->GetOwningNode()))
 			{
-				RigNode->GetModel()->PrepareCycleChecking(nullptr, true);
+				if (URigVMGraph* ModelGraph = RigNode->GetModel())
+				{
+					ModelGraph->PrepareCycleChecking(nullptr, true);
+				}
 			}
 		}
 	}
