Select and try to rename newly added variables in outliner

diff --git a/Engine/Plugins/Experimental/UAF/UAF/Source/UAFEditor/Private/Variables/Outliner/VariablesOutlinerMode.cpp b/Engine/Plugins/Experimental/UAF/UAF/Source/UAFEditor/Private/Variables/Outliner/VariablesOutlinerMode.cpp
--- a/Engine/Plugins/Experimental/UAF/UAF/Source/UAFEditor/Private/Variables/Outliner/VariablesOutlinerMode.cpp
+++ b/Engine/Plugins/Experimental/UAF/UAF/Source/UAFEditor/Private/Variables/Outliner/VariablesOutlinerMode.cpp
@@ -656,7 +656,7 @@ void FVariablesOutlinerMode::PopulateNewVariableToolMenuEntries(UToolMenu* InMen
 	}
 	else
 	{
-		TSharedPtr<UE::UAF::Editor::SVariablesOutliner> SharedOutliner = MenuContext->WeakOutliner.Pin();
+		TSharedPtr<SVariablesOutliner> SharedOutliner = MenuContext->WeakOutliner.Pin();
 		if (SharedOutliner.IsValid())
 		{
 			TArray<FSceneOutlinerTreeItemPtr> SelectedItems = SharedOutliner->GetSelectedItems();
@@ -706,7 +706,7 @@ void FVariablesOutlinerMode::PopulateNewVariableToolMenuEntries(UToolMenu* InMen
 
 		FToolMenuSection& Section = InMenu->AddSection("Variables", LOCTEXT("VariablesSectionLabel", "Variables"));
 		Section.AddMenuEntry("AddNewVariable", LOCTEXT("VariableLabel", "Add Variable"), LOCTEXT("VariableTooltip", "Adds a new Variable"),FSlateIcon(FAppStyle::GetAppStyleSetName(), "Icons.PlusCircle"),
-			FUIAction(FExecuteAction::CreateLambda([AssetsAndCategoriesToAddTo]()
+			FUIAction(FExecuteAction::CreateLambda([WeakOutliner = MenuContext->WeakOutliner, AssetsAndCategoriesToAddTo]()
 			{
 				for (auto& Pair : AssetsAndCategoriesToAddTo)
 				{
@@ -717,11 +717,21 @@ void FVariablesOutlinerMode::PopulateNewVariableToolMenuEntries(UToolMenu* InMen
                         
                         const FName VariableName = UncookedOnly::FUtils::GetValidVariableName(EditorData, Settings->GetLastVariableName());		
 
-						if (UAnimNextVariableEntry* VariableEntry = Pair.Key->AddVariable(VariableName, Settings->GetLastVariableType()))
+						if (UAnimNextVariableEntry* VariableEntry = EditorData->AddVariable(VariableName, Settings->GetLastVariableType()))
 						{
 							VariableEntry->SetVariableCategory(Pair.Value);
+							Settings->SetLastVariableName(VariableName);
+
+							// Prompt user to rename added variable in outliner
+							TSharedPtr<SVariablesOutliner> SharedOutliner = WeakOutliner.Pin();
+							if (SharedOutliner.IsValid())
+							{
+								// FVariablesOutlinerEntryItem::SharedVariablesEntry is only populated for Struct based SharedVariables
+								const UAnimNextSharedVariablesEntry* SharedVariablesEntry = nullptr;
+								FSceneOutlinerTreeItemID EntryID = HashCombine(GetTypeHash(VariableEntry), GetTypeHash(SharedVariablesEntry));
+								SharedOutliner->OnItemAdded(EntryID, SceneOutliner::ENewItemAction::Select | SceneOutliner::ENewItemAction::Rename);
+							}
 						}
-                        Settings->SetLastVariableName(VariableName);
                     }
 				}
 			}))
