Add support for nested categories for UAF variables (view)

diff --git a/Engine/Plugins/Experimental/UAF/UAF/Source/UAFEditor/Private/Variables/Outliner/VariablesOutlinerCategoryItem.cpp b/Engine/Plugins/Experimental/UAF/UAF/Source/UAFEditor/Private/Variables/Outliner/VariablesOutlinerCategoryItem.cpp
--- a/Engine/Plugins/Experimental/UAF/UAF/Source/UAFEditor/Private/Variables/Outliner/VariablesOutlinerCategoryItem.cpp
+++ b/Engine/Plugins/Experimental/UAF/UAF/Source/UAFEditor/Private/Variables/Outliner/VariablesOutlinerCategoryItem.cpp
@@ -118,9 +118,9 @@ class SVariablesOutlinerCategoryLabel : FSceneOutlinerCommonLabelData, public SC
 						check(NewVariableEntry);
 						
 						// Only add if not already part of category
-						if (NewVariableEntry->GetVariableCategory() != CategoryItem->CategoryName)
+						if (NewVariableEntry->GetVariableCategory() != CategoryItem->CategoryPath)
 						{
-							NewVariableEntry->SetVariableCategory(CategoryItem->CategoryName);							
+							NewVariableEntry->SetVariableCategory(CategoryItem->CategoryPath);
 						}
 					}
 				}
@@ -251,7 +251,7 @@ class SVariablesOutlinerCategoryLabel : FSceneOutlinerCommonLabelData, public SC
 	TSharedPtr<SInlineEditableTextBlock> TextBlock;
 };
 
-FVariablesOutlinerCategoryItem::FVariablesOutlinerCategoryItem(const FItemData& InItemData) : ISceneOutlinerTreeItem(FVariablesOutlinerCategoryItem::Type), WeakOwner(InItemData.InOwner), CategoryName(InItemData.InCategoryName)
+FVariablesOutlinerCategoryItem::FVariablesOutlinerCategoryItem(const FItemData& InItemData) : ISceneOutlinerTreeItem(FVariablesOutlinerCategoryItem::Type), WeakOwner(InItemData.InOwner), CategoryName(InItemData.InCategoryName), ParentCategoryName(InItemData.InParentCategoryName), CategoryPath(InItemData.InCategoryPath)
 {
 }
 
@@ -264,7 +264,7 @@ FSceneOutlinerTreeItemID FVariablesOutlinerCategoryItem::GetID() const
 {
 	UAnimNextRigVMAsset* Owner = WeakOwner.Get();
 	const FSoftObjectPath OwnerPath = Owner;
-	return HashCombine(GetTypeHash(OwnerPath), GetTypeHash(CategoryName));	
+	return HashCombine(GetTypeHash(OwnerPath), GetTypeHash(CategoryPath));	
 }
 
 FString FVariablesOutlinerCategoryItem::GetDisplayString() const
diff --git a/Engine/Plugins/Experimental/UAF/UAF/Source/UAFEditor/Private/Variables/Outliner/VariablesOutlinerCategoryItem.h b/Engine/Plugins/Experimental/UAF/UAF/Source/UAFEditor/Private/Variables/Outliner/VariablesOutlinerCategoryItem.h
--- a/Engine/Plugins/Experimental/UAF/UAF/Source/UAFEditor/Private/Variables/Outliner/VariablesOutlinerCategoryItem.h
+++ b/Engine/Plugins/Experimental/UAF/UAF/Source/UAFEditor/Private/Variables/Outliner/VariablesOutlinerCategoryItem.h
@@ -17,6 +17,8 @@ struct FVariablesOutlinerCategoryItem : ISceneOutlinerTreeItem
 	{
 		UAnimNextRigVMAsset* InOwner;
 		FStringView InCategoryName;
+		FStringView InParentCategoryName;
+		FStringView InCategoryPath;
 	};
 
 	FVariablesOutlinerCategoryItem(const FItemData& InItemData);
@@ -38,6 +40,8 @@ struct FVariablesOutlinerCategoryItem : ISceneOutlinerTreeItem
 	
 	TWeakObjectPtr<UAnimNextRigVMAsset> WeakOwner;
 	FString CategoryName;
+	FString ParentCategoryName;
+	FString CategoryPath;
 	int32 SortValue = 0;
 };
 
diff --git a/Engine/Plugins/Experimental/UAF/UAF/Source/UAFEditor/Private/Variables/Outliner/VariablesOutlinerHierarchy.cpp b/Engine/Plugins/Experimental/UAF/UAF/Source/UAFEditor/Private/Variables/Outliner/VariablesOutlinerHierarchy.cpp
--- a/Engine/Plugins/Experimental/UAF/UAF/Source/UAFEditor/Private/Variables/Outliner/VariablesOutlinerHierarchy.cpp
+++ b/Engine/Plugins/Experimental/UAF/UAF/Source/UAFEditor/Private/Variables/Outliner/VariablesOutlinerHierarchy.cpp
@@ -96,7 +96,9 @@ FSceneOutlinerTreeItemPtr FVariablesOutlinerHierarchy::FindOrCreateParentItem(co
 		if (UAnimNextRigVMAsset* Asset = CategoryItem->WeakOwner.Get())
 		{
 			const FSoftObjectPath OwnerPath = Asset;
-			ParentHash = GetTypeHash(OwnerPath);
+			const uint32 AssetHash = GetTypeHash(OwnerPath);
+			const FStringView ParentCategory = CategoryItem->ParentCategoryName;
+			ParentHash = ParentCategory.Len() ? HashCombine(AssetHash, GetTypeHash(ParentCategory)) : AssetHash;
 		}
 	}
 
@@ -181,17 +183,37 @@ void FVariablesOutlinerHierarchy::ProcessAsset(const TSoftObjectPtr<UAnimNextRig
 
 	OutItems.Append(ParentItems);
 
-	TMap<FString, FSceneOutlinerTreeItemPtr> CategoryItems;
+	TMap<FString, TSharedPtr<FVariablesOutlinerCategoryItem>> CategoryItems;
 
 	for (int32 CategoryIndex = 0; CategoryIndex < EditorData->VariableAndFunctionCategories.Num(); ++CategoryIndex)
 	{
 		const FString& CategoryName = EditorData->VariableAndFunctionCategories[CategoryIndex];
-		if (FSceneOutlinerTreeItemPtr Item = Mode->CreateItemFor<FVariablesOutlinerCategoryItem>(FVariablesOutlinerCategoryItem::FItemData{Asset, CategoryName}))
+		TArray<FString> SubCategories;
+		CategoryName.ParseIntoArray(SubCategories, TEXT("|"));
+
+		FString RebuildCategories;
+		for (const FString& SubCategoryName : SubCategories)
 		{
-			Item->CastTo<FVariablesOutlinerCategoryItem>()->SortValue = CategoryIndex;
-			CategoryItems.Add(CategoryName, Item);
+			const FString ParentCategoryName = RebuildCategories;
+
+			if (!RebuildCategories.IsEmpty())
+			{
+				RebuildCategories.Append(TEXT("|"));
+			}
+			RebuildCategories.Append(SubCategoryName);
+			
+			const FString FullCategoryName = RebuildCategories;
+
+			if (!CategoryItems.Contains(FullCategoryName))
+			{
+				if (TSharedPtr<FVariablesOutlinerCategoryItem> Item = StaticCastSharedPtr<FVariablesOutlinerCategoryItem>( Mode->CreateItemFor<FVariablesOutlinerCategoryItem>(FVariablesOutlinerCategoryItem::FItemData{Asset, SubCategoryName, ParentCategoryName, FullCategoryName})))
+				{
+					Item->CastTo<FVariablesOutlinerCategoryItem>()->SortValue = CategoryIndex;
+					CategoryItems.Add(FullCategoryName, Item);
 
-			UE_CLOG(GOutputHierarchyDebugInformationCvar.GetValueOnAnyThread(),LogAnimation, Warning, TEXT("+ FVariablesOutlinerCategoryItem: %s %i"), *CategoryName, CategoryIndex);
+					UE_CLOG(GOutputHierarchyDebugInformationCvar.GetValueOnAnyThread(),LogAnimation, Warning, TEXT("+ FVariablesOutlinerCategoryItem: %s  %i"), *FullCategoryName, CategoryIndex);
+				}
+			}
 		}
 	}
 
@@ -200,6 +222,19 @@ void FVariablesOutlinerHierarchy::ProcessAsset(const TSoftObjectPtr<UAnimNextRig
 
 	constexpr int32 NumberOfVariableSortValueEntriesInCategory = 1024;
 
+	auto EnsureCategoryItemExists = [&CategoryItems, &OutItems](const FString& CategoryName)
+	{		 
+		TSharedPtr<FVariablesOutlinerCategoryItem>& CategoryItemToAdd = CategoryItems.FindChecked(CategoryName);
+		OutItems.AddUnique(CategoryItemToAdd);
+
+		auto CategoryPtr = CategoryItems.Find(CategoryItemToAdd->ParentCategoryName);
+		while (CategoryPtr)
+		{
+			OutItems.AddUnique(*CategoryPtr);
+			CategoryPtr = CategoryItems.Find((*CategoryPtr)->ParentCategoryName);
+		}
+	};
+
 	if (FSceneOutlinerTreeItemPtr Item = Mode->CreateItemFor<FVariablesOutlinerAssetItem>(FVariablesOutlinerAssetItem::FItemData{InSoftAsset, AssetDepth, OutSharedVariableSourcePaths}))
 	{
 		OutItems.Add(Item);
@@ -207,7 +242,7 @@ void FVariablesOutlinerHierarchy::ProcessAsset(const TSoftObjectPtr<UAnimNextRig
 		UE_CLOG(GOutputHierarchyDebugInformationCvar.GetValueOnAnyThread(), LogAnimation, Warning, TEXT("* FVariablesOutlinerAssetItem: %s %i"), *InSoftAsset.GetAssetName(), AssetDepth);
 
 		int32 Index = 1;
-		EditorData->ForEachEntryOfType<UAnimNextVariableEntry>([this, &OutItems, &CategoryItems, &Index, EditorData](UAnimNextVariableEntry* InVariable)
+		EditorData->ForEachEntryOfType<UAnimNextVariableEntry>([this, &OutItems, EnsureCategoryItemExists, &Index, EditorData](UAnimNextVariableEntry* InVariable)
 		{
 			if (FSceneOutlinerTreeItemPtr Item = Mode->CreateItemFor<FVariablesOutlinerEntryItem>(CastChecked<UAnimNextVariableEntry>(InVariable)))
 			{
@@ -217,13 +252,12 @@ void FVariablesOutlinerHierarchy::ProcessAsset(const TSoftObjectPtr<UAnimNextRig
 				OutItems.Add(Item);
 				
 				OutItems.Add(Item);
-				UE_CLOG(GOutputHierarchyDebugInformationCvar.GetValueOnAnyThread(), LogAnimation, Warning, TEXT("\t- FVariablesOutlinerEntryItem: %s %i"), *InVariable->GetEntryName().ToString(), Item->CastTo<FVariablesOutlinerEntryItem>()->SortValue);
+				UE_CLOG(GOutputHierarchyDebugInformationCvar.GetValueOnAnyThread(), LogAnimation, Warning, TEXT("\t- FVariablesOutlinerEntryItem: %s %i %s"), *InVariable->GetEntryName().ToString(), Item->CastTo<FVariablesOutlinerEntryItem>()->SortValue, InVariable->GetVariableCategory().GetData());
 
 				// Add category item on demand (to ensure we do not end up with empty category item hierarchies)
 				if (!InVariable->Category.IsEmpty())
 				{
-					FSceneOutlinerTreeItemPtr& CategoryItemToAdd = CategoryItems.FindChecked(InVariable->Category);
-					OutItems.AddUnique(CategoryItemToAdd);
+					EnsureCategoryItemExists(InVariable->Category);
 				}
 			}
 			return true;
diff --git a/Engine/Plugins/Experimental/UAF/UAF/Source/UAFEditor/Private/Variables/VariableCustomization.cpp b/Engine/Plugins/Experimental/UAF/UAF/Source/UAFEditor/Private/Variables/VariableCustomization.cpp
--- a/Engine/Plugins/Experimental/UAF/UAF/Source/UAFEditor/Private/Variables/VariableCustomization.cpp
+++ b/Engine/Plugins/Experimental/UAF/UAF/Source/UAFEditor/Private/Variables/VariableCustomization.cpp
@@ -20,6 +20,10 @@
 namespace UE::UAF::Editor
 {
 
+FVariableCustomization::FVariableCustomization() : DefaultCategoryName(LOCTEXT("DefaultCategoryLabel", "Default").ToString())
+{
+}
+
 void FVariableCustomization::CustomizeDetails(IDetailLayoutBuilder& DetailBuilder)
 {	
 	TArray<TWeakObjectPtr<UObject>> Objects;
@@ -233,6 +237,11 @@ bool FVariableCustomization::OnVerifyCategoryName(const FText& InText, FText& Ou
 		OutErrorMessage = LOCTEXT("CategoryEmpty", "Cannot add a category with an empty string.");
 		return false;
 	}
+	else if (InText.ToString() == DefaultCategoryName)
+	{
+		OutErrorMessage = LOCTEXT("CategoryDefault", "Cannot add a category with name 'Default'.");
+		return false;
+	}
 	return true;
 }
 
@@ -261,6 +270,8 @@ void FVariableCustomization::PopulateCategories()
 			return MakeShared<FText>(FText::FromString(InCategory));
 		});
 	}
+	
+	Categories.Add(MakeShared<FText>(FText::FromString(DefaultCategoryName)));
 }
 
 void FVariableCustomization::SetCategory(const FString& CategoryName)
@@ -269,7 +280,8 @@ void FVariableCustomization::SetCategory(const FString& CategoryName)
 	{
 		if(UAnimNextVariableEntry* PinnedVariable = Cast<UAnimNextVariableEntry>(WeakVariable.Get()))
 		{
-			PinnedVariable->SetVariableCategory(CategoryName);
+			const FString NewCategoryName = CategoryName == DefaultCategoryName ? TEXT("") : CategoryName;
+			PinnedVariable->SetVariableCategory(NewCategoryName);
 		}
 	}
 	
diff --git a/Engine/Plugins/Experimental/UAF/UAF/Source/UAFEditor/Private/Variables/VariableCustomization.h b/Engine/Plugins/Experimental/UAF/UAF/Source/UAFEditor/Private/Variables/VariableCustomization.h
--- a/Engine/Plugins/Experimental/UAF/UAF/Source/UAFEditor/Private/Variables/VariableCustomization.h
+++ b/Engine/Plugins/Experimental/UAF/UAF/Source/UAFEditor/Private/Variables/VariableCustomization.h
@@ -15,6 +15,9 @@ namespace UE::UAF::Editor
 
 class FVariableCustomization : public IDetailCustomization
 {
+public:	
+	FVariableCustomization();
+
 private:
 	/** Called when details should be customized */
 	virtual void CustomizeDetails(IDetailLayoutBuilder& DetailBuilder) override;
@@ -42,6 +45,8 @@ private:
 
 	TArray<TWeakObjectPtr<UAnimNextVariableEntry>> Variables;
 	TWeakObjectPtr<UAnimNextRigVMAsset> OuterRigVMAsset;
+
+	FString DefaultCategoryName;
 };
 
 }
