Add functionality for switching Variable type selector widget type (stored in settings so persistent across editor runs)

diff --git a/Engine/Plugins/Experimental/UAF/UAF/Source/UAFEditor/Private/Variables/AnimNextVariableSettings.cpp b/Engine/Plugins/Experimental/UAF/UAF/Source/UAFEditor/Private/Variables/AnimNextVariableSettings.cpp
--- a/Engine/Plugins/Experimental/UAF/UAF/Source/UAFEditor/Private/Variables/AnimNextVariableSettings.cpp
+++ b/Engine/Plugins/Experimental/UAF/UAF/Source/UAFEditor/Private/Variables/AnimNextVariableSettings.cpp
@@ -27,3 +27,14 @@ void UAnimNextVariableSettings::SetLastVariableName(FName InLastVariableName)
 {
 	LastVariableName = InLastVariableName;
 }
+
+SPinTypeSelector::ESelectorType UAnimNextVariableSettings::GetVariablesViewDefaultPinSelectorType() const
+{
+	return static_cast<SPinTypeSelector::ESelectorType>(VariablesViewPinSelectorType);
+}
+
+void UAnimNextVariableSettings::SetVariablesViewDefaultPinSelectorType(SPinTypeSelector::ESelectorType InLastSelectorType)
+{
+	VariablesViewPinSelectorType = static_cast<uint32>(InLastSelectorType);
+	SaveConfig();
+}
diff --git a/Engine/Plugins/Experimental/UAF/UAF/Source/UAFEditor/Private/Variables/AnimNextVariableSettings.h b/Engine/Plugins/Experimental/UAF/UAF/Source/UAFEditor/Private/Variables/AnimNextVariableSettings.h
--- a/Engine/Plugins/Experimental/UAF/UAF/Source/UAFEditor/Private/Variables/AnimNextVariableSettings.h
+++ b/Engine/Plugins/Experimental/UAF/UAF/Source/UAFEditor/Private/Variables/AnimNextVariableSettings.h
@@ -2,12 +2,13 @@
 
 #pragma once
 
-#include "UObject/Object.h"
+#include "SPinTypeSelector.h"
+#include "Engine/DeveloperSettings.h"
 #include "Param/ParamType.h"
 #include "AnimNextVariableSettings.generated.h"
 
 UCLASS(Config = EditorPerProjectUserSettings)
-class UAnimNextVariableSettings : public UObject
+class UAnimNextVariableSettings : public UDeveloperSettings
 {
 	GENERATED_BODY()
 
@@ -25,10 +26,16 @@ public:
 
 	// Set the name of the last variable that we created
 	void SetLastVariableName(FName InLastVariableName);
+	
+	SPinTypeSelector::ESelectorType GetVariablesViewDefaultPinSelectorType() const;
+	void SetVariablesViewDefaultPinSelectorType(SPinTypeSelector::ESelectorType InLastSelectorType);
 private:
 	UPROPERTY(Transient)
 	FAnimNextParamType LastVariableType = FAnimNextParamType::GetType<bool>();
 
 	UPROPERTY(Transient)
 	FName LastVariableName = FName("NewVariable");
+
+	UPROPERTY(VisibleAnywhere, Category="UAF", Config)
+	uint32 VariablesViewPinSelectorType = static_cast<uint32>(SPinTypeSelector::ESelectorType::Partial);
 };
\ No newline at end of file
diff --git a/Engine/Plugins/Experimental/UAF/UAF/Source/UAFEditor/Private/Variables/Outliner/VariablesOutlinerColumns.cpp b/Engine/Plugins/Experimental/UAF/UAF/Source/UAFEditor/Private/Variables/Outliner/VariablesOutlinerColumns.cpp
--- a/Engine/Plugins/Experimental/UAF/UAF/Source/UAFEditor/Private/Variables/Outliner/VariablesOutlinerColumns.cpp
+++ b/Engine/Plugins/Experimental/UAF/UAF/Source/UAFEditor/Private/Variables/Outliner/VariablesOutlinerColumns.cpp
@@ -17,6 +17,7 @@
 #include "Entries/AnimNextSharedVariablesEntry.h"
 #include "Entries/AnimNextRigVMAssetEntry.h"
 #include "Entries/AnimNextVariableEntry.h"
+#include "Framework/MultiBox/MultiBoxBuilder.h"
 #include "Widgets/Images/SImage.h"
 #include "Variables/AnimNextUniversalObjectLocatorBindingData.h"
 #include "Variables/AnimNextVariableSettings.h"
@@ -35,15 +36,21 @@ FName FVariablesOutlinerTypeColumn::GetID()
 	return VariablesOutlinerType;
 }
 
+FVariablesOutlinerTypeColumn::FVariablesOutlinerTypeColumn(ISceneOutliner& SceneOutliner) : WeakSceneOutliner(StaticCastSharedRef<ISceneOutliner>(SceneOutliner.AsShared()))
+{
+	SelectorType = GetDefault<UAnimNextVariableSettings>()->GetVariablesViewDefaultPinSelectorType();
+}
+
 SHeaderRow::FColumn::FArguments FVariablesOutlinerTypeColumn::ConstructHeaderRowColumn()
 {
 	return SHeaderRow::Column(GetColumnID())
-		.FixedWidth(24.f)
+		.ManualWidth(this, &FVariablesOutlinerTypeColumn::GetColumnWidth)
 		.HAlignHeader(HAlign_Center)
 		.VAlignHeader(VAlign_Center)
-		.HAlignCell(HAlign_Center)
+		.HAlignCell(HAlign_Left)
 		.VAlignCell(VAlign_Center)
 		.SortMode(EColumnSortMode::None)
+		.HeaderComboVisibility(EHeaderComboVisibility::Ghosted)
 		[
 			SNew(SBox)
 			.WidthOverride(16.0f)
@@ -56,7 +63,59 @@ SHeaderRow::FColumn::FArguments FVariablesOutlinerTypeColumn::ConstructHeaderRow
 				.Image(FAppStyle::GetBrush(TEXT("Kismet.VariableList.TypeIcon")))
 				.ToolTipText(LOCTEXT("TypeTooltip", "Type of this entry"))
 			]
-		];
+		]
+		.OnGetMenuContent_Lambda([this]()
+		{
+			const bool bCloseAfterSelection = true;
+			FMenuBuilder MenuBuilder(bCloseAfterSelection, nullptr);
+			MenuBuilder.BeginSection("", LOCTEXT("VariablesOutlinerTypeColumnContextMenu", "Variable Type Column") );
+			{
+				MenuBuilder.AddMenuEntry(
+						LOCTEXT("CompactView", "Compact"),
+						LOCTEXT("CompactView_ToolTip", "Displays the Variable Type using Compact representation."),
+						FSlateIcon(),
+						FUIAction(FExecuteAction::CreateLambda([this]()
+						{
+							SelectorType = SPinTypeSelector::ESelectorType::Compact;
+							GetMutableDefault<UAnimNextVariableSettings>()->SetVariablesViewDefaultPinSelectorType(SelectorType);
+							WeakSceneOutliner.Pin()->FullRefresh();
+						}), FCanExecuteAction(), FGetActionCheckState::CreateLambda([this]() {return SelectorType == SPinTypeSelector::ESelectorType::Compact ? ECheckBoxState::Checked : ECheckBoxState::Undetermined;})),
+						NAME_None,
+						EUserInterfaceActionType::Check
+				);
+
+				MenuBuilder.AddMenuEntry(
+						LOCTEXT("PartialView", "Partial"),
+						LOCTEXT("PartialView_ToolTip", "Displays the Variable Type using Partial representation."),
+						FSlateIcon(),
+						FUIAction(FExecuteAction::CreateLambda([this]()
+						{
+							SelectorType = SPinTypeSelector::ESelectorType::Partial;
+							GetMutableDefault<UAnimNextVariableSettings>()->SetVariablesViewDefaultPinSelectorType(SelectorType);
+							WeakSceneOutliner.Pin()->FullRefresh();					
+						}), FCanExecuteAction(), FGetActionCheckState::CreateLambda([this]() {return SelectorType == SPinTypeSelector::ESelectorType::Partial ? ECheckBoxState::Checked : ECheckBoxState::Undetermined;})),
+						NAME_None,
+						EUserInterfaceActionType::Check
+				);
+
+				MenuBuilder.AddMenuEntry(
+						LOCTEXT("FullView", "Full"),
+						LOCTEXT("FullView_ToolTip", "Displays the Variable Type using Full representation."),
+						FSlateIcon(),
+						FUIAction(FExecuteAction::CreateLambda([this]()
+						{
+							SelectorType = SPinTypeSelector::ESelectorType::Full;
+							GetMutableDefault<UAnimNextVariableSettings>()->SetVariablesViewDefaultPinSelectorType(SelectorType);
+							WeakSceneOutliner.Pin()->FullRefresh();
+						}), FCanExecuteAction(), FGetActionCheckState::CreateLambda([this]() {return SelectorType == SPinTypeSelector::ESelectorType::Full ? ECheckBoxState::Checked : ECheckBoxState::Undetermined;})),
+						NAME_None,
+						EUserInterfaceActionType::Check
+				);
+			}
+			MenuBuilder.EndSection();
+
+			return MenuBuilder.MakeWidget();
+		});
 }
 
 const TSharedRef<SWidget> FVariablesOutlinerTypeColumn::ConstructRowWidget(FSceneOutlinerTreeItemRef Item, const STableRow<FSceneOutlinerTreeItemPtr>& Row)
@@ -72,54 +131,70 @@ const TSharedRef<SWidget> FVariablesOutlinerTypeColumn::ConstructRowWidget(FScen
 		.HAlign(HAlign_Left)
 		.VAlign(VAlign_Center)
 		.IsEnabled(!TreeItem->WeakSharedVariablesEntry.IsValid() || !TreeItem->PropertyPath.IsPathToFieldEmpty())
-		[
-			SNew(SPinTypeSelector, FGetPinTypeTree::CreateStatic(&Editor::FUtils::GetFilteredVariableTypeTree))
-			.TargetPinType_Lambda([WeakItem = TWeakPtr<ISceneOutlinerTreeItem>(Item)]()
-			{
-				TSharedPtr<ISceneOutlinerTreeItem> PinnedItem = WeakItem.Pin();
-				if (!PinnedItem->IsValid())
+			[
+				SNew(SPinTypeSelector, FGetPinTypeTree::CreateStatic(&Editor::FUtils::GetFilteredVariableTypeTree))
+				.TargetPinType_Lambda([WeakItem = TWeakPtr<ISceneOutlinerTreeItem>(Item)]()
 				{
-					return FEdGraphPinType();
-				}
+					TSharedPtr<ISceneOutlinerTreeItem> PinnedItem = WeakItem.Pin();
+					if (!PinnedItem->IsValid())
+					{
+						return FEdGraphPinType();
+					}
 
-				const FVariablesOutlinerEntryItem* VariablesItem = PinnedItem->CastTo<FVariablesOutlinerEntryItem>();
-				
-				if(const IAnimNextRigVMVariableInterface* Variable = Cast<IAnimNextRigVMVariableInterface>(VariablesItem->WeakEntry.Get()))
-				{
-					return UncookedOnly::FUtils::GetPinTypeFromParamType(Variable->GetType());
-				}
-				else if (const FProperty* Property = VariablesItem->PropertyPath.Get())
-				{
-					FEdGraphPinType PinType;
-					GetDefault<UEdGraphSchema_K2>()->ConvertPropertyToPinType(Property, PinType);
-					return PinType;
-				}
+					const FVariablesOutlinerEntryItem* VariablesItem = PinnedItem->CastTo<FVariablesOutlinerEntryItem>();
+					
+					if(const IAnimNextRigVMVariableInterface* Variable = Cast<IAnimNextRigVMVariableInterface>(VariablesItem->WeakEntry.Get()))
+					{
+						return UncookedOnly::FUtils::GetPinTypeFromParamType(Variable->GetType());
+					}
+					else if (const FProperty* Property = VariablesItem->PropertyPath.Get())
+					{
+						FEdGraphPinType PinType;
+						GetDefault<UEdGraphSchema_K2>()->ConvertPropertyToPinType(Property, PinType);
+						return PinType;
+					}
 
-				return FEdGraphPinType();
-			})
-			.OnPinTypeChanged_Lambda([WeakEntry = TreeItem->WeakEntry](const FEdGraphPinType& PinType)
-			{
-				if(IAnimNextRigVMVariableInterface* Variable = Cast<IAnimNextRigVMVariableInterface>(WeakEntry.Get()))
+					return FEdGraphPinType();
+				})
+				.OnPinTypeChanged_Lambda([WeakEntry = TreeItem->WeakEntry](const FEdGraphPinType& PinType)
 				{
-					const FAnimNextParamType ParamType = UncookedOnly::FUtils::GetParamTypeFromPinType(PinType);
-					if(ParamType.IsValid())
+					if(IAnimNextRigVMVariableInterface* Variable = Cast<IAnimNextRigVMVariableInterface>(WeakEntry.Get()))
 					{
-						FScopedTransaction Transaction(LOCTEXT("SetTypeTransaction", "Set Variable Type"));
+						const FAnimNextParamType ParamType = UncookedOnly::FUtils::GetParamTypeFromPinType(PinType);
+						if(ParamType.IsValid())
+						{
+							FScopedTransaction Transaction(LOCTEXT("SetTypeTransaction", "Set Variable Type"));
 
-						UncookedOnly::FUtils::SetVariableType(Cast<UAnimNextVariableEntry>(Variable), ParamType, true, true);
-						UAnimNextVariableSettings* Settings = GetMutableDefault<UAnimNextVariableSettings>();
-						Settings->SetLastVariableType(ParamType);
+							UncookedOnly::FUtils::SetVariableType(Cast<UAnimNextVariableEntry>(Variable), ParamType, true, true);
+							UAnimNextVariableSettings* Settings = GetMutableDefault<UAnimNextVariableSettings>();
+							Settings->SetLastVariableType(ParamType);
+						}
 					}
-				}
-			})
-			.Schema(GetDefault<UPropertyBagSchema>())
-			.bAllowArrays(true)
-			.TypeTreeFilter(ETypeTreeFilter::None)
-			.Font(IDetailLayoutBuilder::GetDetailFont())
-			.SelectorType(SPinTypeSelector::ESelectorType::Compact)
+				})
+				.Schema(GetDefault<UPropertyBagSchema>())
+				.bAllowArrays(true)
+				.TypeTreeFilter(ETypeTreeFilter::None)
+				.ReadOnly_Lambda([&Row, this]() {return !Row.IsHovered(); })
+				.SelectorType(SelectorType)
 		];
 }
 
+float FVariablesOutlinerTypeColumn::GetColumnWidth() const
+{
+	switch (SelectorType)
+	{
+	case SPinTypeSelector::ESelectorType::Partial:
+		return 128.f;
+		
+	case SPinTypeSelector::ESelectorType::Full:
+		return 164.f;
+			
+	case SPinTypeSelector::ESelectorType::Compact:
+	default:
+		return 48.f;
+	}
+}
+
 FLazyName VariablesOutlinerValue("Value");
 
 FName FVariablesOutlinerValueColumn::GetID()
diff --git a/Engine/Plugins/Experimental/UAF/UAF/Source/UAFEditor/Private/Variables/Outliner/VariablesOutlinerColumns.h b/Engine/Plugins/Experimental/UAF/UAF/Source/UAFEditor/Private/Variables/Outliner/VariablesOutlinerColumns.h
--- a/Engine/Plugins/Experimental/UAF/UAF/Source/UAFEditor/Private/Variables/Outliner/VariablesOutlinerColumns.h
+++ b/Engine/Plugins/Experimental/UAF/UAF/Source/UAFEditor/Private/Variables/Outliner/VariablesOutlinerColumns.h
@@ -5,6 +5,7 @@
 #include "ISceneOutliner.h"
 #include "ISceneOutlinerColumn.h"
 #include "SceneOutlinerFwd.h"
+#include "SPinTypeSelector.h"
 #include "Misc/NotifyHook.h"
 #include "Templates/SharedPointer.h"
 #include "UObject/NameTypes.h"
@@ -20,9 +21,7 @@ class FVariablesOutlinerTypeColumn : public ISceneOutlinerColumn
 public:
 	static FName GetID();
 
-	FVariablesOutlinerTypeColumn(ISceneOutliner& SceneOutliner)
-		: WeakSceneOutliner(StaticCastSharedRef<ISceneOutliner>(SceneOutliner.AsShared()))
-	{}
+	FVariablesOutlinerTypeColumn(ISceneOutliner& SceneOutliner);
 
 	// ISceneOutlinerColumn interface
 	virtual FName GetColumnID() override { return GetID(); }
@@ -30,8 +29,12 @@ public:
 	virtual const TSharedRef<SWidget> ConstructRowWidget(FSceneOutlinerTreeItemRef TreeItem, const STableRow<FSceneOutlinerTreeItemPtr>& Row) override;
 	virtual bool SupportsSorting() const override { return false; }
 
+protected:
+	float GetColumnWidth() const;
+	
 private:
 	TWeakPtr<ISceneOutliner> WeakSceneOutliner;
+	SPinTypeSelector::ESelectorType SelectorType;
 };
 
 class FVariablesOutlinerValueColumn : public ISceneOutlinerColumn
diff --git a/Engine/Plugins/Experimental/UAF/UAF/Source/UAFEditor/Private/Variables/SVariablesView.cpp b/Engine/Plugins/Experimental/UAF/UAF/Source/UAFEditor/Private/Variables/SVariablesView.cpp
--- a/Engine/Plugins/Experimental/UAF/UAF/Source/UAFEditor/Private/Variables/SVariablesView.cpp
+++ b/Engine/Plugins/Experimental/UAF/UAF/Source/UAFEditor/Private/Variables/SVariablesView.cpp
@@ -237,7 +237,8 @@ void SVariablesView::Construct(const FArguments& InArgs, TSharedRef<UE::Workspac
 	InitOptions.OutlinerIdentifier = TEXT("AnimNextVariablesOutliner");
 	InitOptions.bShowHeaderRow = true;
 	InitOptions.ColumnMap.Add(FSceneOutlinerBuiltInColumnTypes::Label(), FSceneOutlinerColumnInfo(ESceneOutlinerColumnVisibility::Visible, 0, FCreateSceneOutlinerColumn(), false, 0.5f));
-	InitOptions.ColumnMap.Add(FVariablesOutlinerTypeColumn::GetID(), FSceneOutlinerColumnInfo(ESceneOutlinerColumnVisibility::Visible, 10, FCreateSceneOutlinerColumn::CreateLambda([](ISceneOutliner& InSceneOutliner) { return MakeShared<FVariablesOutlinerTypeColumn>(InSceneOutliner); }), false));
+	InitOptions.ColumnMap.Add(FVariablesOutlinerTypeColumn::GetID(), FSceneOutlinerColumnInfo(ESceneOutlinerColumnVisibility::Visible, 10, FCreateSceneOutlinerColumn::CreateLambda([](ISceneOutliner& InSceneOutliner) { return MakeShared<FVariablesOutlinerTypeColumn>(InSceneOutliner); }), false,
+		TOptional<float>(), TAttribute<FText>(), EHeaderComboVisibility::Ghosted));
 	InitOptions.ColumnMap.Add(FVariablesOutlinerValueColumn::GetID(), FSceneOutlinerColumnInfo(ESceneOutlinerColumnVisibility::Visible, 20, FCreateSceneOutlinerColumn::CreateLambda([](ISceneOutliner& InSceneOutliner) { return MakeShared<FVariablesOutlinerValueColumn>(InSceneOutliner); }), true, 0.5f));
 	InitOptions.ColumnMap.Add(FSceneOutlinerSourceControlColumn::GetID(), FSceneOutlinerColumnInfo(ESceneOutlinerColumnVisibility::Visible, 30, FCreateSceneOutlinerColumn::CreateLambda([](ISceneOutliner& InSceneOutliner) { return MakeShared<FSceneOutlinerSourceControlColumn>(InSceneOutliner); }), true));
 	InitOptions.ColumnMap.Add(FVariablesOutlinerAccessSpecifierColumn::GetID(), FSceneOutlinerColumnInfo(ESceneOutlinerColumnVisibility::Visible, 40, FCreateSceneOutlinerColumn::CreateLambda([](ISceneOutliner& InSceneOutliner) { return MakeShared<FVariablesOutlinerAccessSpecifierColumn>(InSceneOutliner); }), false));
diff --git a/Engine/Plugins/Experimental/UAF/UAF/Source/UAFEditor/UAFEditor.Build.cs b/Engine/Plugins/Experimental/UAF/UAF/Source/UAFEditor/UAFEditor.Build.cs
--- a/Engine/Plugins/Experimental/UAF/UAF/Source/UAFEditor/UAFEditor.Build.cs
+++ b/Engine/Plugins/Experimental/UAF/UAF/Source/UAFEditor/UAFEditor.Build.cs
@@ -56,7 +56,8 @@ namespace UnrealBuildTool.Rules
 					"RewindDebuggerInterface",
 					"GameplayInsights",
 					"Projects",
-					"Persona"
+					"Persona",
+					"DeveloperSettings"
 				}
 			);
 		}
diff --git a/Engine/Source/Editor/KismetWidgets/Private/SPinTypeSelector.cpp b/Engine/Source/Editor/KismetWidgets/Private/SPinTypeSelector.cpp
--- a/Engine/Source/Editor/KismetWidgets/Private/SPinTypeSelector.cpp
+++ b/Engine/Source/Editor/KismetWidgets/Private/SPinTypeSelector.cpp
@@ -562,9 +562,11 @@ void SPinTypeSelector::Construct(const FArguments& InArgs, FGetPinTypeTree GetPi
 	this->ChildSlot
 	[
 		SNew(SWidgetSwitcher)
-		.WidgetIndex_Lambda([this](){return ReadOnly.Get() ? 1 : 0; })
+		// Don't show the read-only widget when in Compact mode
+		.WidgetIndex_Lambda([this](){return ReadOnly.Get() && SelectorType != ESelectorType::Compact ? 1 : 0; })
 		+ SWidgetSwitcher::Slot() // editable version
-		.Padding(SelectorType == ESelectorType::Partial ? FMargin(-6.0f, 0.0f,0.0f,0.0f) : FMargin(0))
+		// Apply additional padding to Partial and Full selector type to prevent padding jumps
+		.Padding(SelectorType == ESelectorType::Partial || SelectorType == ESelectorType::Full ? FMargin(-6.0f, 0.0f,0.0f,0.0f) : FMargin(0))
 		[
 			Widget.ToSharedRef()
 		]
