Prevent broadcasting notifies for entries/editor-data which is marked as garbage

diff --git a/Engine/Plugins/Experimental/UAF/UAF/Source/UAFUncookedOnly/Private/AnimNextRigVMAssetEditorData.cpp b/Engine/Plugins/Experimental/UAF/UAF/Source/UAFUncookedOnly/Private/AnimNextRigVMAssetEditorData.cpp
--- a/Engine/Plugins/Experimental/UAF/UAF/Source/UAFUncookedOnly/Private/AnimNextRigVMAssetEditorData.cpp
+++ b/Engine/Plugins/Experimental/UAF/UAF/Source/UAFUncookedOnly/Private/AnimNextRigVMAssetEditorData.cpp
@@ -55,6 +55,11 @@ void UAnimNextRigVMAssetEditorData::BroadcastModified(EAnimNextEditorDataNotifTy
 {
 	using namespace UE::UAF::UncookedOnly;
 
+	if (!IsValid(this))
+	{
+		return;
+	}
+
 	// Modifications here can trigger compilation, so add a scope to catch compile batches
 	UAnimNextRigVMAsset* Asset = FUtils::GetAsset<UAnimNextRigVMAsset>(this);
 	FScopedCompileJob CompilerResults(LOCTEXT("ModifiedAssetsJobName", "Modified Assets"), { Asset });
@@ -108,8 +113,12 @@ void UAnimNextRigVMAssetEditorData::Serialize(FArchive& Ar)
 {
 	Ar.UsingCustomVersion(FFortniteMainBranchObjectVersion::GUID);
 
-	RigVMClient.SetDefaultSchemaClass(UAnimNextRigVMAssetSchema::StaticClass());
-	RigVMClient.SetOuterClientHost(this, GET_MEMBER_NAME_CHECKED(UAnimNextRigVMAssetEditorData, RigVMClient));
+	// Prevent doing this during transactions, as NewObject<UAnimNextRigVMAssetEditorData> with RF_Transactional will mark _this_ as garbage which breaks assumptions around SetOuterClientHost (reliant on TWeakObjPtr which thus returns nullptr)
+	if (!Ar.IsTransacting())
+	{
+		RigVMClient.SetDefaultSchemaClass(UAnimNextRigVMAssetSchema::StaticClass());
+		RigVMClient.SetOuterClientHost(this, GET_MEMBER_NAME_CHECKED(UAnimNextRigVMAssetEditorData, RigVMClient));
+	}
 
 	Super::Serialize(Ar);
 
diff --git a/Engine/Plugins/Experimental/UAF/UAF/Source/UAFUncookedOnly/Private/Entries/AnimNextRigVMAssetEntry.cpp b/Engine/Plugins/Experimental/UAF/UAF/Source/UAFUncookedOnly/Private/Entries/AnimNextRigVMAssetEntry.cpp
--- a/Engine/Plugins/Experimental/UAF/UAF/Source/UAFUncookedOnly/Private/Entries/AnimNextRigVMAssetEntry.cpp
+++ b/Engine/Plugins/Experimental/UAF/UAF/Source/UAFUncookedOnly/Private/Entries/AnimNextRigVMAssetEntry.cpp
@@ -51,6 +51,11 @@ void UAnimNextRigVMAssetEntry::PostEditUndo()
 
 void UAnimNextRigVMAssetEntry::BroadcastModified(EAnimNextEditorDataNotifType InType)
 {
+	if (!IsValid(this))
+	{
+		return;
+	}
+	
 	if(UAnimNextRigVMAssetEditorData* EditorData = Cast<UAnimNextRigVMAssetEditorData>(GetOuter()))
 	{
 		EditorData->BroadcastModified(InType, this);
