Add ability to create SharedVariables from outliner selection
* Requires selection to not leave an asset empty (e.g. removing sole variable from existing asset)
* Prompts user for asset path (pre-populated with currently focussed asset its path)

diff --git a/Engine/Plugins/Experimental/UAF/UAF/Source/UAFEditor/Internal/Variables/AnimNextSharedVariablesFactory.h b/Engine/Plugins/Experimental/UAF/UAF/Source/UAFEditor/Internal/Variables/AnimNextSharedVariablesFactory.h
--- a/Engine/Plugins/Experimental/UAF/UAF/Source/UAFEditor/Internal/Variables/AnimNextSharedVariablesFactory.h
+++ b/Engine/Plugins/Experimental/UAF/UAF/Source/UAFEditor/Internal/Variables/AnimNextSharedVariablesFactory.h
@@ -11,6 +11,7 @@ class UAnimNextSharedVariablesFactory : public UFactory
 {
 	GENERATED_BODY()
 
+public:
 	UAnimNextSharedVariablesFactory();
 	
 	virtual bool ConfigureProperties() override;
diff --git a/Engine/Plugins/Experimental/UAF/UAF/Source/UAFEditor/Private/Variables/Outliner/VariablesOutlinerCommands.cpp b/Engine/Plugins/Experimental/UAF/UAF/Source/UAFEditor/Private/Variables/Outliner/VariablesOutlinerCommands.cpp
--- a/Engine/Plugins/Experimental/UAF/UAF/Source/UAFEditor/Private/Variables/Outliner/VariablesOutlinerCommands.cpp
+++ b/Engine/Plugins/Experimental/UAF/UAF/Source/UAFEditor/Private/Variables/Outliner/VariablesOutlinerCommands.cpp
@@ -17,6 +17,8 @@ void FVariablesOutlinerCommands::RegisterCommands()
 	UI_COMMAND(SaveAsset, "Save Asset", "Save the selected asset.", EUserInterfaceActionType::Button, FInputChord());
 
 	UI_COMMAND(ToggleVariableExport, "Toggle Variable Export", "Sets Variable to be either Public or Private.", EUserInterfaceActionType::Button, FInputChord());
+
+	UI_COMMAND(CreateSharedVariablesAssets, "Create Shared Variables from Selection", "Creates a new Shared Variables asset from the selected set of variables.", EUserInterfaceActionType::Button, FInputChord());
 }
 
 #undef LOCTEXT_NAMESPACE
\ No newline at end of file
diff --git a/Engine/Plugins/Experimental/UAF/UAF/Source/UAFEditor/Private/Variables/Outliner/VariablesOutlinerCommands.h b/Engine/Plugins/Experimental/UAF/UAF/Source/UAFEditor/Private/Variables/Outliner/VariablesOutlinerCommands.h
--- a/Engine/Plugins/Experimental/UAF/UAF/Source/UAFEditor/Private/Variables/Outliner/VariablesOutlinerCommands.h
+++ b/Engine/Plugins/Experimental/UAF/UAF/Source/UAFEditor/Private/Variables/Outliner/VariablesOutlinerCommands.h
@@ -24,6 +24,7 @@ public:
 	TSharedPtr<FUICommandInfo> FindReferencesInAsset;
 	TSharedPtr<FUICommandInfo> SaveAsset;	
 	TSharedPtr<FUICommandInfo> ToggleVariableExport;
+	TSharedPtr<FUICommandInfo> CreateSharedVariablesAssets;
 };
 
 #undef LOCTEXT_NAMESPACE
\ No newline at end of file
diff --git a/Engine/Plugins/Experimental/UAF/UAF/Source/UAFEditor/Private/Variables/Outliner/VariablesOutlinerMode.cpp b/Engine/Plugins/Experimental/UAF/UAF/Source/UAFEditor/Private/Variables/Outliner/VariablesOutlinerMode.cpp
--- a/Engine/Plugins/Experimental/UAF/UAF/Source/UAFEditor/Private/Variables/Outliner/VariablesOutlinerMode.cpp
+++ b/Engine/Plugins/Experimental/UAF/UAF/Source/UAFEditor/Private/Variables/Outliner/VariablesOutlinerMode.cpp
@@ -20,6 +20,7 @@
 #include "AnimNextRigVMAsset.h"
 #include "ContentBrowserModule.h"
 #include "FileHelpers.h"
+#include "IWorkspaceEditorModule.h"
 #include "StructViewerFilter.h"
 #include "StructViewerModule.h"
 #include "VariablesOutlinerAssetItem.h"
@@ -34,6 +35,7 @@
 #include "Entries/AnimNextSharedVariablesEntry.h"
 #include "Styling/SlateIconFinder.h"
 #include "Variables/AnimNextSharedVariables.h"
+#include "Variables/AnimNextSharedVariablesFactory.h"
 #include "Variables/AnimNextVariableSettings.h"
 #include "Widgets/Docking/SDockTab.h"
 
@@ -144,6 +146,8 @@ TSharedPtr<SWidget> FVariablesOutlinerMode::CreateContextMenu()
 							VariablesSection.AddMenuEntryWithCommandList(FVariablesOutlinerCommands::Get().ToggleVariableExport, MenuContext->WeakCommandList.Pin(), VariableEntry->GetExportAccessSpecifier() == EAnimNextExportAccessSpecifier::Private ? LOCTEXT("MakePublicMenuLabel", "Expose as Public") : LOCTEXT("MakePrivateMenuLabel", "Make Private"), TAttribute<FText>(), FSlateIcon(FAppStyle::Get().GetStyleSetName(), VariableEntry->GetExportAccessSpecifier() == EAnimNextExportAccessSpecifier::Private ? "Level.VisibleIcon16x" : "Level.NotVisibleHighlightIcon16x"));
 						}
 					}
+
+					AssetSection.AddMenuEntryWithCommandList(FVariablesOutlinerCommands::Get().CreateSharedVariablesAssets, MenuContext->WeakCommandList.Pin(), TAttribute<FText>(), TAttribute<FText>(), FSlateIcon(FUAFStyle::Get().GetStyleSetName(), "ClassIcon.AnimNextSharedVariables"));
 				}
 
 				if (TSharedPtr<SVariablesOutliner> Outliner = MenuContext->WeakOutliner.Pin())
@@ -267,6 +271,7 @@ bool FVariablesOutlinerMode::CanRenameItem(const ISceneOutlinerTreeItem& Item) c
 	{
 		return true;
 	}
+	
 	return false;
 }
 
@@ -352,6 +357,14 @@ void FVariablesOutlinerMode::BindCommands(const TSharedRef<FUICommandList>& OutC
 		FIsActionChecked(),
 		FIsActionButtonVisible::CreateRaw(this, &FVariablesOutlinerMode::CanToggleVariableExport)
 		);
+
+	CommandList->MapAction(		
+		FVariablesOutlinerCommands::Get().CreateSharedVariablesAssets,
+		FExecuteAction::CreateRaw(this, &FVariablesOutlinerMode::CreateSharedVariablesAssets),
+		FCanExecuteAction::CreateRaw(this, &FVariablesOutlinerMode::CanCreateSharedVariablesAssets),
+		FIsActionChecked(),
+		FIsActionButtonVisible::CreateRaw(this, &FVariablesOutlinerMode::CanCreateSharedVariablesAssets)
+	);
 }
 
 TUniquePtr<ISceneOutlinerHierarchy> FVariablesOutlinerMode::CreateHierarchy()
@@ -966,6 +979,127 @@ bool FVariablesOutlinerMode::IsAsset() const
 	return Selection.SelectedItems.ContainsByPredicate([](const TWeakPtr<ISceneOutlinerTreeItem>& WeakItem) { return WeakItem.IsValid() && WeakItem.Pin()->IsA<FVariablesOutlinerAssetItem>(); });	
 }
 
+void FVariablesOutlinerMode::CreateSharedVariablesAssets() const
+{
+	ensure(CanCreateSharedVariablesAssets());
+	const FSceneOutlinerItemSelection& Selection = SceneOutliner->GetSelection();
+
+	TArray<UAnimNextVariableEntry*> EntriesToMove;
+	Selection.ForEachItem<FVariablesOutlinerEntryItem>([&EntriesToMove](const FVariablesOutlinerEntryItem& VariableItem)
+	{
+		if (UAnimNextVariableEntry* VariableEntry = VariableItem.WeakEntry.Get())
+		{
+			EntriesToMove.Add(VariableEntry);
+		}
+	});
+
+	if (EntriesToMove.Num())
+	{
+		// Prompt user for asset path
+		FString SaveObjectPath;
+		{			
+			FSaveAssetDialogConfig SaveAssetDialogConfig;
+			{
+				// Populate default path to match currently focussed assets
+				TSharedPtr<UE::Workspace::IWorkspaceEditor> SharedWorkspaceEditor = WeakWorkspaceEditor.Pin();
+				if (SharedWorkspaceEditor.IsValid())
+				{
+					const Workspace::FWorkspaceDocument& Document =SharedWorkspaceEditor->GetFocusedWorkspaceDocument();
+					if (UObject* FocussedObject = Document.GetObject())
+					{
+						SaveAssetDialogConfig.DefaultPath = FPaths::GetPath(FocussedObject->GetPackage()->GetPathName());
+					}
+				}
+				
+				SaveAssetDialogConfig.DefaultAssetName = TEXT("NewSharedVariables");
+				SaveAssetDialogConfig.AssetClassNames.Add(UAnimNextSharedVariables::StaticClass()->GetClassPathName());
+				SaveAssetDialogConfig.ExistingAssetPolicy = ESaveAssetDialogExistingAssetPolicy::AllowButWarn;
+				SaveAssetDialogConfig.DialogTitleOverride = LOCTEXT("SaveLiveLinkPresetDialogTitle", "Add new Shared Variables Asset");
+			}
+
+			FContentBrowserModule& ContentBrowserModule = FModuleManager::LoadModuleChecked<FContentBrowserModule>("ContentBrowser");
+			SaveObjectPath = ContentBrowserModule.Get().CreateModalSaveAssetDialog(SaveAssetDialogConfig);
+		}
+
+		if (!SaveObjectPath.IsEmpty())
+		{
+			const FString PackagePath = FPackageName::ObjectPathToPackageName(SaveObjectPath);
+			const FString AssetName = FPaths::GetBaseFilename(PackagePath, true);
+			
+			const FText FormattedMessage = FText::Format(LOCTEXT("CreateSharedVariablesAsset", "Creating new Shared Variables Asset {0}"), FText::FromString(AssetName));
+			FScopedTransaction Transaction(FormattedMessage);			
+			
+			UPackage* Package = CreatePackage(*PackagePath);
+			Package->Modify();
+				
+			UAnimNextSharedVariablesFactory* Factory = NewObject<UAnimNextSharedVariablesFactory>(GetTransientPackage());
+			UAnimNextSharedVariables* NewSharedVariables = CastChecked<UAnimNextSharedVariables>(Factory->FactoryCreateNew(UAnimNextSharedVariables::StaticClass(), Package, *AssetName, RF_Public | RF_Standalone | RF_Transactional, nullptr, nullptr, NAME_None));
+			check(NewSharedVariables);
+			
+			TArray<UAnimNextRigVMAssetEditorData*> UniqueOuterEditorDatas;
+			for (UAnimNextVariableEntry* Entry : EntriesToMove)
+			{
+				// Force variable to be public as its being moved to SharedVariables asset
+				Entry->SetExportAccessSpecifier(EAnimNextExportAccessSpecifier::Public);
+
+				UniqueOuterEditorDatas.AddUnique(Entry->GetTypedOuter<UAnimNextRigVMAssetEditorData>());
+				UncookedOnly::FUtils::MoveVariableToAsset(Entry, NewSharedVariables);
+			}
+
+			for (UAnimNextRigVMAssetEditorData* OuterEditorData : UniqueOuterEditorDatas)
+			{
+				OuterEditorData->AddSharedVariables(NewSharedVariables);
+			}
+			
+			// mark asset dirty 
+			FAssetRegistryModule::AssetCreated(NewSharedVariables);
+			NewSharedVariables->MarkPackageDirty();
+		}
+	}
+}
+
+bool FVariablesOutlinerMode::CanCreateSharedVariablesAssets() const
+{
+	const FSceneOutlinerItemSelection& Selection = SceneOutliner->GetSelection();
+	const bool bSelectionOnlyContainsVariables = Selection.Num() == Selection.Num<FVariablesOutlinerEntryItem>();
+
+	// Prevent moving from an assets which would leave it empty (removing all variables)
+	TMap<const UAnimNextRigVMAssetEditorData*, int32> EditorDataToNumRemainingVariables;
+	Selection.ForEachItem<FVariablesOutlinerEntryItem>([&EditorDataToNumRemainingVariables](const FVariablesOutlinerEntryItem& Item)
+	{
+		if (const UAnimNextVariableEntry* VariableEntry = Item.WeakEntry.Get())
+		{
+			if (const UAnimNextRigVMAssetEditorData* EditorData = CastChecked<UAnimNextRigVMAssetEditorData>(VariableEntry->GetOuter()))
+			{
+				if (int32* ExistingEntryValue = EditorDataToNumRemainingVariables.Find(EditorData))
+				{
+					--(*ExistingEntryValue);
+				}
+				else
+				{
+					int32 NumVariables = 0;
+					EditorData->ForEachEntryOfType<UAnimNextVariableEntry>([&NumVariables](const UAnimNextVariableEntry* Entry)
+					{
+						++NumVariables;
+						return true;
+					});
+
+					EditorDataToNumRemainingVariables.Add(EditorData, NumVariables - 1);
+				}
+			}
+		}
+
+		return true;
+	});
+
+	TArray<int32> NumRemainVariablesArray;
+	
+	EditorDataToNumRemainingVariables.GenerateValueArray(NumRemainVariablesArray);
+	const bool bMoveWillLeaveEmptyAsset = NumRemainVariablesArray.Contains(0);
+	
+	return bSelectionOnlyContainsVariables && !bMoveWillLeaveEmptyAsset;
+}
+
 TSharedPtr<FDragDropOperation> FVariablesOutlinerMode::CreateDragDropOperation(const FPointerEvent& MouseEvent, const TArray<FSceneOutlinerTreeItemPtr>& InTreeItems) const
 {
 	if(InTreeItems.Num() != 1)
diff --git a/Engine/Plugins/Experimental/UAF/UAF/Source/UAFEditor/Private/Variables/Outliner/VariablesOutlinerMode.h b/Engine/Plugins/Experimental/UAF/UAF/Source/UAFEditor/Private/Variables/Outliner/VariablesOutlinerMode.h
--- a/Engine/Plugins/Experimental/UAF/UAF/Source/UAFEditor/Private/Variables/Outliner/VariablesOutlinerMode.h
+++ b/Engine/Plugins/Experimental/UAF/UAF/Source/UAFEditor/Private/Variables/Outliner/VariablesOutlinerMode.h
@@ -83,6 +83,8 @@ protected:
 	bool CanSaveAsset() const;
 	bool IsAsset() const;
 
+	void CreateSharedVariablesAssets() const;
+	bool CanCreateSharedVariablesAssets() const;
 private:
 	friend class FVariablesOutlinerHierarchy;
 
