Verify that users can assign an AnimNext Module asset to a Blueprint
Automate QMetry Test Case:  UE-TC-10711

diff --git a/Engine/Plugins/Experimental/UAF/UAF/Source/UAF/Public/Component/AnimNextComponent.h b/Engine/Plugins/Experimental/UAF/UAF/Source/UAF/Public/Component/AnimNextComponent.h
--- a/Engine/Plugins/Experimental/UAF/UAF/Source/UAF/Public/Component/AnimNextComponent.h
+++ b/Engine/Plugins/Experimental/UAF/UAF/Source/UAF/Public/Component/AnimNextComponent.h
@@ -198,6 +198,7 @@ private:
 
 private:
 	friend struct UE::UAF::UncookedOnly::FUtils;
+	friend struct FAssetBlueprintTests;
 	friend class UAnimNextComponentWorldSubsystem;
 	friend class UE::UAF::Editor::FAssetWizard;
 
