Fix crash in abstract skeleton editor when viewing attributes due to missing null check

#rnx

diff --git a/Engine/Plugins/Experimental/UAF/UAF/Source/UAFEditor/Private/UAF/AbstractSkeleton/Sets/SAttributesList.cpp b/Engine/Plugins/Experimental/UAF/UAF/Source/UAFEditor/Private/UAF/AbstractSkeleton/Sets/SAttributesList.cpp
--- a/Engine/Plugins/Experimental/UAF/UAF/Source/UAFEditor/Private/UAF/AbstractSkeleton/Sets/SAttributesList.cpp
+++ b/Engine/Plugins/Experimental/UAF/UAF/Source/UAFEditor/Private/UAF/AbstractSkeleton/Sets/SAttributesList.cpp
@@ -256,33 +256,36 @@ void SAttributesList::RegisterMenus()
 			FNewToolMenuSectionDelegate::CreateLambda([this](FToolMenuSection& Section)
 			{
 				FPersonaModule& PersonaModule = FModuleManager::LoadModuleChecked<FPersonaModule>("Persona");
-						
-				FToolMenuEntry WidgetEntry = FToolMenuEntry::InitWidget(
-					"AttributePicker",
-					PersonaModule.CreateAttributePicker(
-						SetBinding->GetSkeleton(),
-						FOnAttributesPicked::CreateLambda([this](const TConstArrayView<FAnimationAttributeIdentifier> InSelectedAttributes)
-							{
-								const FScopedTransaction Transaction(LOCTEXT("AddAttributeToSet", "Add Attribute to Set"));
-								SetBinding->Modify();
-
-								for (const FAnimationAttributeIdentifier& Attribute : InSelectedAttributes)
+				
+				if (SetBinding.IsValid())
+				{	
+					FToolMenuEntry WidgetEntry = FToolMenuEntry::InitWidget(
+						"AttributePicker",
+						PersonaModule.CreateAttributePicker(
+							SetBinding->GetSkeleton(),
+							FOnAttributesPicked::CreateLambda([this](const TConstArrayView<FAnimationAttributeIdentifier> InSelectedAttributes)
 								{
-									SetBinding->AddAttributeToSet(Attribute, NAME_None);
-								}
-
-								FSlateApplication::Get().DismissAllMenus();
-								RepopulateListData();
-							}),
+									const FScopedTransaction Transaction(LOCTEXT("AddAttributeToSet", "Add Attribute to Set"));
+									SetBinding->Modify();
+
+									for (const FAnimationAttributeIdentifier& Attribute : InSelectedAttributes)
+									{
+										SetBinding->AddAttributeToSet(Attribute, NAME_None);
+									}
+
+									FSlateApplication::Get().DismissAllMenus();
+									RepopulateListData();
+								}),
+							true,
+							false),
+						FText(),
 						true,
-						false),
-					FText(),
-					true,
-					false,
-					true
-				);
-
-				Section.AddEntry(WidgetEntry);
+						false,
+						true
+					);
+
+					Section.AddEntry(WidgetEntry);
+				}
 			}));
 	}
 
@@ -295,28 +298,31 @@ void SAttributesList::RegisterMenus()
 			{
 				FPersonaModule& PersonaModule = FModuleManager::LoadModuleChecked<FPersonaModule>("Persona");
 						
-				FToolMenuEntry WidgetEntry = FToolMenuEntry::InitWidget(
-					"CurvePicker",
-					PersonaModule.CreateCurvePicker(
-						SetBinding->GetSkeleton(),
-						FOnCurvePicked::CreateLambda([this](const FName& InCurve)
-							{
-								const FScopedTransaction Transaction(LOCTEXT("AddCurvesToSet", "Add Curve to Set"));
-								SetBinding->Modify();
-
-								FAnimationAttributeIdentifier CurveAttrId(InCurve, INDEX_NONE, NAME_None, FFloatAnimationAttribute::StaticStruct());
-								SetBinding->AddAttributeToSet(CurveAttrId, NAME_None);
-
-								FSlateApplication::Get().DismissAllMenus();
-								RepopulateListData();
-							})),
-					FText(),
-					true,
-					false,
-					true
-				);
-
-				Section.AddEntry(WidgetEntry);
+				if (SetBinding.IsValid())
+				{
+					FToolMenuEntry WidgetEntry = FToolMenuEntry::InitWidget(
+						"CurvePicker",
+						PersonaModule.CreateCurvePicker(
+							SetBinding->GetSkeleton(),
+							FOnCurvePicked::CreateLambda([this](const FName& InCurve)
+								{
+									const FScopedTransaction Transaction(LOCTEXT("AddCurvesToSet", "Add Curve to Set"));
+									SetBinding->Modify();
+
+									FAnimationAttributeIdentifier CurveAttrId(InCurve, INDEX_NONE, NAME_None, FFloatAnimationAttribute::StaticStruct());
+									SetBinding->AddAttributeToSet(CurveAttrId, NAME_None);
+
+									FSlateApplication::Get().DismissAllMenus();
+									RepopulateListData();
+								})),
+						FText(),
+						true,
+						false,
+						true
+					);
+
+					Section.AddEntry(WidgetEntry);
+				}
 			}));
 	}
 }
