Fix linker error when compiling UAFEditor and UAF in the same executable.

diff --git a/Engine/Plugins/Experimental/UAF/UAF/Source/UAFEditor/Private/AnimNextEditorModule.cpp b/Engine/Plugins/Experimental/UAF/UAF/Source/UAFEditor/Private/AnimNextEditorModule.cpp
--- a/Engine/Plugins/Experimental/UAF/UAF/Source/UAFEditor/Private/AnimNextEditorModule.cpp
+++ b/Engine/Plugins/Experimental/UAF/UAF/Source/UAFEditor/Private/AnimNextEditorModule.cpp
@@ -160,7 +160,7 @@ void FAnimNextEditorModule::StartupModule()
 	PropertyModule.RegisterCustomClassLayout("AnimNextRigVMAssetEditorData",
 		FOnGetDetailCustomizationInstance::CreateLambda([] { return MakeShared<FAnimNextRigVMAssetEditorDataCustomization>(); }));
 
-	GraphPanelPinFactory = MakeShared<FAnimNextGraphPanelPinFactory>();
+	GraphPanelPinFactory = MakeShared<FAnimNextGraphPanelPinFactoryEditor>();
 	FEdGraphUtilities::RegisterVisualPinFactory(GraphPanelPinFactory);
 
 	Workspace::IWorkspaceEditorModule& WorkspaceEditorModule = FModuleManager::Get().LoadModuleChecked<Workspace::IWorkspaceEditorModule>("WorkspaceEditor");
diff --git a/Engine/Plugins/Experimental/UAF/UAF/Source/UAFEditor/Private/AnimNextEditorModule.h b/Engine/Plugins/Experimental/UAF/UAF/Source/UAFEditor/Private/AnimNextEditorModule.h
--- a/Engine/Plugins/Experimental/UAF/UAF/Source/UAFEditor/Private/AnimNextEditorModule.h
+++ b/Engine/Plugins/Experimental/UAF/UAF/Source/UAFEditor/Private/AnimNextEditorModule.h
@@ -9,7 +9,7 @@
 class UAnimNextWorkspaceSchema;
 class FAnimNextGraphPanelNodeFactory;
 class UAnimNextWorkspaceEditorMode;
-class FAnimNextGraphPanelPinFactory;
+class FAnimNextGraphPanelPinFactoryEditor;
 
 namespace UE::UAF::Editor
 {
@@ -68,7 +68,7 @@ private:
 
 	TMap<FTopLevelAssetPath, FAssetCompilationHandlerFactoryDelegate> AssetCompilationHandlerFactories;
 
-	TSharedPtr<FAnimNextGraphPanelPinFactory> GraphPanelPinFactory;
+	TSharedPtr<FAnimNextGraphPanelPinFactoryEditor> GraphPanelPinFactory;
 	TSharedPtr<FModuleEventPropertyTypeIdentifier> ModuleEventPropertyTypeIdentifier;
 
 	friend class FLocatorContext;
diff --git a/Engine/Plugins/Experimental/UAF/UAF/Source/UAFEditor/Private/Graph/AnimNextGraphPanelPinFactory.cpp b/Engine/Plugins/Experimental/UAF/UAF/Source/UAFEditor/Private/Graph/AnimNextGraphPanelPinFactory.cpp
--- a/Engine/Plugins/Experimental/UAF/UAF/Source/UAFEditor/Private/Graph/AnimNextGraphPanelPinFactory.cpp
+++ b/Engine/Plugins/Experimental/UAF/UAF/Source/UAFEditor/Private/Graph/AnimNextGraphPanelPinFactory.cpp
@@ -1,4 +1,4 @@
-﻿// Copyright Epic Games, Inc. All Rights Reserved.
+// Copyright Epic Games, Inc. All Rights Reserved.
 
 #include "AnimNextGraphPanelPinFactory.h"
 
@@ -10,7 +10,7 @@
 #include "Variables/AnimNextSoftVariableReference.h"
 #include "Widgets/DeclarativeSyntaxSupport.h"
 
-TSharedPtr<SGraphPin> FAnimNextGraphPanelPinFactory::CreatePin(UEdGraphPin* InPin) const
+TSharedPtr<SGraphPin> FAnimNextGraphPanelPinFactoryEditor::CreatePin(UEdGraphPin* InPin) const
 {
 	if (URigVMEdGraphNode* RigNode = Cast<URigVMEdGraphNode>(InPin->GetOwningNode()))
 	{
diff --git a/Engine/Plugins/Experimental/UAF/UAF/Source/UAFEditor/Private/Graph/AnimNextGraphPanelPinFactory.h b/Engine/Plugins/Experimental/UAF/UAF/Source/UAFEditor/Private/Graph/AnimNextGraphPanelPinFactory.h
--- a/Engine/Plugins/Experimental/UAF/UAF/Source/UAFEditor/Private/Graph/AnimNextGraphPanelPinFactory.h
+++ b/Engine/Plugins/Experimental/UAF/UAF/Source/UAFEditor/Private/Graph/AnimNextGraphPanelPinFactory.h
@@ -7,7 +7,7 @@
 
 #define UE_API UAFEDITOR_API
 
-class FAnimNextGraphPanelPinFactory : public FGraphPanelPinFactory
+class FAnimNextGraphPanelPinFactoryEditor : public FGraphPanelPinFactory
 {
 public:
 	// FGraphPanelPinFactory interface
