UAF Read / Write Variable in Function Fixes
- With this variables can be read / written to in functions in UAF / UAF StateTree

---

Remap function variables to themselves, allowing variables to be written / read from if the calling graph implements the owning shared variables / data interface.

Separate compile context generation in 'OnPreCompileGetProgrammaticFunctionHeaders' & 'OnPreCompileGetProgrammaticVariables' into const methods 'BuildFunctionHeadersContext' & 'BuildProgrammaticVariablesContext' respectively

Fix an issue where RigVM Asset shared variables were determined using the wrong compile context, add ownership checks to prevent this issue in the future.

Fix an issue where AST not updating would cause compile issues forming connections during graph gen if more than one function writing to variables was present.

Fix incorrect variable indexing when primitive arg types were used in a function.

Prevent function headers from being added twice, leading to graph gen issues.

Add compile scopes to prevent un-needed rebroadcast of module init multiple times during compilation

#tests PIE

diff --git a/Engine/Plugins/Experimental/UAF/UAF/Source/UAFEditor/Private/Common/GraphEditorSchemaActions.cpp b/Engine/Plugins/Experimental/UAF/UAF/Source/UAFEditor/Private/Common/GraphEditorSchemaActions.cpp
--- a/Engine/Plugins/Experimental/UAF/UAF/Source/UAFEditor/Private/Common/GraphEditorSchemaActions.cpp
+++ b/Engine/Plugins/Experimental/UAF/UAF/Source/UAFEditor/Private/Common/GraphEditorSchemaActions.cpp
@@ -441,6 +441,19 @@ UEdGraphNode* FAnimNextSchemaAction_Function::PerformAction(UEdGraph* ParentGrap
 				Controller->ClearNodeSelection(true, true);
 				Controller->SelectNode(ModelNode, true, true, true);
 			}
+
+			const TArray<FRigVMExternalVariable> ExternalVariables = ModelNode->GetExternalVariables(false);
+			if (!ExternalVariables.IsEmpty())
+			{
+				if (URigVMLibraryNode* LibraryNode = ModelNode->LoadReferencedNode())
+				{
+					for (const FRigVMExternalVariable& ExternalVariable : ExternalVariables)
+					{
+						Controller->SetRemappedVariable(ModelNode, ExternalVariable.Name, ExternalVariable.Name);
+					}
+				}
+			}
+
 			Controller->CloseUndoBracket();
 		}
 		else
diff --git a/Engine/Plugins/Experimental/UAF/UAF/Source/UAFUncookedOnly/Internal/AnimNextRigVMAssetEditorData.h b/Engine/Plugins/Experimental/UAF/UAF/Source/UAFUncookedOnly/Internal/AnimNextRigVMAssetEditorData.h
--- a/Engine/Plugins/Experimental/UAF/UAF/Source/UAFUncookedOnly/Internal/AnimNextRigVMAssetEditorData.h
+++ b/Engine/Plugins/Experimental/UAF/UAF/Source/UAFUncookedOnly/Internal/AnimNextRigVMAssetEditorData.h
@@ -399,6 +399,12 @@ protected:
 	////////////////////////////////////////////////////////////////////////////////////////////////////////////////////////
 	// End Compilation overrides
 
+	// Return a function header compile context. Used when an external asset needs compile contexts for this asset.
+	UE_API FAnimNextGetFunctionHeaderCompileContext GetFunctionHeaderContext(const FRigVMCompileSettings& InSettings, FAnimNextRigVMAssetCompileContext& InCompileContext) const;
+
+	// Return a variable compile context. Used when an external asset needs compile contexts for this asset.
+	UE_API FAnimNextGetVariableCompileContext GetVariableCompileContext(const FRigVMCompileSettings& InSettings, FAnimNextRigVMAssetCompileContext& InCompileContext) const;
+
 	// Customization point for derived types to transform new asset entries
 	virtual void CustomizeNewAssetEntry(UAnimNextRigVMAssetEntry* InNewEntry) const {}
 
@@ -601,14 +607,20 @@ protected:
 	// Check whether this asset should be recompiled
 	UE_API bool IsDirtyForRecompilation() const;
 
+	// Called pre programmatic function header population to allow population of compile context in a const manner. May be used by external assets querying our VM.
+	UE_API virtual void BuildFunctionHeadersContext(const FRigVMCompileSettings& InSettings, FAnimNextGetFunctionHeaderCompileContext& OutCompileContext) const {}
+
+	// Called pre programmatic variable population to allow population of compile context in a const manner. May be used by external assets querying our VM.
+	UE_API virtual void BuildProgrammaticVariablesContext(const FRigVMCompileSettings& InSettings, FAnimNextGetVariableCompileContext& OutCompileContext) const {}
+
 	// Called during compilation to generate internal variables used for natively-callable RigVM functions
-	void BuildFunctionWrapperEventVariables(FAnimNextRigVMAssetCompileContext& InContext);
+	UE_API void BuildFunctionWrapperEventVariables(FAnimNextRigVMAssetCompileContext& InContext) const;
 
 	// Called during compilation to generate wrapper events for natively-callable RigVM functions
-	void BuildFunctionWrapperEvents(FAnimNextRigVMAssetCompileContext& InContext, const FRigVMCompileSettings& InSettings);
+	UE_API void BuildFunctionWrapperEvents(FAnimNextRigVMAssetCompileContext& InContext, const FRigVMCompileSettings& InSettings);
 
 	// Generates a combined property bag consisting of all variables both internal and shared
-	UE_API FInstancedPropertyBag GenerateCombinedPropertyBag(const FAnimNextGetVariableCompileContext& InCompileContext) const;
+	UE_API FInstancedPropertyBag GenerateCombinedPropertyBag(const FRigVMCompileSettings& InSettings, const FAnimNextGetVariableCompileContext& InCompileContext) const;
 
 	// Serialization/post-load fixup for old data
 	void UpgradeDataInterfaces();
diff --git a/Engine/Plugins/Experimental/UAF/UAF/Source/UAFUncookedOnly/Internal/Compilation/AnimNextGetFunctionHeaderCompileContext.h b/Engine/Plugins/Experimental/UAF/UAF/Source/UAFUncookedOnly/Internal/Compilation/AnimNextGetFunctionHeaderCompileContext.h
--- a/Engine/Plugins/Experimental/UAF/UAF/Source/UAFUncookedOnly/Internal/Compilation/AnimNextGetFunctionHeaderCompileContext.h
+++ b/Engine/Plugins/Experimental/UAF/UAF/Source/UAFUncookedOnly/Internal/Compilation/AnimNextGetFunctionHeaderCompileContext.h
@@ -6,6 +6,8 @@
 
 #define UE_API UAFUNCOOKEDONLY_API
 
+class UAnimNextRigVMAssetEditorData;
+
 /**
  * Struct holding temporary compilation info during function header population 
  */
@@ -16,7 +18,8 @@ struct FAnimNextGetFunctionHeaderCompileContext
 public:
 	
 	UE_API const TArray<FAnimNextProgrammaticFunctionHeader>& GetFunctionHeaders() const;
-	UE_API TArray<FAnimNextProgrammaticFunctionHeader>& GetMutableFunctionHeaders();
+	UE_API void AddUniqueFunctionHeader(const FAnimNextProgrammaticFunctionHeader& InFunctionHeader);
+	UE_API const UAnimNextRigVMAssetEditorData* GetOwningAssetEditorData() const;
 
 protected:
 
diff --git a/Engine/Plugins/Experimental/UAF/UAF/Source/UAFUncookedOnly/Internal/Compilation/AnimNextGetGraphCompileContext.h b/Engine/Plugins/Experimental/UAF/UAF/Source/UAFUncookedOnly/Internal/Compilation/AnimNextGetGraphCompileContext.h
--- a/Engine/Plugins/Experimental/UAF/UAF/Source/UAFUncookedOnly/Internal/Compilation/AnimNextGetGraphCompileContext.h
+++ b/Engine/Plugins/Experimental/UAF/UAF/Source/UAFUncookedOnly/Internal/Compilation/AnimNextGetGraphCompileContext.h
@@ -6,6 +6,8 @@
 
 #define UE_API UAFUNCOOKEDONLY_API
 
+class UAnimNextRigVMAssetEditorData;
+
 /**
  * Struct holding temporary compilation info during function header population
  */
@@ -19,6 +21,7 @@ public:
 
 	UE_API const TArray<URigVMGraph*>& GetProgrammaticGraphs() const;
 	UE_API TArray<URigVMGraph*>& GetMutableProgrammaticGraphs();
+	UE_API const UAnimNextRigVMAssetEditorData* GetOwningAssetEditorData() const;
 
 protected:
 
diff --git a/Engine/Plugins/Experimental/UAF/UAF/Source/UAFUncookedOnly/Internal/Compilation/AnimNextGetVariableCompileContext.h b/Engine/Plugins/Experimental/UAF/UAF/Source/UAFUncookedOnly/Internal/Compilation/AnimNextGetVariableCompileContext.h
--- a/Engine/Plugins/Experimental/UAF/UAF/Source/UAFUncookedOnly/Internal/Compilation/AnimNextGetVariableCompileContext.h
+++ b/Engine/Plugins/Experimental/UAF/UAF/Source/UAFUncookedOnly/Internal/Compilation/AnimNextGetVariableCompileContext.h
@@ -6,6 +6,8 @@
 
 #define UE_API UAFUNCOOKEDONLY_API
 
+class UAnimNextRigVMAssetEditorData;
+
 /**
  * Struct holding temporary compilation info during function header population
  */
@@ -19,6 +21,7 @@ public:
 
 	UE_API const TArray<FAnimNextProgrammaticVariable>& GetProgrammaticVariables() const;
 	UE_API TArray<FAnimNextProgrammaticVariable>& GetMutableProgrammaticVariables();
+	UE_API const UAnimNextRigVMAssetEditorData* GetOwningAssetEditorData() const;
 
 protected:
 
diff --git a/Engine/Plugins/Experimental/UAF/UAF/Source/UAFUncookedOnly/Internal/Compilation/AnimNextProcessGraphCompileContext.h b/Engine/Plugins/Experimental/UAF/UAF/Source/UAFUncookedOnly/Internal/Compilation/AnimNextProcessGraphCompileContext.h
--- a/Engine/Plugins/Experimental/UAF/UAF/Source/UAFUncookedOnly/Internal/Compilation/AnimNextProcessGraphCompileContext.h
+++ b/Engine/Plugins/Experimental/UAF/UAF/Source/UAFUncookedOnly/Internal/Compilation/AnimNextProcessGraphCompileContext.h
@@ -6,6 +6,8 @@
 
 #define UE_API UAFUNCOOKEDONLY_API
 
+class UAnimNextRigVMAssetEditorData;
+
 /**
  * Struct holding temporary compilation info during function header population
  */
@@ -17,6 +19,7 @@ public:
 
 	UE_API const TArray<URigVMGraph*>& GetAllGraphs() const;
 	UE_API TArray<URigVMGraph*>& GetMutableAllGraphs();
+	UE_API const UAnimNextRigVMAssetEditorData* GetOwningAssetEditorData() const;
 
 protected:
 
diff --git a/Engine/Plugins/Experimental/UAF/UAF/Source/UAFUncookedOnly/Internal/Compilation/AnimNextProgrammaticFunctionHeader.h b/Engine/Plugins/Experimental/UAF/UAF/Source/UAFUncookedOnly/Internal/Compilation/AnimNextProgrammaticFunctionHeader.h
--- a/Engine/Plugins/Experimental/UAF/UAF/Source/UAFUncookedOnly/Internal/Compilation/AnimNextProgrammaticFunctionHeader.h
+++ b/Engine/Plugins/Experimental/UAF/UAF/Source/UAFUncookedOnly/Internal/Compilation/AnimNextProgrammaticFunctionHeader.h
@@ -14,5 +14,7 @@ struct FAnimNextProgrammaticFunctionHeader
 {
 	/** RigVM Function Header we are wrapping */
 	FRigVMGraphFunctionHeader FunctionHeader;
+
+	bool operator==(const FAnimNextProgrammaticFunctionHeader& Other) const = default;
 };
 
diff --git a/Engine/Plugins/Experimental/UAF/UAF/Source/UAFUncookedOnly/Internal/Compilation/AnimNextRigVMAssetCompileContext.h b/Engine/Plugins/Experimental/UAF/UAF/Source/UAFUncookedOnly/Internal/Compilation/AnimNextRigVMAssetCompileContext.h
--- a/Engine/Plugins/Experimental/UAF/UAF/Source/UAFUncookedOnly/Internal/Compilation/AnimNextRigVMAssetCompileContext.h
+++ b/Engine/Plugins/Experimental/UAF/UAF/Source/UAFUncookedOnly/Internal/Compilation/AnimNextRigVMAssetCompileContext.h
@@ -10,6 +10,7 @@
 #include "Variables/AnimNextProgrammaticVariable.h"
 
 class URigVMGraph;
+class UAnimNextRigVMAssetEditorData;
 
 /**
  * Struct holding temporary compilation info for a single RigVM asset
@@ -19,6 +20,13 @@ class URigVMGraph;
  */
 struct FAnimNextRigVMAssetCompileContext
 {
+public:
+
+	FAnimNextRigVMAssetCompileContext(const UAnimNextRigVMAssetEditorData* InOwningAssetEditorData)
+		: OwningAssetEditorData(InOwningAssetEditorData)
+	{
+	}
+
 private:
 	friend class UAnimNextRigVMAssetEditorData;
 	friend struct FAnimNextGetFunctionHeaderCompileContext;
@@ -39,5 +47,8 @@ protected:
 
 	/** All graphs compiled by this RigVM Asset */
 	TArray<URigVMGraph*> AllGraphs;
+
+	/** RigVM Asset Editor Data this compile context is designed for. Used to ensure we don't cross contexts. */
+	const UAnimNextRigVMAssetEditorData* OwningAssetEditorData;
 };
 
diff --git a/Engine/Plugins/Experimental/UAF/UAF/Source/UAFUncookedOnly/Private/AnimNextRigVMAssetEditorData.cpp b/Engine/Plugins/Experimental/UAF/UAF/Source/UAFUncookedOnly/Private/AnimNextRigVMAssetEditorData.cpp
--- a/Engine/Plugins/Experimental/UAF/UAF/Source/UAFUncookedOnly/Private/AnimNextRigVMAssetEditorData.cpp
+++ b/Engine/Plugins/Experimental/UAF/UAF/Source/UAFUncookedOnly/Private/AnimNextRigVMAssetEditorData.cpp
@@ -807,7 +807,7 @@ TObjectPtr<URigVMGraph> UAnimNextRigVMAssetEditorData::CreateContainedGraphModel
 	return Model;
 }
 
-void UAnimNextRigVMAssetEditorData::BuildFunctionWrapperEventVariables(FAnimNextRigVMAssetCompileContext& InContext)
+void UAnimNextRigVMAssetEditorData::BuildFunctionWrapperEventVariables(FAnimNextRigVMAssetCompileContext& InContext) const
 {
 	using namespace UE::UAF::UncookedOnly;
 	
@@ -819,7 +819,7 @@ void UAnimNextRigVMAssetEditorData::BuildFunctionWrapperEventVariables(FAnimNext
 			// Don't create internal variables for execution IO, visible only, or hidden pins. Those types do not need variable nodes in the graph.
 			if (Argument.Direction == ERigVMPinDirection::Input || Argument.Direction == ERigVMPinDirection::Output)
 			{
-				if (Argument.IsCPPTypeObjectValid())
+				if (Argument.IsValid())
 				{
 					FRigVMGraphFunctionArgument InternallyNamedArgument = Argument;
 					InternallyNamedArgument.Name = FName(FUtils::MakeFunctionWrapperVariableName(FunctionHeader.Name, Argument.Name));
@@ -849,6 +849,10 @@ void UAnimNextRigVMAssetEditorData::BuildFunctionWrapperEvents(FAnimNextRigVMAss
 
 	for(const FAnimNextProgrammaticFunctionHeader& AnimNextFunctionHeader : InContext.FunctionHeaders)
 	{
+		// Controller needs to notify the AST of variable changes to make new links
+		constexpr bool bSuspendNotificationForInternalVariables = false;
+		FRigVMControllerNotifGuard VarNotifGuard(Controller, bSuspendNotificationForInternalVariables);
+
 		const FRigVMGraphFunctionHeader& FunctionHeader = AnimNextFunctionHeader.FunctionHeader;
 
 		URigVMLibraryNode* LibraryNode = Cast<URigVMLibraryNode>(FunctionHeader.LibraryPointer.GetNodeSoftPath().TryLoad());
@@ -882,6 +886,15 @@ void UAnimNextRigVMAssetEditorData::BuildFunctionWrapperEvents(FAnimNextRigVMAss
 			continue;
 		}
 
+		const TArray<FRigVMExternalVariable> FuncExternalVariables = FunctionNode->GetExternalVariables(false);
+		if (!ExternalVariables.IsEmpty())
+		{
+			for (const FRigVMExternalVariable& FuncExternalVariable : FuncExternalVariables)
+			{
+				Controller->SetRemappedVariable(FunctionNode, FuncExternalVariable.Name, FuncExternalVariable.Name);
+			}
+		}
+
 		// Link up Execute nodes if needed, function may be pure & lack an input pin
 		URigVMPin* CurrentExecuteOutputPin = EventNode->FindPin(FRigVMStruct::ExecuteContextName.ToString());
 		URigVMPin* ExecuteInputPin = FunctionNode->FindPin(FRigVMStruct::ExecuteContextName.ToString());
@@ -894,16 +907,12 @@ void UAnimNextRigVMAssetEditorData::BuildFunctionWrapperEvents(FAnimNextRigVMAss
 		// Update current execute pin, RigVM doesn't have a concept of input / output execute pins, just one execute content pin used for both
 		CurrentExecuteOutputPin = ExecuteInputPin ? ExecuteInputPin : CurrentExecuteOutputPin;
 
-			// Controller needs to notify the AST of variable changes to make new links
-		bool bSuspendNotificationForInternalVariables = false;
-		FRigVMControllerNotifGuard VarNotifGuard(Controller, bSuspendNotificationForInternalVariables);
-
 		// Generate & link input arguments, also generate result variable node but link later
 		TArray<int32> ArgIndices;
 		for (const FRigVMGraphFunctionArgument& Argument : FunctionHeader.Arguments)
 		{
 			// Execution context is captured as arg pins, skip those for internal variable gen
-			if (Argument.Direction == ERigVMPinDirection::IO || !Argument.IsCPPTypeObjectValid())
+			if (Argument.Direction == ERigVMPinDirection::IO || !Argument.IsValid())
 			{
 				continue;
 			}
@@ -1031,7 +1040,7 @@ void UAnimNextRigVMAssetEditorData::RecompileVM()
 	RigGraphDisplaySettings.MinMicroSeconds = RigGraphDisplaySettings.LastMinMicroSeconds = DBL_MAX;
 	RigGraphDisplaySettings.MaxMicroSeconds = RigGraphDisplaySettings.LastMaxMicroSeconds = (double)INDEX_NONE;
 
-	FAnimNextRigVMAssetCompileContext CompileContext = {};
+	FAnimNextRigVMAssetCompileContext CompileContext = { this };
 	{
 		TGuardValue<bool> ReentrantGuardSelf(bSuspendModelNotificationsForSelf, true);
 		TGuardValue<bool> ReentrantGuardOthers(RigVMClient.bSuspendModelNotificationsForOthers, true);
@@ -1040,6 +1049,7 @@ void UAnimNextRigVMAssetEditorData::RecompileVM()
 
 		{
 			FAnimNextGetFunctionHeaderCompileContext GetFunctionHeaderCompileContext(CompileContext);
+			BuildFunctionHeadersContext(Settings, GetFunctionHeaderCompileContext);
 			OnPreCompileGetProgrammaticFunctionHeaders(Settings, GetFunctionHeaderCompileContext);
 		}
 
@@ -1047,6 +1057,7 @@ void UAnimNextRigVMAssetEditorData::RecompileVM()
 
 		{
 			FAnimNextGetVariableCompileContext GetVariableCompileContext(CompileContext);
+			BuildProgrammaticVariablesContext(Settings, GetVariableCompileContext);
 			FUtils::CompileVariables(Settings, Asset, GetVariableCompileContext);
 			OnPostCompileVariables(Settings, GetVariableCompileContext);
 		}
@@ -1794,6 +1805,33 @@ bool UAnimNextRigVMAssetEditorData::RemoveAllEntries(bool bSetupUndoRedo, bool b
 	return bResult;
 }
 
+FAnimNextGetFunctionHeaderCompileContext UAnimNextRigVMAssetEditorData::GetFunctionHeaderContext(const FRigVMCompileSettings& InSettings, FAnimNextRigVMAssetCompileContext& InCompileContext) const
+{
+	// Since the compile context is a reference it must be set in advance by the caller. If we created it we would return a ref to a temp.
+	ensureMsgf(this == InCompileContext.OwningAssetEditorData, TEXT("Crossing RigVM asset with a different asset's compile context. Expect incorrect memory layout / variable generation"));
+
+	FAnimNextGetFunctionHeaderCompileContext GetFunctionHeaderCompileContext(InCompileContext);
+	BuildFunctionHeadersContext(InSettings, GetFunctionHeaderCompileContext);
+	return GetFunctionHeaderCompileContext;
+}
+
+FAnimNextGetVariableCompileContext UAnimNextRigVMAssetEditorData::GetVariableCompileContext(const FRigVMCompileSettings& InSettings, FAnimNextRigVMAssetCompileContext& InCompileContext) const
+{
+	// Since the compile context is a reference it must be set in advance by the caller. If we created it we would return a ref to a temp.
+	ensureMsgf(this == InCompileContext.OwningAssetEditorData, TEXT("Crossing RigVM asset with a different asset's compile context. Expect incorrect memory layout / variable generation"));
+
+	{
+		FAnimNextGetFunctionHeaderCompileContext GetFunctionHeaderCompileContext(InCompileContext);
+		BuildFunctionHeadersContext(InSettings, GetFunctionHeaderCompileContext);
+	}
+
+	BuildFunctionWrapperEventVariables(InCompileContext);
+
+	FAnimNextGetVariableCompileContext GetVariableCompileContext(InCompileContext);
+	BuildProgrammaticVariablesContext(InSettings, GetVariableCompileContext);
+	return GetVariableCompileContext;
+}
+
 UObject* UAnimNextRigVMAssetEditorData::CreateNewSubEntry(UAnimNextRigVMAssetEditorData* InEditorData, TSubclassOf<UObject> InClass)
 {
 	UObject* NewEntry = NewObject<UObject>(InEditorData, InClass.Get(), NAME_None, RF_Transactional);
@@ -2098,6 +2136,8 @@ UAnimNextVariableEntry* UAnimNextRigVMAssetLibrary::AddVariable(UAnimNextRigVMAs
 
 UAnimNextVariableEntry* UAnimNextRigVMAssetEditorData::AddVariable(FName InName, FAnimNextParamType InType, const FString& InDefaultValue, bool bSetupUndoRedo, bool bPrintPythonCommand)
 {
+	using namespace UE::UAF::UncookedOnly;
+
 	if(InName == NAME_None)
 	{
 		ReportError(TEXT("UAnimNextRigVMAssetEditorData::AddVariable: Invalid variable name supplied."));
@@ -2122,6 +2162,10 @@ UAnimNextVariableEntry* UAnimNextRigVMAssetEditorData::AddVariable(FName InName,
 		return InEntry->GetEntryName() == NewParameterName;
 	};
 
+	// We are about to perform multiple changes, add a scope now to avoid multiple compile broadcasts per change
+	UAnimNextRigVMAsset* Asset = FUtils::GetAsset<UAnimNextRigVMAsset>(this);
+	FScopedCompileJob CompilerResults(LOCTEXT("ModifiedAssetsAddVariable", "Modified Assets Add Variable"), { Asset });
+
 	bool bAlreadyExists = Entries.ContainsByPredicate(DuplicateNamePredicate);
 	int32 NameNumber = InName.GetNumber() + 1;
 	while(bAlreadyExists)
@@ -2181,6 +2225,8 @@ UAnimNextEventGraphEntry* UAnimNextRigVMAssetLibrary::AddEventGraph(UAnimNextRig
 
 UAnimNextEventGraphEntry* UAnimNextRigVMAssetEditorData::AddEventGraph(FName InName, UScriptStruct* InEventStruct, bool bSetupUndoRedo, bool bPrintPythonCommand)
 {
+	using namespace UE::UAF::UncookedOnly;
+
 	if(InName == NAME_None)
 	{
 		ReportError(TEXT("UAnimNextRigVMAssetEditorData::AddEventGraph: Invalid graph name supplied."));
@@ -2214,6 +2260,10 @@ UAnimNextEventGraphEntry* UAnimNextRigVMAssetEditorData::AddEventGraph(FName InN
 		bAlreadyExists =  Entries.ContainsByPredicate(DuplicateNamePredicate);
 	}
 
+	// We are about to perform multiple changes, add a scope now to avoid multiple compile broadcasts per change
+	UAnimNextRigVMAsset* Asset = FUtils::GetAsset<UAnimNextRigVMAsset>(this);
+	FScopedCompileJob CompilerResults(LOCTEXT("ModifiedAssetsAddEventGraph", "Modified Assets Add Event Graph"), { Asset });
+
 	UAnimNextEventGraphEntry* NewEntry = CreateNewSubEntry<UAnimNextEventGraphEntry>(this);
 	NewEntry->GraphName = NewGraphName;
 	NewEntry->Initialize(this);
@@ -2270,6 +2320,8 @@ UAnimNextSharedVariablesEntry* UAnimNextRigVMAssetLibrary::AddSharedVariables(UA
 
 UAnimNextSharedVariablesEntry* UAnimNextRigVMAssetEditorData::AddSharedVariables(const UAnimNextSharedVariables* InSharedVariables, bool bSetupUndoRedo, bool bPrintPythonCommand)
 {
+	using namespace UE::UAF::UncookedOnly;
+
 	if(InSharedVariables == nullptr)
 	{
 		ReportError(TEXT("UAnimNextRigVMAssetEditorData::AddSharedVariables: Invalid asset supplied."));
@@ -2367,6 +2419,10 @@ UAnimNextSharedVariablesEntry* UAnimNextRigVMAssetEditorData::AddSharedVariables
 		return nullptr;
 	}
 
+	// We are about to perform multiple changes, add a scope now to avoid multiple compile broadcasts per change
+	UAnimNextRigVMAsset* Asset = FUtils::GetAsset<UAnimNextRigVMAsset>(this);
+	FScopedCompileJob CompilerResults(LOCTEXT("ModifiedAssetsAddSharedVariables", "Modified Assets Add Shared Variables"), { Asset });
+
 	UAnimNextSharedVariablesEntry* NewEntry = CreateNewSubEntry<UAnimNextSharedVariablesEntry>(this);
 	NewEntry->SetAsset(InSharedVariables);
 	NewEntry->Initialize(this);
@@ -2401,6 +2457,8 @@ UAnimNextSharedVariablesEntry* UAnimNextRigVMAssetLibrary::AddSharedVariablesStr
 
 UAnimNextSharedVariablesEntry* UAnimNextRigVMAssetEditorData::AddSharedVariablesStruct(const UScriptStruct* InStruct, bool bSetupUndoRedo, bool bPrintPythonCommand)
 {
+	using namespace UE::UAF::UncookedOnly;
+
 	if(InStruct == nullptr)
 	{
 		ReportError(TEXT("UAnimNextRigVMAssetEditorData::AddSharedVariablesStruct: Invalid struct supplied."));
@@ -2429,6 +2487,10 @@ UAnimNextSharedVariablesEntry* UAnimNextRigVMAssetEditorData::AddSharedVariables
 		return nullptr;
 	}
 
+	// We are about to perform multiple changes, add a scope now to avoid multiple compile broadcasts per change
+	UAnimNextRigVMAsset* Asset = FUtils::GetAsset<UAnimNextRigVMAsset>(this);
+	FScopedCompileJob CompilerResults(LOCTEXT("ModifiedAssetsAddSharedVariablesStruct", "Modified Assets Add Shared Variables Struct"), { Asset });
+
 	UAnimNextSharedVariablesEntry* NewEntry = CreateNewSubEntry<UAnimNextSharedVariablesEntry>(this);
 	NewEntry->SetStruct(InStruct);
 	NewEntry->Initialize(this);
@@ -2927,7 +2989,7 @@ void UAnimNextRigVMAssetEditorData::SetUseExternalPackages(TArrayView<UAnimNextR
 
 #endif
 
-FInstancedPropertyBag UAnimNextRigVMAssetEditorData::GenerateCombinedPropertyBag(const FAnimNextGetVariableCompileContext& InCompileContext) const
+FInstancedPropertyBag UAnimNextRigVMAssetEditorData::GenerateCombinedPropertyBag(const FRigVMCompileSettings& InSettings, const FAnimNextGetVariableCompileContext& InCompileContext) const
 {
 	UAnimNextRigVMAsset* Asset = UE::UAF::UncookedOnly::FUtils::GetAsset<UAnimNextRigVMAsset>(this);
 
@@ -2961,7 +3023,9 @@ FInstancedPropertyBag UAnimNextRigVMAssetEditorData::GenerateCombinedPropertyBag
 		}
 
 		const UAnimNextRigVMAssetEditorData* ReferencedEditorData = UE::UAF::UncookedOnly::FUtils::GetEditorData<UAnimNextRigVMAssetEditorData>(ReferencedVariableAsset);
-		AssetBag = UE::UAF::UncookedOnly::FUtils::MakePropertyBagForEditorData(ReferencedEditorData, InCompileContext);
+
+		FAnimNextRigVMAssetCompileContext CompileContext = { ReferencedEditorData };
+		AssetBag = UE::UAF::UncookedOnly::FUtils::MakePropertyBagForEditorData(ReferencedEditorData, ReferencedEditorData->GetVariableCompileContext(InSettings, CompileContext));
 		AddInstancedPropertyBag(AssetBag);
 	}
 
diff --git a/Engine/Plugins/Experimental/UAF/UAF/Source/UAFUncookedOnly/Private/Compilation/AnimNextGetFunctionHeaderCompileContext.cpp b/Engine/Plugins/Experimental/UAF/UAF/Source/UAFUncookedOnly/Private/Compilation/AnimNextGetFunctionHeaderCompileContext.cpp
--- a/Engine/Plugins/Experimental/UAF/UAF/Source/UAFUncookedOnly/Private/Compilation/AnimNextGetFunctionHeaderCompileContext.cpp
+++ b/Engine/Plugins/Experimental/UAF/UAF/Source/UAFUncookedOnly/Private/Compilation/AnimNextGetFunctionHeaderCompileContext.cpp
@@ -12,7 +12,12 @@ const TArray<FAnimNextProgrammaticFunctionHeader>& FAnimNextGetFunctionHeaderCom
 	return CompilerContext.FunctionHeaders;
 }
 
-TArray<FAnimNextProgrammaticFunctionHeader>& FAnimNextGetFunctionHeaderCompileContext::GetMutableFunctionHeaders()
+void FAnimNextGetFunctionHeaderCompileContext::AddUniqueFunctionHeader(const FAnimNextProgrammaticFunctionHeader& InFunctionHeader)
 {
-	return CompilerContext.FunctionHeaders;
+	CompilerContext.FunctionHeaders.AddUnique(InFunctionHeader);
+}
+
+const UAnimNextRigVMAssetEditorData* FAnimNextGetFunctionHeaderCompileContext::GetOwningAssetEditorData() const
+{
+	return CompilerContext.OwningAssetEditorData;
 }
diff --git a/Engine/Plugins/Experimental/UAF/UAF/Source/UAFUncookedOnly/Private/Compilation/AnimNextGetGraphCompileContext.cpp b/Engine/Plugins/Experimental/UAF/UAF/Source/UAFUncookedOnly/Private/Compilation/AnimNextGetGraphCompileContext.cpp
--- a/Engine/Plugins/Experimental/UAF/UAF/Source/UAFUncookedOnly/Private/Compilation/AnimNextGetGraphCompileContext.cpp
+++ b/Engine/Plugins/Experimental/UAF/UAF/Source/UAFUncookedOnly/Private/Compilation/AnimNextGetGraphCompileContext.cpp
@@ -21,3 +21,8 @@ TArray<URigVMGraph*>& FAnimNextGetGraphCompileContext::GetMutableProgrammaticGra
 {
 	return CompilerContext.ProgrammaticGraphs;
 }
+
+const UAnimNextRigVMAssetEditorData* FAnimNextGetGraphCompileContext::GetOwningAssetEditorData() const
+{
+	return CompilerContext.OwningAssetEditorData;
+}
diff --git a/Engine/Plugins/Experimental/UAF/UAF/Source/UAFUncookedOnly/Private/Compilation/AnimNextGetVariableCompileContext.cpp b/Engine/Plugins/Experimental/UAF/UAF/Source/UAFUncookedOnly/Private/Compilation/AnimNextGetVariableCompileContext.cpp
--- a/Engine/Plugins/Experimental/UAF/UAF/Source/UAFUncookedOnly/Private/Compilation/AnimNextGetVariableCompileContext.cpp
+++ b/Engine/Plugins/Experimental/UAF/UAF/Source/UAFUncookedOnly/Private/Compilation/AnimNextGetVariableCompileContext.cpp
@@ -21,3 +21,8 @@ TArray<FAnimNextProgrammaticVariable>& FAnimNextGetVariableCompileContext::GetMu
 {
 	return CompilerContext.ProgrammaticVariables;
 }
+
+const UAnimNextRigVMAssetEditorData* FAnimNextGetVariableCompileContext::GetOwningAssetEditorData() const
+{
+	return CompilerContext.OwningAssetEditorData;
+}
diff --git a/Engine/Plugins/Experimental/UAF/UAF/Source/UAFUncookedOnly/Private/Compilation/AnimNextProcessGraphCompileContext.cpp b/Engine/Plugins/Experimental/UAF/UAF/Source/UAFUncookedOnly/Private/Compilation/AnimNextProcessGraphCompileContext.cpp
--- a/Engine/Plugins/Experimental/UAF/UAF/Source/UAFUncookedOnly/Private/Compilation/AnimNextProcessGraphCompileContext.cpp
+++ b/Engine/Plugins/Experimental/UAF/UAF/Source/UAFUncookedOnly/Private/Compilation/AnimNextProcessGraphCompileContext.cpp
@@ -16,3 +16,8 @@ TArray<URigVMGraph*>& FAnimNextProcessGraphCompileContext::GetMutableAllGraphs()
 {
 	return CompilerContext.AllGraphs;
 }
+
+const UAnimNextRigVMAssetEditorData* FAnimNextProcessGraphCompileContext::GetOwningAssetEditorData() const
+{
+	return CompilerContext.OwningAssetEditorData;
+}
diff --git a/Engine/Plugins/Experimental/UAF/UAF/Source/UAFUncookedOnly/Private/UncookedOnlyUtils.cpp b/Engine/Plugins/Experimental/UAF/UAF/Source/UAFUncookedOnly/Private/UncookedOnlyUtils.cpp
--- a/Engine/Plugins/Experimental/UAF/UAF/Source/UAFUncookedOnly/Private/UncookedOnlyUtils.cpp
+++ b/Engine/Plugins/Experimental/UAF/UAF/Source/UAFUncookedOnly/Private/UncookedOnlyUtils.cpp
@@ -59,6 +59,8 @@ void FUtils::RecreateVM(UAnimNextRigVMAsset* InAsset)
 
 FInstancedPropertyBag FUtils::MakePropertyBagForEditorData(const UAnimNextRigVMAssetEditorData* InEditorData, const FAnimNextGetVariableCompileContext& InCompileContext)
 {
+	ensureMsgf(InEditorData == InCompileContext.GetOwningAssetEditorData(), TEXT("Crossing RigVM asset with a different asset's compile context. Expect incorrect memory layout / variable generation"));
+
 	FMessageLog Log("AnimNextCompilerResults");
 	
 	struct FStructEntryInfo
@@ -262,7 +264,7 @@ void FUtils::CompileVariables(const FRigVMCompileSettings& InSettings, UAnimNext
 	}
 
 	// Now rebuild combined the combined property bag we used for stable properties 
-	InAsset->CombinedPropertyBag = EditorData->GenerateCombinedPropertyBag(OutCompileContext);
+	InAsset->CombinedPropertyBag = EditorData->GenerateCombinedPropertyBag(InSettings, OutCompileContext);
 
 	// Rebuild external variables
 	if (InAsset->CombinedPropertyBag.GetNumPropertiesInBag() > 0)
diff --git a/Engine/Plugins/Experimental/UAF/UAFAnimGraph/Source/UAFAnimGraphUncookedOnly/Internal/Graph/AnimNextAnimationGraph_EditorData.h b/Engine/Plugins/Experimental/UAF/UAFAnimGraph/Source/UAFAnimGraphUncookedOnly/Internal/Graph/AnimNextAnimationGraph_EditorData.h
--- a/Engine/Plugins/Experimental/UAF/UAFAnimGraph/Source/UAFAnimGraphUncookedOnly/Internal/Graph/AnimNextAnimationGraph_EditorData.h
+++ b/Engine/Plugins/Experimental/UAF/UAFAnimGraph/Source/UAFAnimGraphUncookedOnly/Internal/Graph/AnimNextAnimationGraph_EditorData.h
@@ -75,7 +75,7 @@ protected:
 	UE_API virtual void InitializeAssetUserData() override;
 	UE_API virtual void OnCompileJobStarted() override;
 	UE_API virtual void OnPreCompileAsset(FRigVMCompileSettings& InSettings) override;
-	UE_API virtual void OnPreCompileGetProgrammaticFunctionHeaders(const FRigVMCompileSettings& InSettings, FAnimNextGetFunctionHeaderCompileContext& OutCompileContext) override;
+	UE_API virtual void BuildFunctionHeadersContext(const FRigVMCompileSettings& InSettings, FAnimNextGetFunctionHeaderCompileContext& OutCompileContext) const override;
 	UE_API virtual void OnPreCompileProcessGraphs(const FRigVMCompileSettings& InSettings, FAnimNextProcessGraphCompileContext& OutCompileContext) override;
 	UE_API virtual void OnCompileJobFinished() override;
 	UE_API virtual UClass* GetRigVMEdGraphSchemaClass() const override;
diff --git a/Engine/Plugins/Experimental/UAF/UAFAnimGraph/Source/UAFAnimGraphUncookedOnly/Private/Graph/AnimNextAnimationGraph_EditorData.cpp b/Engine/Plugins/Experimental/UAF/UAFAnimGraph/Source/UAFAnimGraphUncookedOnly/Private/Graph/AnimNextAnimationGraph_EditorData.cpp
--- a/Engine/Plugins/Experimental/UAF/UAFAnimGraph/Source/UAFAnimGraphUncookedOnly/Private/Graph/AnimNextAnimationGraph_EditorData.cpp
+++ b/Engine/Plugins/Experimental/UAF/UAFAnimGraph/Source/UAFAnimGraphUncookedOnly/Private/Graph/AnimNextAnimationGraph_EditorData.cpp
@@ -609,16 +609,16 @@ void UAnimNextAnimationGraph_EditorData::OnPreCompileAsset(FRigVMCompileSettings
 	AnimationGraph->DefaultEntryPoint = NAME_None;
 }
 
-void UAnimNextAnimationGraph_EditorData::OnPreCompileGetProgrammaticFunctionHeaders(const FRigVMCompileSettings& InSettings, FAnimNextGetFunctionHeaderCompileContext& OutCompileContext)
+void UAnimNextAnimationGraph_EditorData::BuildFunctionHeadersContext(const FRigVMCompileSettings& InSettings, FAnimNextGetFunctionHeaderCompileContext& OutCompileContext) const
 {
 	using namespace UE::UAF::UncookedOnly;
 	using namespace UE::UAF::UncookedOnly::Private;
 
-	Super::OnPreCompileGetProgrammaticFunctionHeaders(InSettings, OutCompileContext);
+	Super::BuildFunctionHeadersContext(InSettings, OutCompileContext);
 
 	// Gather all 'call function' traits and create shim-calls for them.
 	// For the compiler to pick them up if they are not public we need a calling reference to the function from a graph
-	FRigVMClient* VMClient = GetRigVMClient();
+	const FRigVMClient* VMClient = GetRigVMClient();
 	for (const URigVMGraph* Graph : VMClient->GetAllModels(false, false))
 	{
 		for(const URigVMNode* Node : Graph->GetNodes())
@@ -659,7 +659,7 @@ void UAnimNextAnimationGraph_EditorData::OnPreCompileGetProgrammaticFunctionHead
 				{
 					FAnimNextProgrammaticFunctionHeader AnimNextFunctionHeader = {};
 					AnimNextFunctionHeader.FunctionHeader = FunctionHeader;
-					OutCompileContext.GetMutableFunctionHeaders().Add(AnimNextFunctionHeader);
+					OutCompileContext.AddUniqueFunctionHeader(AnimNextFunctionHeader);
 				}
 			}
 		}
diff --git a/Engine/Plugins/Experimental/UAF/UAFStateTree/Source/UAFStateTree/Internal/AnimNextStateTreeEditorOnlyTypes.h b/Engine/Plugins/Experimental/UAF/UAFStateTree/Source/UAFStateTree/Internal/AnimNextStateTreeEditorOnlyTypes.h
--- a/Engine/Plugins/Experimental/UAF/UAFStateTree/Source/UAFStateTree/Internal/AnimNextStateTreeEditorOnlyTypes.h
+++ b/Engine/Plugins/Experimental/UAF/UAFStateTree/Source/UAFStateTree/Internal/AnimNextStateTreeEditorOnlyTypes.h
@@ -18,7 +18,7 @@ struct FRigVMClient;
  */
 struct UAFSTATETREE_API FAnimNextStateTreeProgrammaticFunctionHeaderParams
 {
-	FAnimNextStateTreeProgrammaticFunctionHeaderParams(UAnimNextSharedVariables_EditorData* InEditorData, const FRigVMCompileSettings& InSettings, FRigVMClient& InRigVMClient, FAnimNextGetFunctionHeaderCompileContext& InOutCompileContext)
+	FAnimNextStateTreeProgrammaticFunctionHeaderParams(const UAnimNextSharedVariables_EditorData* InEditorData, const FRigVMCompileSettings& InSettings, const FRigVMClient& InRigVMClient, FAnimNextGetFunctionHeaderCompileContext& InOutCompileContext)
 		: EditorData(InEditorData)
 		, Settings(InSettings)
 		, RigVMClient(InRigVMClient)
@@ -26,9 +26,9 @@ struct UAFSTATETREE_API FAnimNextStateTreeProgrammaticFunctionHeaderParams
 	{
 	}
 
-	UAnimNextSharedVariables_EditorData* EditorData;
+	const UAnimNextSharedVariables_EditorData* EditorData;
 	const FRigVMCompileSettings& Settings;
-	FRigVMClient& RigVMClient;
+	const FRigVMClient& RigVMClient;
 	FAnimNextGetFunctionHeaderCompileContext& OutCompileContext;
 };
 
diff --git a/Engine/Plugins/Experimental/UAF/UAFStateTree/Source/UAFStateTree/Private/Conditions/AnimNextStateTreeRigVMConditionBase.cpp b/Engine/Plugins/Experimental/UAF/UAFStateTree/Source/UAFStateTree/Private/Conditions/AnimNextStateTreeRigVMConditionBase.cpp
--- a/Engine/Plugins/Experimental/UAF/UAFStateTree/Source/UAFStateTree/Private/Conditions/AnimNextStateTreeRigVMConditionBase.cpp
+++ b/Engine/Plugins/Experimental/UAF/UAFStateTree/Source/UAFStateTree/Private/Conditions/AnimNextStateTreeRigVMConditionBase.cpp
@@ -154,7 +154,7 @@ void FAnimNextStateTreeRigVMConditionBase::GetProgrammaticFunctionHeaders(FAnimN
 
 	FAnimNextProgrammaticFunctionHeader AnimNextFunctionHeader = {};
 	AnimNextFunctionHeader.FunctionHeader = RigVMFunctionHeader;
-	OutCompileContext.GetMutableFunctionHeaders().Add(AnimNextFunctionHeader);
+	OutCompileContext.AddUniqueFunctionHeader(AnimNextFunctionHeader);
 }
 
 #endif // WITH_EDITOR
diff --git a/Engine/Plugins/Experimental/UAF/UAFStateTree/Source/UAFStateTree/Private/Tasks/AnimNextStateTreeRigVMTaskBase.cpp b/Engine/Plugins/Experimental/UAF/UAFStateTree/Source/UAFStateTree/Private/Tasks/AnimNextStateTreeRigVMTaskBase.cpp
--- a/Engine/Plugins/Experimental/UAF/UAFStateTree/Source/UAFStateTree/Private/Tasks/AnimNextStateTreeRigVMTaskBase.cpp
+++ b/Engine/Plugins/Experimental/UAF/UAFStateTree/Source/UAFStateTree/Private/Tasks/AnimNextStateTreeRigVMTaskBase.cpp
@@ -131,9 +131,9 @@ void FAnimNextStateTreeRigVMTaskBase::PostEditNodeChangeChainProperty(const FPro
 					FString CppTypeString = Argument.CPPType.ToString();
 					UObject* CppTypeObject = Argument.CPPTypeObject.Get();
 					FPropertyBagPropertyDesc PropertyDesc = FRigVMMemoryStorageStruct::GeneratePropertyBagDescriptor(FRigVMPropertyDescription(Argument.Name, CppTypeString, CppTypeObject, Argument.DefaultValue));
-					if (Argument.Direction == ERigVMPinDirection::Output && Argument.CPPType == RigVMTypeUtils::BoolTypeName)
+					if (Argument.Direction == ERigVMPinDirection::Output)
 					{
-						PropertyDesc.PropertyFlags &= ~CPF_Edit;	// Hide the result from the UI
+						PropertyDesc.PropertyFlags &= ~CPF_Edit;	// Hide the results from the UI
 						ResultIndex = PropertyDescs.Num();
 					}
 					PropertyDescs.Add(PropertyDesc);
@@ -157,7 +157,7 @@ void FAnimNextStateTreeRigVMTaskBase::GetProgrammaticFunctionHeaders(FAnimNextSt
 
 	FAnimNextProgrammaticFunctionHeader AnimNextFunctionHeader = {};
 	AnimNextFunctionHeader.FunctionHeader = RigVMFunctionHeader;
-	OutCompileContext.GetMutableFunctionHeaders().Add(AnimNextFunctionHeader);
+	OutCompileContext.AddUniqueFunctionHeader(AnimNextFunctionHeader);
 }
 
 #endif // WITH_EDITOR
diff --git a/Engine/Plugins/Experimental/UAF/UAFStateTree/Source/UAFStateTreeUncookedOnly/Internal/AnimNextStateTree_EditorData.h b/Engine/Plugins/Experimental/UAF/UAFStateTree/Source/UAFStateTreeUncookedOnly/Internal/AnimNextStateTree_EditorData.h
--- a/Engine/Plugins/Experimental/UAF/UAFStateTree/Source/UAFStateTreeUncookedOnly/Internal/AnimNextStateTree_EditorData.h
+++ b/Engine/Plugins/Experimental/UAF/UAFStateTree/Source/UAFStateTreeUncookedOnly/Internal/AnimNextStateTree_EditorData.h
@@ -22,7 +22,7 @@ protected:
 	
 	// UAnimNextRigVMAssetEditorData interface
 	virtual TConstArrayView<TSubclassOf<UAnimNextRigVMAssetEntry>> GetEntryClasses() const override;
-	virtual void OnPreCompileGetProgrammaticFunctionHeaders(const FRigVMCompileSettings& InSettings, FAnimNextGetFunctionHeaderCompileContext& OutCompileContext) override;
+	virtual void BuildFunctionHeadersContext(const FRigVMCompileSettings& InSettings, FAnimNextGetFunctionHeaderCompileContext& OutCompileContext) const override;
 	virtual void OnPreCompileGetProgrammaticGraphs(const FRigVMCompileSettings& InSettings, FAnimNextGetGraphCompileContext& OutCompileContext) override;
 	virtual void OnPostCompileVariables(const FRigVMCompileSettings& InSettings, const FAnimNextGetVariableCompileContext& InCompileContext) override;
 
diff --git a/Engine/Plugins/Experimental/UAF/UAFStateTree/Source/UAFStateTreeUncookedOnly/Private/AnimNextStateTree_EditorData.cpp b/Engine/Plugins/Experimental/UAF/UAFStateTree/Source/UAFStateTreeUncookedOnly/Private/AnimNextStateTree_EditorData.cpp
--- a/Engine/Plugins/Experimental/UAF/UAFStateTree/Source/UAFStateTreeUncookedOnly/Private/AnimNextStateTree_EditorData.cpp
+++ b/Engine/Plugins/Experimental/UAF/UAFStateTree/Source/UAFStateTreeUncookedOnly/Private/AnimNextStateTree_EditorData.cpp
@@ -68,9 +68,9 @@ TConstArrayView<TSubclassOf<UAnimNextRigVMAssetEntry>> UAnimNextStateTree_Editor
 	return Classes;
 }
 
-void UAnimNextStateTree_EditorData::OnPreCompileGetProgrammaticFunctionHeaders(const FRigVMCompileSettings& InSettings, FAnimNextGetFunctionHeaderCompileContext& OutCompileContext)
+void UAnimNextStateTree_EditorData::BuildFunctionHeadersContext(const FRigVMCompileSettings& InSettings, FAnimNextGetFunctionHeaderCompileContext& OutCompileContext) const
 {
-	Super::OnPreCompileGetProgrammaticFunctionHeaders(InSettings, OutCompileContext);
+	Super::BuildFunctionHeadersContext(InSettings, OutCompileContext);
 
 	if (UAnimNextStateTree* AnimStateTree = UE::UAF::UncookedOnly::FUtils::GetAsset<UAnimNextStateTree>(this))
 	{
@@ -222,5 +222,5 @@ void UAnimNextStateTree_EditorData::OnPostCompileVariables(const FRigVMCompileSe
 {
 	Super::OnPostCompileVariables(InSettings, InCompileContext);
 
-	CombinedVariablesPropertyBag = GenerateCombinedPropertyBag(InCompileContext);
+	CombinedVariablesPropertyBag = GenerateCombinedPropertyBag(InSettings, InCompileContext);
 }
diff --git a/Engine/Plugins/Runtime/RigVM/Source/RigVM/Private/RigVMCore/RigVMGraphFunctionDefinition.cpp b/Engine/Plugins/Runtime/RigVM/Source/RigVM/Private/RigVMCore/RigVMGraphFunctionDefinition.cpp
--- a/Engine/Plugins/Runtime/RigVM/Source/RigVM/Private/RigVMCore/RigVMGraphFunctionDefinition.cpp
+++ b/Engine/Plugins/Runtime/RigVM/Source/RigVM/Private/RigVMCore/RigVMGraphFunctionDefinition.cpp
@@ -38,6 +38,11 @@ FRigVMExternalVariable FRigVMGraphFunctionArgument::GetExternalVariable() const
 	return Variable;
 }
 
+bool FRigVMGraphFunctionArgument::IsValid() const
+{
+	return !CPPType.IsNone() || IsCPPTypeObjectValid();
+}
+
 bool FRigVMGraphFunctionArgument::IsCPPTypeObjectValid() const
 {
 	if(!CPPTypeObject.IsValid())
diff --git a/Engine/Plugins/Runtime/RigVM/Source/RigVM/Public/RigVMCore/RigVMGraphFunctionDefinition.h b/Engine/Plugins/Runtime/RigVM/Source/RigVM/Public/RigVMCore/RigVMGraphFunctionDefinition.h
--- a/Engine/Plugins/Runtime/RigVM/Source/RigVM/Public/RigVMCore/RigVMGraphFunctionDefinition.h
+++ b/Engine/Plugins/Runtime/RigVM/Source/RigVM/Public/RigVMCore/RigVMGraphFunctionDefinition.h
@@ -264,6 +264,9 @@ struct FRigVMGraphFunctionArgument
 
 	UE_API FRigVMExternalVariable GetExternalVariable() const;
 
+	// True if CPP Type or obj type is valid. May potentially load the CPP Type Object
+	UE_API bool IsValid() const;
+
 	// validates and potentially loads the CPP Type Object
 	UE_API bool IsCPPTypeObjectValid() const;
 
