Fix build issue in AnimNextModuleTrack.h from missing imports

#rnx

diff --git a/Engine/Plugins/Experimental/UAF/UAF/Source/UAFEditor/Private/RewindDebugger/AnimNextModuleTrack.h b/Engine/Plugins/Experimental/UAF/UAF/Source/UAFEditor/Private/RewindDebugger/AnimNextModuleTrack.h
--- a/Engine/Plugins/Experimental/UAF/UAF/Source/UAFEditor/Private/RewindDebugger/AnimNextModuleTrack.h
+++ b/Engine/Plugins/Experimental/UAF/UAF/Source/UAFEditor/Private/RewindDebugger/AnimNextModuleTrack.h
@@ -8,9 +8,9 @@
 #include "RewindDebuggerTrack.h"
 #include "Textures/SlateIcon.h"
 #include "SEventTimelineView.h"
-#include "StructUtils/PropertyBag.h"
 #include "IDetailsView.h"
-#endif // ANIMNEXT_TRACE_ENABLED 
+#endif // ANIMNEXT_TRACE_ENABLED
+#include "StructUtils/PropertyBag.h"
 
 #include "AnimNextModuleTrack.generated.h"
 
