Fixed case where hitting enter, while in text edit mode, on a newly added variable prompts user with "Variable name already exists in this asset'

diff --git a/Engine/Plugins/Experimental/UAF/UAF/Source/UAFEditor/Private/Variables/Outliner/VariablesOutlinerEntryItem.cpp b/Engine/Plugins/Experimental/UAF/UAF/Source/UAFEditor/Private/Variables/Outliner/VariablesOutlinerEntryItem.cpp
--- a/Engine/Plugins/Experimental/UAF/UAF/Source/UAFEditor/Private/Variables/Outliner/VariablesOutlinerEntryItem.cpp
+++ b/Engine/Plugins/Experimental/UAF/UAF/Source/UAFEditor/Private/Variables/Outliner/VariablesOutlinerEntryItem.cpp
@@ -429,7 +429,9 @@ bool FVariablesOutlinerEntryItem::ValidateName(const FText& InNewName, FText& Ou
 	}
 
 	FName Name(*NewString);
-	if(UAnimNextRigVMAssetEntry* ExistingEntry = EditorData->FindEntry(Name))
+	// Cannot have two entries with matching names
+	UAnimNextRigVMAssetEntry* ExistingEntry = EditorData->FindEntry(Name);
+	if(ExistingEntry && ExistingEntry != Entry)
 	{
 		OutErrorMessage = LOCTEXT("NameExistsError", "Variable name already exists in this asset");
 		return false;
