Fix asset wizard not resetting correctly after consecutive uses

#rnx

diff --git a/Engine/Plugins/Experimental/UAF/UAF/Source/UAFEditor/Private/Wizard/STemplatePicker.cpp b/Engine/Plugins/Experimental/UAF/UAF/Source/UAFEditor/Private/Wizard/STemplatePicker.cpp
--- a/Engine/Plugins/Experimental/UAF/UAF/Source/UAFEditor/Private/Wizard/STemplatePicker.cpp
+++ b/Engine/Plugins/Experimental/UAF/UAF/Source/UAFEditor/Private/Wizard/STemplatePicker.cpp
@@ -65,6 +65,7 @@ namespace UE::UAF::Editor
 			];
 
 		TemplateConfig = GetMutableDefault<UUAFTemplateConfig>();
+		TemplateConfig->Reset();
 		TemplateConfig->OnOutputPathChanged.AddSP(this, &STemplatePicker::HandleOutputPathChanged);
 		
 		FDetailsViewArgs DetailsViewArgs;
diff --git a/Engine/Plugins/Experimental/UAF/UAF/Source/UAFEditor/Private/Wizard/TemplateConfig.cpp b/Engine/Plugins/Experimental/UAF/UAF/Source/UAFEditor/Private/Wizard/TemplateConfig.cpp
--- a/Engine/Plugins/Experimental/UAF/UAF/Source/UAFEditor/Private/Wizard/TemplateConfig.cpp
+++ b/Engine/Plugins/Experimental/UAF/UAF/Source/UAFEditor/Private/Wizard/TemplateConfig.cpp
@@ -13,3 +13,9 @@ void UUAFTemplateConfig::PostEditChangeProperty(FPropertyChangedEvent& PropertyC
 		OnOutputPathChanged.Broadcast();
 	}
 }
+
+void UUAFTemplateConfig::Reset()
+{
+	OnOutputPathChanged.Clear();
+	AssetNaming.Empty();
+}
\ No newline at end of file
diff --git a/Engine/Plugins/Experimental/UAF/UAF/Source/UAFEditor/Private/Wizard/TemplateConfig.h b/Engine/Plugins/Experimental/UAF/UAF/Source/UAFEditor/Private/Wizard/TemplateConfig.h
--- a/Engine/Plugins/Experimental/UAF/UAF/Source/UAFEditor/Private/Wizard/TemplateConfig.h
+++ b/Engine/Plugins/Experimental/UAF/UAF/Source/UAFEditor/Private/Wizard/TemplateConfig.h
@@ -49,6 +49,8 @@ public:
 
 	FSimpleMulticastDelegate OnOutputPathChanged;
 
+	void Reset();
+
 public:
 	TArray<FString> AssetNaming;
 };
\ No newline at end of file
