Remove check when adding a bone to a set binding
* AddBoneToSet returns false if a bone already exists in a set

#rnx

diff --git a/Engine/Plugins/Experimental/UAF/UAF/Source/UAFEditor/Private/UAF/AbstractSkeleton/Sets/SSetBinding.cpp b/Engine/Plugins/Experimental/UAF/UAF/Source/UAFEditor/Private/UAF/AbstractSkeleton/Sets/SSetBinding.cpp
--- a/Engine/Plugins/Experimental/UAF/UAF/Source/UAFEditor/Private/UAF/AbstractSkeleton/Sets/SSetBinding.cpp
+++ b/Engine/Plugins/Experimental/UAF/UAF/Source/UAFEditor/Private/UAF/AbstractSkeleton/Sets/SSetBinding.cpp
@@ -179,8 +179,10 @@ void SSetBinding::Construct(const FArguments& InArgs, TWeakObjectPtr<UAbstractSk
 									SetBinding->Modify();
 
 									const bool bSuccess = SetBinding->AddBoneToSet(DragDropOp->Item->BoneName, TargetSetTreeItem->Name);
-									check(bSuccess);
-									RepopulateTreeData();
+									if (bSuccess)
+									{
+										RepopulateTreeData();
+									}
 
 									return FReply::Handled();
 								}
