- Added compile-process error/warning/info/message API to AnimNextRigVMAssetCompileContext.
* This routes messages through EditorData::HandleReportFromCompiler/HandleMessageFromCompiler, matching RigVM compiler messages
* Replaced all uses of FMessageLog with new API
- Added transient CompilationState property to AnimNextRigVMAsset, written by FScopedCompileJob
* Added early-outs according to valid Object/Module and Compilation state to AnimNextModuleInstance to prevent initialization using invalid data

diff --git a/Engine/Plugins/Experimental/UAF/UAF/Source/UAF/Private/Module/AnimNextModuleInstance.cpp b/Engine/Plugins/Experimental/UAF/UAF/Source/UAF/Private/Module/AnimNextModuleInstance.cpp
--- a/Engine/Plugins/Experimental/UAF/UAF/Source/UAF/Private/Module/AnimNextModuleInstance.cpp
+++ b/Engine/Plugins/Experimental/UAF/UAF/Source/UAF/Private/Module/AnimNextModuleInstance.cpp
@@ -173,9 +173,28 @@ void FAnimNextModuleInstance::Initialize()
 
 	check(IsInGameThread());
 
-	check(Object)
+	check(Object);
+	if (!Object)
+	{
+		return;
+	}
+	
 	const UAnimNextModule* Module = GetModule();
+
 	check(Module);
+	if (Module)
+	{
+#if WITH_EDITOR
+		if (Module->CompilationState == EAnimNextRigVMAssetState::CompiledWithErrors)
+		{
+			return;
+		}
+#endif
+	}
+	else
+	{
+		return;
+	}
 
 	UWorld* World = Object->GetWorld();
 	if(World)
diff --git a/Engine/Plugins/Experimental/UAF/UAF/Source/UAF/Public/AnimNextRigVMAsset.h b/Engine/Plugins/Experimental/UAF/UAF/Source/UAF/Public/AnimNextRigVMAsset.h
--- a/Engine/Plugins/Experimental/UAF/UAF/Source/UAF/Public/AnimNextRigVMAsset.h
+++ b/Engine/Plugins/Experimental/UAF/UAF/Source/UAF/Public/AnimNextRigVMAsset.h
@@ -34,8 +34,19 @@ namespace UE::UAF
 namespace UE::UAF::UncookedOnly
 {
 	struct FUtils;
+	class FScopedCompileJob;
 }
 
+#if WITH_EDITORONLY_DATA
+UENUM()
+enum class EAnimNextRigVMAssetState : uint8
+{
+	CompiledWithErrors,
+	CompiledWithWarnings,
+	CompiledWithSuccess
+};
+#endif // WITH_EDITORONLY_DATA
+
 /** Base class for all AnimNext assets that can host RigVM logic */
 UCLASS(MinimalAPI, Abstract)
 class UAnimNextRigVMAsset : public URigVMHost
@@ -52,6 +63,7 @@ protected:
 	friend FUAFInstanceVariableData;
 	friend UE::UAF::FInjectionInfo;
 	friend UE::UAF::FAnimGraphBuilderContext;
+	friend UE::UAF::UncookedOnly::FScopedCompileJob;
 
 	UE_API UAnimNextRigVMAsset(const FObjectInitializer& ObjectInitializer);
 
@@ -132,6 +144,9 @@ protected:
 #if WITH_EDITORONLY_DATA
 	UPROPERTY(VisibleAnywhere, Instanced, Category = "Editor Data", meta = (ShowInnerProperties))
 	TObjectPtr<UObject> EditorData;
+
+	UPROPERTY(transient)
+	EAnimNextRigVMAssetState CompilationState;
 #endif
 };
 
diff --git a/Engine/Plugins/Experimental/UAF/UAF/Source/UAFUncookedOnly/Internal/AnimNextRigVMAssetEditorData.h b/Engine/Plugins/Experimental/UAF/UAF/Source/UAFUncookedOnly/Internal/AnimNextRigVMAssetEditorData.h
--- a/Engine/Plugins/Experimental/UAF/UAF/Source/UAFUncookedOnly/Internal/AnimNextRigVMAssetEditorData.h
+++ b/Engine/Plugins/Experimental/UAF/UAF/Source/UAFUncookedOnly/Internal/AnimNextRigVMAssetEditorData.h
@@ -250,6 +250,7 @@ protected:
 	friend struct FAnimNextSchemaAction_Variable;
 	friend struct FAnimNextModuleTests;
 	friend struct FAnimNextAnimationGraphTests;
+	friend FAnimNextRigVMAssetCompileContext;
 
 	// UObject interface
 	UE_API virtual void Serialize(FArchive& Ar) override;
@@ -578,7 +579,8 @@ protected:
 	UE_API void ClearErrorInfoForAllEdGraphs();
 
 	// Handle compiler reporting
-	UE_API void HandleReportFromCompiler(EMessageSeverity::Type InSeverity, UObject* InSubject, const FString& InMessage);
+	UE_API void HandleReportFromCompiler(EMessageSeverity::Type InSeverity, UObject* InSubject, const FString& InMessage) const;
+	UE_API void HandleMessageFromCompiler(TSharedRef<FTokenizedMessage> InMessage, bool bAddMessageToLog = true) const;
 
 	// Support extra references in GC
 	static UE_API void AddReferencedObjects(UObject* InThis, FReferenceCollector& Collector);
@@ -712,8 +714,8 @@ protected:
 	TObjectPtr<UAnimNextEdGraph> FunctionLibraryEdGraph;
 
 	bool bAutoRecompileVM = true;
-	bool bErrorsDuringCompilation = false;
-	bool bWarningsDuringCompilation = false;
+	mutable bool bErrorsDuringCompilation = false;
+	mutable bool bWarningsDuringCompilation = false;
 	bool bSuspendModelNotificationsForSelf = false;
 	bool bSuspendAllNotifications = false;
 	bool bCompileInDebugMode = false;
diff --git a/Engine/Plugins/Experimental/UAF/UAF/Source/UAFUncookedOnly/Internal/Compilation/AnimNextGetVariableCompileContext.h b/Engine/Plugins/Experimental/UAF/UAF/Source/UAFUncookedOnly/Internal/Compilation/AnimNextGetVariableCompileContext.h
--- a/Engine/Plugins/Experimental/UAF/UAF/Source/UAFUncookedOnly/Internal/Compilation/AnimNextGetVariableCompileContext.h
+++ b/Engine/Plugins/Experimental/UAF/UAF/Source/UAFUncookedOnly/Internal/Compilation/AnimNextGetVariableCompileContext.h
@@ -22,6 +22,8 @@ public:
 	UE_API const TArray<FAnimNextProgrammaticVariable>& GetProgrammaticVariables() const;
 	UE_API TArray<FAnimNextProgrammaticVariable>& GetMutableProgrammaticVariables();
 	UE_API const UAnimNextRigVMAssetEditorData* GetOwningAssetEditorData() const;
+	
+	UE_API const FAnimNextRigVMAssetCompileContext& GetAssetCompileContext() const;
 
 protected:
 
diff --git a/Engine/Plugins/Experimental/UAF/UAF/Source/UAFUncookedOnly/Internal/Compilation/AnimNextRigVMAssetCompileContext.h b/Engine/Plugins/Experimental/UAF/UAF/Source/UAFUncookedOnly/Internal/Compilation/AnimNextRigVMAssetCompileContext.h
--- a/Engine/Plugins/Experimental/UAF/UAF/Source/UAFUncookedOnly/Internal/Compilation/AnimNextRigVMAssetCompileContext.h
+++ b/Engine/Plugins/Experimental/UAF/UAF/Source/UAFUncookedOnly/Internal/Compilation/AnimNextRigVMAssetCompileContext.h
@@ -9,6 +9,8 @@
 #include "Compilation/AnimNextProgrammaticFunctionHeader.h"
 #include "Variables/AnimNextProgrammaticVariable.h"
 
+#define UE_API UAFUNCOOKEDONLY_API
+
 class URigVMGraph;
 class UAnimNextRigVMAssetEditorData;
 
@@ -27,6 +29,49 @@ public:
 	{
 	}
 
+#if WITH_EDITOR
+	template<typename... ArgTypes>
+	void Error(UObject* Object, FTextFormat Format, ArgTypes... Args) const
+	{
+		Message(EMessageSeverity::Error, Object, FText::Format(Format, Args...));
+	}
+
+	template<typename... ArgTypes>
+	void Error(FTextFormat Format, ArgTypes... Args) const
+	{
+		Message(EMessageSeverity::Error, OwningAssetEditorData, FText::Format(Format, Args...));
+	}
+
+	template<typename... ArgTypes>
+	void Warning(UObject* Object, FTextFormat Format, ArgTypes... Args) const
+	{
+		Message(EMessageSeverity::Warning, Object, FText::Format(Format, Args...));
+	}
+
+	template<typename... ArgTypes>
+	void Warning(FTextFormat Format, ArgTypes... Args) const
+	{
+		Message(EMessageSeverity::Warning, OwningAssetEditorData, FText::Format(Format, Args...));
+	}
+
+	template<typename... ArgTypes>
+	void Info(UObject* Object, FTextFormat Format, ArgTypes... Args) const
+	{
+		Message(EMessageSeverity::Info, Object, FText::Format(Format, Args...));
+	}
+
+	template<typename... ArgTypes>
+	void Info(FTextFormat Format, ArgTypes... Args) const
+	{
+		Message(EMessageSeverity::Info, OwningAssetEditorData, FText::Format(Format, Args...));
+	}
+	
+	UE_API void Message(TSharedRef<FTokenizedMessage> InMessage) const;
+
+private:
+	UE_API void Message(EMessageSeverity::Type Severity, UObject* Object, const FText& Message) const;
+#endif // WITH_EDITOR
+
 private:
 	friend class UAnimNextRigVMAssetEditorData;
 	friend struct FAnimNextGetFunctionHeaderCompileContext;
@@ -52,3 +97,4 @@ protected:
 	const UAnimNextRigVMAssetEditorData* OwningAssetEditorData;
 };
 
+#undef UE_API // UAFUNCOOKEDONLY_API
\ No newline at end of file
diff --git a/Engine/Plugins/Experimental/UAF/UAF/Source/UAFUncookedOnly/Private/AnimNextRigVMAssetEditorData.cpp b/Engine/Plugins/Experimental/UAF/UAF/Source/UAFUncookedOnly/Private/AnimNextRigVMAssetEditorData.cpp
--- a/Engine/Plugins/Experimental/UAF/UAF/Source/UAFUncookedOnly/Private/AnimNextRigVMAssetEditorData.cpp
+++ b/Engine/Plugins/Experimental/UAF/UAF/Source/UAFUncookedOnly/Private/AnimNextRigVMAssetEditorData.cpp
@@ -2644,7 +2644,7 @@ void UAnimNextRigVMAssetEditorData::GetAllVariables(TArray<FVariableInfo>& OutVa
 	}
 }
 
-void UAnimNextRigVMAssetEditorData::HandleReportFromCompiler(EMessageSeverity::Type InSeverity, UObject* InSubject, const FString& InMessage)
+void UAnimNextRigVMAssetEditorData::HandleReportFromCompiler(EMessageSeverity::Type InSeverity, UObject* InSubject, const FString& InMessage) const
 {
 	if (bSuspendCompilerReports)
 	{
@@ -2689,13 +2689,6 @@ void UAnimNextRigVMAssetEditorData::HandleReportFromCompiler(EMessageSeverity::T
 			Message = Log.Error(*InMessage);
 		}
 
-		// see UnitTest "ControlRig.Basics.OrphanedPins" to learn why errors are suppressed this way
-		if (!VMCompileSettings.SurpressErrors)
-		{ 
-			FScriptExceptionHandler::Get().HandleException(ELogVerbosity::Error, *InMessage, *FString());
-		}
-
-		bErrorsDuringCompilation = true;
 	}
 	else if (InSeverity == EMessageSeverity::Warning)
 	{
@@ -2707,10 +2700,6 @@ void UAnimNextRigVMAssetEditorData::HandleReportFromCompiler(EMessageSeverity::T
 		{
 			Message = Log.Warning(*InMessage);
 		}
-
-		FScriptExceptionHandler::Get().HandleException(ELogVerbosity::Warning, *InMessage, *FString());
-
-		bWarningsDuringCompilation = true;
 	}
 	else
 	{
@@ -2722,8 +2711,6 @@ void UAnimNextRigVMAssetEditorData::HandleReportFromCompiler(EMessageSeverity::T
 		{
 			Message = Log.Note(*InMessage);
 		}
-
-		UE_LOG(LogAnimation, Display, TEXT("%s"), *InMessage);
 	}
 
 	if (URigVMEdGraphNode* EdGraphNode = Cast<URigVMEdGraphNode>(SubjectForMessage))
@@ -2739,6 +2726,54 @@ void UAnimNextRigVMAssetEditorData::HandleReportFromCompiler(EMessageSeverity::T
 
 		EdGraphNode->bHasCompilerMessage = EdGraphNode->ErrorType <= int32(EMessageSeverity::Info);
 	}
+
+	if (Message.IsValid())
+	{
+		HandleMessageFromCompiler(Message.ToSharedRef(), false);
+	}
+}
+
+void UAnimNextRigVMAssetEditorData::HandleMessageFromCompiler(TSharedRef<FTokenizedMessage> InMessage, bool bAddMessageToLog) const
+{
+	if (bSuspendCompilerReports)
+	{
+		return;
+	}
+
+	FCompilerResultsLog& Log = UE::UAF::UncookedOnly::FScopedCompileJob::GetLog();
+	
+	TSharedPtr<FTokenizedMessage> Message;
+	if (InMessage->GetSeverity() == EMessageSeverity::Error)
+	{
+		// see UnitTest "ControlRig.Basics.OrphanedPins" to learn why errors are suppressed this way
+		if (VMCompileSettings.SurpressErrors)
+		{
+			Log.bSilentMode = true;
+		}
+		
+		// see UnitTest "ControlRig.Basics.OrphanedPins" to learn why errors are suppressed this way
+		if (!VMCompileSettings.SurpressErrors)
+		{ 
+			FScriptExceptionHandler::Get().HandleException(ELogVerbosity::Error, *InMessage->ToText().ToString(), *FString());
+		}
+
+		bErrorsDuringCompilation = true;
+	}
+	else if (InMessage->GetSeverity() == EMessageSeverity::Warning)
+	{
+		FScriptExceptionHandler::Get().HandleException(ELogVerbosity::Warning, *InMessage->ToText().ToString(), *FString());
+
+		bWarningsDuringCompilation = true;
+	}
+	else
+	{
+		UE_LOG(LogAnimation, Display, TEXT("%s"), *InMessage->ToText().ToString());
+	}
+
+	if (bAddMessageToLog)
+	{		
+		Log.AddTokenizedMessage(InMessage);
+	}
 }
 
 void UAnimNextRigVMAssetEditorData::ClearErrorInfoForAllEdGraphs()
diff --git a/Engine/Plugins/Experimental/UAF/UAF/Source/UAFUncookedOnly/Private/AnimNextScopedCompileJob.cpp b/Engine/Plugins/Experimental/UAF/UAF/Source/UAFUncookedOnly/Private/AnimNextScopedCompileJob.cpp
--- a/Engine/Plugins/Experimental/UAF/UAF/Source/UAFUncookedOnly/Private/AnimNextScopedCompileJob.cpp
+++ b/Engine/Plugins/Experimental/UAF/UAF/Source/UAFUncookedOnly/Private/AnimNextScopedCompileJob.cpp
@@ -143,6 +143,21 @@ FScopedCompileJob::~FScopedCompileJob()
 					continue;
 				}
 
+				Asset->CompilationState = [&EditorData]()
+				{
+					if (EditorData->bErrorsDuringCompilation)
+					{
+						return EAnimNextRigVMAssetState::CompiledWithErrors;
+					}
+
+					if (EditorData->bWarningsDuringCompilation)
+					{
+						return EAnimNextRigVMAssetState::CompiledWithWarnings;
+					}
+
+					return EAnimNextRigVMAssetState::CompiledWithSuccess;
+				}();			
+
 				if (!EditorData->bSuspendCompilationNotifications)
 				{
 					UAnimNextRigVMAsset::OnCompileJobFinished().Broadcast(Asset);
diff --git a/Engine/Plugins/Experimental/UAF/UAF/Source/UAFUncookedOnly/Private/Compilation/AnimNextGetVariableCompileContext.cpp b/Engine/Plugins/Experimental/UAF/UAF/Source/UAFUncookedOnly/Private/Compilation/AnimNextGetVariableCompileContext.cpp
--- a/Engine/Plugins/Experimental/UAF/UAF/Source/UAFUncookedOnly/Private/Compilation/AnimNextGetVariableCompileContext.cpp
+++ b/Engine/Plugins/Experimental/UAF/UAF/Source/UAFUncookedOnly/Private/Compilation/AnimNextGetVariableCompileContext.cpp
@@ -26,3 +26,8 @@ const UAnimNextRigVMAssetEditorData* FAnimNextGetVariableCompileContext::GetOwni
 {
 	return CompilerContext.OwningAssetEditorData;
 }
+
+const FAnimNextRigVMAssetCompileContext& FAnimNextGetVariableCompileContext::GetAssetCompileContext() const
+{
+	return CompilerContext;
+}
diff --git a/Engine/Plugins/Experimental/UAF/UAF/Source/UAFUncookedOnly/Private/Compilation/AnimNextRigVMAssetCompileContext.cpp b/Engine/Plugins/Experimental/UAF/UAF/Source/UAFUncookedOnly/Private/Compilation/AnimNextRigVMAssetCompileContext.cpp
new file mode 100644
--- /dev/null
+++ b/Engine/Plugins/Experimental/UAF/UAF/Source/UAFUncookedOnly/Private/Compilation/AnimNextRigVMAssetCompileContext.cpp
@@ -0,0 +1,18 @@
+// Copyright Epic Games, Inc. All Rights Reserved.
+
+#include "Compilation/AnimNextRigVMAssetCompileContext.h"
+#include "AnimNextRigVMAssetEditorData.h"
+
+#if WITH_EDITOR
+void FAnimNextRigVMAssetCompileContext::Message(TSharedRef<FTokenizedMessage> InMessage) const
+{
+	check(OwningAssetEditorData);
+	OwningAssetEditorData->HandleMessageFromCompiler(InMessage);	
+}
+
+void FAnimNextRigVMAssetCompileContext::Message(EMessageSeverity::Type Severity, UObject* Object, const FText& Message) const
+{
+	check(OwningAssetEditorData);
+	OwningAssetEditorData->HandleReportFromCompiler(Severity, Object, Message.ToString());
+}
+#endif // WITH_EDITOR
diff --git a/Engine/Plugins/Experimental/UAF/UAF/Source/UAFUncookedOnly/Private/UncookedOnlyUtils.cpp b/Engine/Plugins/Experimental/UAF/UAF/Source/UAFUncookedOnly/Private/UncookedOnlyUtils.cpp
--- a/Engine/Plugins/Experimental/UAF/UAF/Source/UAFUncookedOnly/Private/UncookedOnlyUtils.cpp
+++ b/Engine/Plugins/Experimental/UAF/UAF/Source/UAFUncookedOnly/Private/UncookedOnlyUtils.cpp
@@ -35,6 +35,8 @@
 #include "AnimNextExports.h"
 #include "AnimNextScopedCompileJob.h"
 #include "AnimNextSharedVariableNode.h"
+#include "ScopedTransaction.h"
+#include "Misc/UObjectToken.h"
 #include "Styling/SlateIconFinder.h"
 
 #define LOCTEXT_NAMESPACE "AnimNextUncookedOnlyUtils"
@@ -61,8 +63,6 @@ FInstancedPropertyBag FUtils::MakePropertyBagForEditorData(const UAnimNextRigVMA
 {
 	ensureMsgf(InEditorData == InCompileContext.GetOwningAssetEditorData(), TEXT("Crossing RigVM asset with a different asset's compile context. Expect incorrect memory layout / variable generation"));
 
-	FMessageLog Log("AnimNextCompilerResults");
-	
 	struct FStructEntryInfo
 	{
 		FName Name;
@@ -79,18 +79,19 @@ FInstancedPropertyBag FUtils::MakePropertyBagForEditorData(const UAnimNextRigVMA
 	StructEntryInfos.Reserve(InEditorData->Entries.Num() + ProgrammaticVariables.Num());
 	int32 NumPublicVariables = 0;
 
-	auto AddVariable = [&Log, &NumPublicVariables, &StructEntryInfos, &EntryInfoIndexMap](FName InName, const FAnimNextParamType& InType, EAnimNextExportAccessSpecifier InAccess, TConstArrayView<uint8> InValue, EPropertyFlags InFlags)
+	UAnimNextRigVMAsset* Asset = GetAsset<UAnimNextRigVMAsset>(InEditorData);
+	auto AddVariable = [Asset, &NumPublicVariables, &StructEntryInfos, &EntryInfoIndexMap, &InCompileContext](FName InName, const FAnimNextParamType& InType, EAnimNextExportAccessSpecifier InAccess, TConstArrayView<uint8> InValue, EPropertyFlags InFlags)
 	{
 		if(!InType.IsValid())
 		{
-			Log.Error(FText::Format(LOCTEXT("InvalidVariableTypeFound", "Variable '{0}' with invalid type found"), FText::FromName(InName)));
+			InCompileContext.GetAssetCompileContext().Error(Asset, LOCTEXT("InvalidVariableTypeFound", " @@ Variable '{0}' with invalid type found"), FText::FromName(InName));
 			return;
 		}
 
 		// Check for conflicts
 		if(EntryInfoIndexMap.Contains(InName))
 		{
-			Log.Error(FText::Format(LOCTEXT("DuplicateVariableFound", "Variable '{0}' with duplicate name found"), FText::FromName(InName)));
+			InCompileContext.GetAssetCompileContext().Error(Asset, LOCTEXT("DuplicateVariableFound", " @@ Variable '{0}' with duplicate name found"), FText::FromName(InName));
 			return;
 		}
 
@@ -196,12 +197,35 @@ void FUtils::CompileVariables(const FRigVMCompileSettings& InSettings, UAnimNext
 	TSet<const UAnimNextRigVMAsset*> SharedVariableAssets;
 	TSet<const UScriptStruct*> SharedVariableStructs;
 
-	FMessageLog Log("AnimNextCompilerResults");
-
 	for(UAnimNextRigVMAssetEntry* Entry : EditorData->Entries)
 	{
-		if(const UAnimNextSharedVariablesEntry* SharedVariablesEntry = Cast<UAnimNextSharedVariablesEntry>(Entry))
+		if(UAnimNextSharedVariablesEntry* SharedVariablesEntry = Cast<UAnimNextSharedVariablesEntry>(Entry))
 		{
+
+			auto SharedAssetEntryError = [InAsset, EditorData, SharedVariablesEntry, &OutCompileContext](const FText& InMessage)
+			{
+				TWeakObjectPtr<UAnimNextRigVMAssetEditorData> WeakEditorData = EditorData;
+				TWeakObjectPtr<UAnimNextSharedVariablesEntry> WeakEntry = SharedVariablesEntry;
+					
+				TSharedRef<FTokenizedMessage> Message = FTokenizedMessage::Create(EMessageSeverity::Error);
+				Message->AddToken(FUObjectToken::Create(InAsset));
+				Message->AddText(LOCTEXT("InvalidSharedVariableAssetFound", "Invalid shared variables entry found ({0})"), FText::FromString(SharedVariablesEntry->GetObjectPath().ToString()));
+				Message->AddToken(FActionToken::Create(LOCTEXT("RemoveSharedVariablesEntry", "Remove Shared Variables entry"), LOCTEXT("RemoveSharedVariablesEntryDesc", "Removes the Shared Variables entry from the asset."), FOnActionTokenExecuted::CreateLambda([WeakEditorData, WeakEntry]()
+					{
+						UAnimNextRigVMAssetEditorData* EditorData = WeakEditorData.Get();
+						UAnimNextSharedVariablesEntry* Entry = WeakEntry.Get();
+						if (EditorData && Entry)
+						{
+							FScopedTransaction Transaction(LOCTEXT("RemoveSharedVariablesTransactionDesc", "Removing Shared Variables entry"));
+							EditorData->RemoveEntry(Entry);
+						}						
+					
+					}), true));						
+
+
+				OutCompileContext.GetAssetCompileContext().Message(Message);
+			};
+			
 			switch (SharedVariablesEntry->GetType())
 			{
 			case EAnimNextSharedVariablesType::Asset:
@@ -209,13 +233,13 @@ void FUtils::CompileVariables(const FRigVMCompileSettings& InSettings, UAnimNext
 					const UAnimNextSharedVariables* SharedVariablesAsset = SharedVariablesEntry->GetAsset();
 					if (SharedVariablesAsset == nullptr)
 					{
-						Log.Error(FText::Format(LOCTEXT("InvalidSharedVariableAssetFound", "Invalid shared variables asset found ({0})"), FText::FromString(SharedVariablesEntry->GetObjectPath().ToString())));
+						SharedAssetEntryError(FText::Format(LOCTEXT("InvalidSharedVariableAssetFound", "Invalid shared variables entry found ({0})"), FText::FromString(SharedVariablesEntry->GetObjectPath().ToString())));						
 						continue;
 					}
 
 					if (SharedVariableAssets.Contains(SharedVariablesAsset))
 					{
-						Log.Error(FText::Format(LOCTEXT("DuplicateSharedVariableAssetFound", "Duplicate shared variables asset found ({0})"), FText::FromString(SharedVariablesEntry->GetObjectPath().ToString())));
+						SharedAssetEntryError(FText::Format(LOCTEXT("DuplicateSharedVariableAssetFound", "Duplicate shared variables entry found ({0})"), FText::FromString(SharedVariablesEntry->GetObjectPath().ToString())));						
 						continue;
 					}
 
@@ -227,13 +251,13 @@ void FUtils::CompileVariables(const FRigVMCompileSettings& InSettings, UAnimNext
 					const UScriptStruct* SharedVariablesStruct = SharedVariablesEntry->GetStruct();
 					if (SharedVariablesStruct == nullptr)
 					{
-						Log.Error(FText::Format(LOCTEXT("InvalidSharedVariableStructFound", "Invalid shared variables struct found ({0})"), FText::FromString(SharedVariablesEntry->GetObjectPath().ToString())));
+						SharedAssetEntryError(FText::Format(LOCTEXT("InvalidSharedVariableStructFound", "Invalid shared variables struct found ({0})"), FText::FromString(SharedVariablesEntry->GetObjectPath().ToString())));
 						continue;
 					}
 
 					if (SharedVariableStructs.Contains(SharedVariablesStruct))
 					{
-						Log.Error(FText::Format(LOCTEXT("DuplicateSharedVariableStructFound", "Duplicate shared variables struct found ({0})"), FText::FromString(SharedVariablesEntry->GetObjectPath().ToString())));
+						SharedAssetEntryError(FText::Format(LOCTEXT("DuplicateSharedVariableStructFound", "Duplicate shared variables struct found ({0})"), FText::FromString(SharedVariablesEntry->GetObjectPath().ToString())));						
 						continue;
 					}
 
@@ -283,14 +307,14 @@ void FUtils::CompileVariables(const FRigVMCompileSettings& InSettings, UAnimNext
 	{
 		if (!InAsset->DefaultInjectionSite.IsValid())
 		{
-			Log.Error(FText::Format(LOCTEXT("MissingDefaultInjectionSiteError", "Could not find default injection site: {0}"), FText::FromName(EditorData->DefaultInjectionSiteReference.GetName())));
+			OutCompileContext.GetAssetCompileContext().Error(InAsset, LOCTEXT("MissingDefaultInjectionSiteError", "@@ Could not find default injection site: {0}"), FText::FromName(EditorData->DefaultInjectionSiteReference.GetName()));
 		}
 		else
 		{
 			// Warn about using the deprecated name-based path
 			if (EditorData->DefaultInjectionSiteReference.GetObject() == nullptr)
 			{
-				Log.Warning(FText::Format(LOCTEXT("DeprecatedNamedInjectionSiteWarning", "Default injection site '{0}' uses a name-based reference. Please select a full reference to an asset's variable."), FText::FromName(EditorData->DefaultInjectionSiteReference.GetName())));
+				OutCompileContext.GetAssetCompileContext().Warning(InAsset, LOCTEXT("DeprecatedNamedInjectionSiteWarning", "@@ Default injection site '{0}' uses a name-based reference. Please select a full reference to an asset's variable."), FText::FromName(EditorData->DefaultInjectionSiteReference.GetName()));
 			}
 
 			InAsset->DefaultInjectionSite = EditorData->DefaultInjectionSiteReference; 
