UAF: Access to component module

diff --git a/Engine/Plugins/Experimental/UAF/UAF/Source/UAF/Private/Component/AnimNextComponent.cpp b/Engine/Plugins/Experimental/UAF/UAF/Source/UAF/Private/Component/AnimNextComponent.cpp
--- a/Engine/Plugins/Experimental/UAF/UAF/Source/UAF/Private/Component/AnimNextComponent.cpp
+++ b/Engine/Plugins/Experimental/UAF/UAF/Source/UAF/Private/Component/AnimNextComponent.cpp
@@ -490,3 +490,8 @@ bool UAnimNextComponent::IsModuleValid()
 {
 	return ModuleHandle.IsValid();
 }
+
+TObjectPtr<UAnimNextModule> UAnimNextComponent::GetModule() const
+{
+	return Module;
+}
diff --git a/Engine/Plugins/Experimental/UAF/UAF/Source/UAF/Public/Component/AnimNextComponent.h b/Engine/Plugins/Experimental/UAF/UAF/Source/UAF/Public/Component/AnimNextComponent.h
--- a/Engine/Plugins/Experimental/UAF/UAF/Source/UAF/Public/Component/AnimNextComponent.h
+++ b/Engine/Plugins/Experimental/UAF/UAF/Source/UAF/Public/Component/AnimNextComponent.h
@@ -193,6 +193,8 @@ public:
 	UFUNCTION(BlueprintCallable, Category = "UAF", meta=(DisplayName = "Get Module Handle"))
 	UAF_API FAnimNextModuleHandle BlueprintGetModuleHandle() const;
 
+	UAF_API TObjectPtr<UAnimNextModule> GetModule() const;
+
 private:
 	UAF_API void SetModule(TObjectPtr<UAnimNextModule> InModule);
 
