Disable adding set bindings without a set collection set

diff --git a/Engine/Plugins/Experimental/UAF/UAF/Source/UAF/Private/UAF/AbstractSkeleton/AbstractSkeletonSetBinding.cpp b/Engine/Plugins/Experimental/UAF/UAF/Source/UAF/Private/UAF/AbstractSkeleton/AbstractSkeletonSetBinding.cpp
--- a/Engine/Plugins/Experimental/UAF/UAF/Source/UAF/Private/UAF/AbstractSkeleton/AbstractSkeletonSetBinding.cpp
+++ b/Engine/Plugins/Experimental/UAF/UAF/Source/UAF/Private/UAF/AbstractSkeleton/AbstractSkeletonSetBinding.cpp
@@ -78,6 +78,11 @@ const TConstArrayView<FAbstractSkeleton_BoneBinding> UAbstractSkeletonSetBinding
 
 bool UAbstractSkeletonSetBinding::AddAttributeToSet(const FAnimationAttributeIdentifier Attribute, const FName SetName)
 {
+	if (!SetCollection)
+	{
+		return false;
+	}
+
 	if (!SetCollection->HasSet(SetName) && SetName != NAME_None)
 	{
 		return false; // Set does not exist
diff --git a/Engine/Plugins/Experimental/UAF/UAF/Source/UAFEditor/Private/UAF/AbstractSkeleton/Sets/SAttributesList.cpp b/Engine/Plugins/Experimental/UAF/UAF/Source/UAFEditor/Private/UAF/AbstractSkeleton/Sets/SAttributesList.cpp
--- a/Engine/Plugins/Experimental/UAF/UAF/Source/UAFEditor/Private/UAF/AbstractSkeleton/Sets/SAttributesList.cpp
+++ b/Engine/Plugins/Experimental/UAF/UAF/Source/UAFEditor/Private/UAF/AbstractSkeleton/Sets/SAttributesList.cpp
@@ -51,6 +51,10 @@ void SAttributesList::Construct(const FArguments& InArgs)
 							.Icon(FAppStyle::Get().GetBrush("Icons.Plus"))
 							.Text(LOCTEXT("AddAttributeButton_Label", "Add"))
 							.OnGetMenuContent( this, &SAttributesList::CreateAddAttributeWidget)
+							.IsEnabled_Lambda([this]()
+							{
+								return SetBinding.IsValid() && SetBinding->GetSetCollection() && SetBinding->GetSkeleton();
+							})
 					]
 					+ SHorizontalBox::Slot()
 					.FillWidth(1.0f)
