Support drag/droping multiple bones in abstract skeleton binding editor

diff --git a/Engine/Plugins/Experimental/UAF/UAF/Source/UAFEditor/Private/UAF/AbstractSkeleton/Sets/SSetBinding.cpp b/Engine/Plugins/Experimental/UAF/UAF/Source/UAFEditor/Private/UAF/AbstractSkeleton/Sets/SSetBinding.cpp
--- a/Engine/Plugins/Experimental/UAF/UAF/Source/UAFEditor/Private/UAF/AbstractSkeleton/Sets/SSetBinding.cpp
+++ b/Engine/Plugins/Experimental/UAF/UAF/Source/UAFEditor/Private/UAF/AbstractSkeleton/Sets/SSetBinding.cpp
@@ -178,8 +178,15 @@ void SSetBinding::Construct(const FArguments& InArgs, TWeakObjectPtr<UAbstractSk
 									const FScopedTransaction Transaction(LOCTEXT("AddAttributeToSet", "Add Attribute to Set"));
 									SetBinding->Modify();
 
-									const bool bSuccess = SetBinding->AddBoneToSet(DragDropOp->Item->BoneName, TargetSetTreeItem->Name);
-									if (bSuccess)
+									bool bTreeDirty = false;
+									
+									for (const TSharedPtr<SSetsSkeletonTree::FTreeItem>& Item : DragDropOp->Items)
+									{
+										bTreeDirty |= SetBinding->RemoveBoneFromSet(Item->BoneName);
+										bTreeDirty |= SetBinding->AddBoneToSet(Item->BoneName, TargetSetTreeItem->Name);
+									}
+									
+									if (bTreeDirty)
 									{
 										RepopulateTreeData();
 									}
diff --git a/Engine/Plugins/Experimental/UAF/UAF/Source/UAFEditor/Private/UAF/AbstractSkeleton/Sets/SSetSkeletonTree.cpp b/Engine/Plugins/Experimental/UAF/UAF/Source/UAFEditor/Private/UAF/AbstractSkeleton/Sets/SSetSkeletonTree.cpp
--- a/Engine/Plugins/Experimental/UAF/UAF/Source/UAFEditor/Private/UAF/AbstractSkeleton/Sets/SSetSkeletonTree.cpp
+++ b/Engine/Plugins/Experimental/UAF/UAF/Source/UAFEditor/Private/UAF/AbstractSkeleton/Sets/SSetSkeletonTree.cpp
@@ -192,11 +192,12 @@ namespace UE::UAF::Sets
 					}
 					return FReply::Unhandled();
 				})
-			.OnDragDetected_Lambda([Item](const FGeometry& MyGeometry, const FPointerEvent& MouseEvent)
+			.OnDragDetected_Lambda([this](const FGeometry& MyGeometry, const FPointerEvent& MouseEvent)
 				{
 					if (MouseEvent.IsMouseButtonDown(EKeys::LeftMouseButton))
 					{
-						const TSharedRef<FSkeletonTreeDragDropOp> DragDropOp = FSkeletonTreeDragDropOp::New(Item);
+						TArray<TSharedPtr<FTreeItem>> Selection = TreeView->GetSelectedItems();
+						const TSharedRef<FSkeletonTreeDragDropOp> DragDropOp = FSkeletonTreeDragDropOp::New(MoveTemp(Selection));
 						return FReply::Handled().BeginDragDrop(DragDropOp);
 					}
 
diff --git a/Engine/Plugins/Experimental/UAF/UAF/Source/UAFEditor/Private/UAF/AbstractSkeleton/Sets/SkeletonTreeDragDrop.h b/Engine/Plugins/Experimental/UAF/UAF/Source/UAFEditor/Private/UAF/AbstractSkeleton/Sets/SkeletonTreeDragDrop.h
--- a/Engine/Plugins/Experimental/UAF/UAF/Source/UAFEditor/Private/UAF/AbstractSkeleton/Sets/SkeletonTreeDragDrop.h
+++ b/Engine/Plugins/Experimental/UAF/UAF/Source/UAFEditor/Private/UAF/AbstractSkeleton/Sets/SkeletonTreeDragDrop.h
@@ -13,10 +13,10 @@ namespace UE::UAF::Sets
 	public:
 		DRAG_DROP_OPERATOR_TYPE(FSkeletonTreeDragDropOp, FDecoratedDragDropOp)
 
-		static TSharedRef<FSkeletonTreeDragDropOp> New(TSharedPtr<SSetsSkeletonTree::FTreeItem> InItem)
+		static TSharedRef<FSkeletonTreeDragDropOp> New(TArray<TSharedPtr<SSetsSkeletonTree::FTreeItem>>&& InItems)
 		{
 			TSharedRef<FSkeletonTreeDragDropOp> Operation = MakeShared<FSkeletonTreeDragDropOp>();
-			Operation->Item = InItem;
+			Operation->Items = MoveTemp(InItems);
 
 			Operation->Construct();
 
@@ -30,10 +30,10 @@ namespace UE::UAF::Sets
 				.BorderImage(FAppStyle::GetBrush("Menu.Background"))
 				[
 					SNew(STextBlock)
-						.Text(FText::FromName(Item->BoneName))
+						.Text(Items.Num() == 1 ? FText::FromName(Items[0]->BoneName) :	FText::Format(NSLOCTEXT("UE::UAF::Sets::FSkeletonTreeDragDropOp", "MultipleBonesTooltip", "{0} bones"), FText::AsNumber(Items.Num())))
 				];
 		}
 
-		TSharedPtr<SSetsSkeletonTree::FTreeItem> Item;
+		TArray<TSharedPtr<SSetsSkeletonTree::FTreeItem>> Items;
 	};
 }
\ No newline at end of file
