Fix the expansion state of tree items resetting when drag/dropping a bone in abstract skeleton binding editor

diff --git a/Engine/Plugins/Experimental/UAF/UAF/Source/UAFEditor/Private/UAF/AbstractSkeleton/Sets/SSetSkeletonTree.cpp b/Engine/Plugins/Experimental/UAF/UAF/Source/UAFEditor/Private/UAF/AbstractSkeleton/Sets/SSetSkeletonTree.cpp
--- a/Engine/Plugins/Experimental/UAF/UAF/Source/UAFEditor/Private/UAF/AbstractSkeleton/Sets/SSetSkeletonTree.cpp
+++ b/Engine/Plugins/Experimental/UAF/UAF/Source/UAFEditor/Private/UAF/AbstractSkeleton/Sets/SSetSkeletonTree.cpp
@@ -235,6 +235,16 @@ namespace UE::UAF::Sets
 		{
 			TGuardValue<bool> RecursionGuard(bRepopulating, true);
 
+			// Save expansion state of each tree item to restore later
+			TSet<FName> CollapsedBoneNames;
+			for (const TSharedPtr<FTreeItem>& TreeItem : GetAllTreeItems())
+			{
+				if (!TreeView->IsItemExpanded(TreeItem))
+				{
+					CollapsedBoneNames.Add(TreeItem->BoneName);
+				}
+			}
+
 			RootItems.Empty();
 
 			if (SetBinding.IsValid() && SetBinding->GetSkeleton())
@@ -267,11 +277,12 @@ namespace UE::UAF::Sets
 				RootItems.Add(RootBoneItem);
 			}
 
-			TreeView->RequestTreeRefresh();
+			TreeView->RebuildList();
 
 			for (TSharedPtr<FTreeItem> TreeItem : GetAllTreeItems())
 			{
-				TreeView->SetItemExpansion(TreeItem, true);
+				const bool bExpanded = !CollapsedBoneNames.Contains(TreeItem->BoneName);
+				TreeView->SetItemExpansion(TreeItem, bExpanded);
 			}
 
 			OnTreeRefreshed.ExecuteIfBound();
