Allow bone items to be dropped onto non-set items in the abstract skeleton set binding editor

diff --git a/Engine/Plugins/Experimental/UAF/UAF/Source/UAF/Public/UAF/AbstractSkeleton/AbstractSkeletonSetCollection.h b/Engine/Plugins/Experimental/UAF/UAF/Source/UAF/Public/UAF/AbstractSkeleton/AbstractSkeletonSetCollection.h
--- a/Engine/Plugins/Experimental/UAF/UAF/Source/UAF/Public/UAF/AbstractSkeleton/AbstractSkeletonSetCollection.h
+++ b/Engine/Plugins/Experimental/UAF/UAF/Source/UAF/Public/UAF/AbstractSkeleton/AbstractSkeletonSetCollection.h
@@ -36,6 +36,7 @@ public:
 
 	UE_API bool ReparentSet(const FName SetName, const FName NewParentName);
 
+	// Returns the sets in this set collection sorted topologically
 	UE_API const TArrayView<const FAbstractSkeletonSet> GetSetHierarchy() const;
 
 	FName GetParentSetName(const FName SetName) const;
diff --git a/Engine/Plugins/Experimental/UAF/UAF/Source/UAFEditor/Private/UAF/AbstractSkeleton/Sets/SAttributesList.cpp b/Engine/Plugins/Experimental/UAF/UAF/Source/UAFEditor/Private/UAF/AbstractSkeleton/Sets/SAttributesList.cpp
--- a/Engine/Plugins/Experimental/UAF/UAF/Source/UAFEditor/Private/UAF/AbstractSkeleton/Sets/SAttributesList.cpp
+++ b/Engine/Plugins/Experimental/UAF/UAF/Source/UAFEditor/Private/UAF/AbstractSkeleton/Sets/SAttributesList.cpp
@@ -228,7 +228,7 @@ TSharedRef<ITableRow> SAttributesList::ListView_OnGenerateRow(FListItemPtr Item,
 					const FScopedTransaction Transaction(LOCTEXT("AddAttributeToSet", "Add Attribute to Set"));
 					SetBinding->Modify();
 
-					SetBinding->AddAttributeToSet(TargetTreeItem->Attribute, DragDropOp->Item->Name);
+					SetBinding->AddAttributeToSet(TargetTreeItem->Attribute, DragDropOp->Item->Set.SetName);
 					
 					RepopulateListData();
 
diff --git a/Engine/Plugins/Experimental/UAF/UAF/Source/UAFEditor/Private/UAF/AbstractSkeleton/Sets/SSetBinding.cpp b/Engine/Plugins/Experimental/UAF/UAF/Source/UAFEditor/Private/UAF/AbstractSkeleton/Sets/SSetBinding.cpp
--- a/Engine/Plugins/Experimental/UAF/UAF/Source/UAFEditor/Private/UAF/AbstractSkeleton/Sets/SSetBinding.cpp
+++ b/Engine/Plugins/Experimental/UAF/UAF/Source/UAFEditor/Private/UAF/AbstractSkeleton/Sets/SSetBinding.cpp
@@ -93,7 +93,7 @@ void SSetBinding::Construct(const FArguments& InArgs, TWeakObjectPtr<UAbstractSk
 		{
 			TSharedPtr<ITreeItem> Selected = Selection[0]; // Single selection only
 
-			const bool bShowMenu = Selected->GenerateContextMenu(*this, MenuBuilder);
+			const bool bShowMenu = Selected->GenerateContextMenu(MenuBuilder);
 
 			if (bShowMenu)
 			{
@@ -135,7 +135,7 @@ void SSetBinding::Construct(const FArguments& InArgs, TWeakObjectPtr<UAbstractSk
 					if (TreeItem->GetType() == FTreeItem_Set::StaticGetType() && !TreeView->IsItemExpanded(TreeItem))
 					{
 						TSharedPtr<FTreeItem_Set> SetTreeItem = StaticCastSharedPtr<FTreeItem_Set>(TreeItem);
-						ClosedSetCollections.Add(SetTreeItem->Name);
+						ClosedSetCollections.Add(SetTreeItem->Set.SetName);
 					}
 				}
 			}
@@ -151,50 +151,7 @@ void SSetBinding::Construct(const FArguments& InArgs, TWeakObjectPtr<UAbstractSk
 				{
 					for (const FAbstractSkeletonSet& Set : SetBinding->GetSetCollection()->GetSetHierarchy())
 					{
-						TSharedPtr<FTreeItem_Set> SetItem = MakeShared<FTreeItem_Set>();
-						SetItem->Name = Set.SetName;
-						SetItem->ParentName = Set.ParentSetName;
-						SetItem->OnCanAcceptDrop = STableRow<TSharedPtr<ITreeItem>>::FOnCanAcceptDrop::CreateLambda([](const FDragDropEvent& DragDropEvent, EItemDropZone DropZone, TSharedPtr<ITreeItem> TargetItem) -> TOptional<EItemDropZone>
-							{
-								using namespace UE::UAF;
-
-								TOptional<EItemDropZone> ReturnedDropZone;
-
-								const TSharedPtr<FSkeletonTreeDragDropOp> DragDropOp = DragDropEvent.GetOperationAs<FSkeletonTreeDragDropOp>();
-								if (DragDropOp.IsValid())
-								{
-									ReturnedDropZone = EItemDropZone::OntoItem;
-								}
-
-								return ReturnedDropZone;
-							});
-						SetItem->OnAcceptDrop = STableRow<TSharedPtr<ITreeItem>>::FOnAcceptDrop::CreateLambda([this](const FDragDropEvent& DragDropEvent, EItemDropZone DropZone, TSharedPtr<ITreeItem> TargetTreeItem) -> FReply
-							{
-								const TSharedPtr<FSkeletonTreeDragDropOp> DragDropOp = DragDropEvent.GetOperationAs<FSkeletonTreeDragDropOp>();
-								if (DragDropOp.IsValid())
-								{
-									TSharedPtr<FTreeItem_Set> TargetSetTreeItem = StaticCastSharedPtr<FTreeItem_Set>(TargetTreeItem);
-
-									const FScopedTransaction Transaction(LOCTEXT("AddAttributeToSet", "Add Attribute to Set"));
-									SetBinding->Modify();
-
-									bool bTreeDirty = false;
-									
-									for (const TSharedPtr<SSetsSkeletonTree::FTreeItem>& Item : DragDropOp->Items)
-									{
-										bTreeDirty |= SetBinding->RemoveBoneFromSet(Item->BoneName);
-										bTreeDirty |= SetBinding->AddBoneToSet(Item->BoneName, TargetSetTreeItem->Name);
-									}
-									
-									if (bTreeDirty)
-									{
-										RepopulateTreeData();
-									}
-
-									return FReply::Handled();
-								}
-								return FReply::Unhandled();
-							});
+						TSharedPtr<FTreeItem_Set> SetItem = MakeShared<FTreeItem_Set>(*this, Set);
 
 						if (Set.ParentSetName == NAME_None)
 						{
@@ -202,17 +159,10 @@ void SSetBinding::Construct(const FArguments& InArgs, TWeakObjectPtr<UAbstractSk
 						}
 						else
 						{
-							if (!TreeItemMap.Contains(Set.ParentSetName))
-							{
-								TSharedPtr<FTreeItem_Set> ParentItem = MakeShared<FTreeItem_Set>();
-								ParentItem->Name = Set.ParentSetName;
-
-								TreeItemMap.Add(Set.ParentSetName, ParentItem);
-							}
-
+							ensure(TreeItemMap.Contains(Set.ParentSetName));
+							
 							TSharedPtr<FTreeItem_Set> ParentItem = TreeItemMap[Set.ParentSetName];
 							ParentItem->Children.Add(SetItem);
-
 						}
 
 						TreeItemMap.Add(Set.SetName, SetItem);
@@ -221,8 +171,7 @@ void SSetBinding::Construct(const FArguments& InArgs, TWeakObjectPtr<UAbstractSk
 
 				for (const FAbstractSkeleton_BoneBinding& Binding : SetBinding->GetBoneBindings())
 				{
-					TSharedPtr<FTreeItem_Bone> BindingItem = MakeShared<FTreeItem_Bone>();
-					BindingItem->BoneName = Binding.BoneName;
+					TSharedPtr<FTreeItem_Bone> BindingItem = MakeShared<FTreeItem_Bone>(*this, Binding);
 
 					if (TreeItemMap.Contains(Binding.SetName))
 					{
@@ -237,8 +186,7 @@ void SSetBinding::Construct(const FArguments& InArgs, TWeakObjectPtr<UAbstractSk
 
 				for (const FAbstractSkeleton_AttributeBinding& Binding : SetBinding->GetAttributeBindings())
 				{
-					TSharedPtr<FTreeItem_Attribute> BindingItem = MakeShared<FTreeItem_Attribute>();
-					BindingItem->Attribute = Binding.Attribute;
+					TSharedPtr<FTreeItem_Attribute> BindingItem = MakeShared<FTreeItem_Attribute>(*this, Binding);
 
 					if (TreeItemMap.Contains(Binding.SetName))
 					{
@@ -257,7 +205,7 @@ void SSetBinding::Construct(const FArguments& InArgs, TWeakObjectPtr<UAbstractSk
 					if (TreeItem->GetType() == FTreeItem_Set::StaticGetType())
 					{
 						const TSharedPtr<FTreeItem_Set> SetTreeItem = StaticCastSharedPtr<FTreeItem_Set>(TreeItem);
-						const bool bShouldExpand = !ClosedSetCollections.Contains(SetTreeItem->Name);
+						const bool bShouldExpand = !ClosedSetCollections.Contains(SetTreeItem->Set.SetName);
 						TreeView->SetItemExpansion(TreeItem, bShouldExpand);
 					}
 				}
@@ -312,8 +260,47 @@ void SSetBinding::Construct(const FArguments& InArgs, TWeakObjectPtr<UAbstractSk
 
 					return FReply::Unhandled();
 				})
-			.OnCanAcceptDrop(OnCanAcceptDrop)
-			.OnAcceptDrop(OnAcceptDrop)
+			.OnCanAcceptDrop_Lambda([](const FDragDropEvent& DragDropEvent, EItemDropZone DropZone, TSharedPtr<ITreeItem> TargetItem) -> TOptional<EItemDropZone>
+				{
+					using namespace UE::UAF;
+
+					TOptional<EItemDropZone> ReturnedDropZone;
+
+					const TSharedPtr<FSkeletonTreeDragDropOp> DragDropOp = DragDropEvent.GetOperationAs<FSkeletonTreeDragDropOp>();
+					if (DragDropOp.IsValid())
+					{
+						ReturnedDropZone = EItemDropZone::OntoItem;
+					}
+
+					return ReturnedDropZone;
+				})
+			.OnAcceptDrop_Lambda([this](const FDragDropEvent& DragDropEvent, EItemDropZone DropZone, TSharedPtr<ITreeItem> TargetTreeItem) -> FReply
+				{
+					const TSharedPtr<FSkeletonTreeDragDropOp> DragDropOp = DragDropEvent.GetOperationAs<FSkeletonTreeDragDropOp>();
+					if (DragDropOp.IsValid())
+					{
+						TSharedPtr<FTreeItem_Set> TargetSetTreeItem = StaticCastSharedPtr<FTreeItem_Set>(TargetTreeItem);
+
+						const FScopedTransaction Transaction(LOCTEXT("AddAttributeToSet", "Add Attribute to Set"));
+						SetBindingWidget.SetBinding->Modify();
+
+						bool bTreeDirty = false;
+									
+						for (const TSharedPtr<SSetsSkeletonTree::FTreeItem>& Item : DragDropOp->Items)
+						{
+							bTreeDirty |= SetBindingWidget.SetBinding->RemoveBoneFromSet(Item->BoneName);
+							bTreeDirty |= SetBindingWidget.SetBinding->AddBoneToSet(Item->BoneName, TargetSetTreeItem->Set.SetName);
+						}
+							
+						if (bTreeDirty)
+						{
+							SetBindingWidget.RepopulateTreeData();
+						}
+
+						return FReply::Handled();
+					}
+					return FReply::Unhandled();
+				})
 			.ShowWires(true)
 			[
 				SNew(SHorizontalBox)
@@ -330,7 +317,7 @@ void SSetBinding::Construct(const FArguments& InArgs, TWeakObjectPtr<UAbstractSk
 					.Padding(3.0f)
 					[
 						SNew(STextBlock)
-							.Text(FText::FromName(Name))
+							.Text(FText::FromName(Set.SetName))
 					]
 
 			];
@@ -352,6 +339,45 @@ void SSetBinding::Construct(const FArguments& InArgs, TWeakObjectPtr<UAbstractSk
 
 		return SNew(STableRow<TSharedPtr<ITreeItem>>, OwnerTable)
 			.Style(FAppStyle::Get(), "TableView.AlternatingRow")
+			.OnCanAcceptDrop_Lambda([](const FDragDropEvent& DragDropEvent, EItemDropZone DropZone, TSharedPtr<ITreeItem> TargetItem) -> TOptional<EItemDropZone>
+				{
+					using namespace UE::UAF;
+
+					TOptional<EItemDropZone> ReturnedDropZone;
+
+					const TSharedPtr<FSkeletonTreeDragDropOp> DragDropOp = DragDropEvent.GetOperationAs<FSkeletonTreeDragDropOp>();
+					if (DragDropOp.IsValid())
+					{
+						ReturnedDropZone = EItemDropZone::OntoItem;
+					}
+
+					return ReturnedDropZone;
+				})
+			.OnAcceptDrop_Lambda([this](const FDragDropEvent& DragDropEvent, EItemDropZone DropZone, TSharedPtr<ITreeItem> TargetTreeItem) -> FReply
+				{
+					const TSharedPtr<FSkeletonTreeDragDropOp> DragDropOp = DragDropEvent.GetOperationAs<FSkeletonTreeDragDropOp>();
+					if (DragDropOp.IsValid())
+					{
+						const FScopedTransaction Transaction(LOCTEXT("AddAttributeToSet", "Add Attribute to Set"));
+						SetBindingWidget.SetBinding->Modify();
+
+						bool bTreeDirty = false;
+									
+						for (const TSharedPtr<SSetsSkeletonTree::FTreeItem>& Item : DragDropOp->Items)
+						{
+							bTreeDirty |= SetBindingWidget.SetBinding->RemoveBoneFromSet(Item->BoneName);
+							bTreeDirty |= SetBindingWidget.SetBinding->AddBoneToSet(Item->BoneName, Binding.SetName);
+						}
+							
+						if (bTreeDirty)
+						{
+							SetBindingWidget.RepopulateTreeData();
+						}
+
+						return FReply::Handled();
+					}
+					return FReply::Unhandled();
+				})
 			.ShowWires(true)
 			[
 				SNew(SHorizontalBox)
@@ -369,7 +395,7 @@ void SSetBinding::Construct(const FArguments& InArgs, TWeakObjectPtr<UAbstractSk
 					.Padding(3.0f)
 					[
 						SNew(STextBlock)
-							.Text(FText::FromName(BoneName))
+							.Text(FText::FromName(Binding.BoneName))
 					]
 
 			];
@@ -380,6 +406,46 @@ void SSetBinding::Construct(const FArguments& InArgs, TWeakObjectPtr<UAbstractSk
 		const FSlateBrush* Icon = FAppStyle::Get().GetBrush("AnimGraph.Attribute.Attributes.Icon");
 
 		return SNew(STableRow<TSharedPtr<ITreeItem>>, OwnerTable)
+			.Style(FAppStyle::Get(), "TableView.AlternatingRow")
+			.OnCanAcceptDrop_Lambda([](const FDragDropEvent& DragDropEvent, EItemDropZone DropZone, TSharedPtr<ITreeItem> TargetItem) -> TOptional<EItemDropZone>
+				{
+					using namespace UE::UAF;
+
+					TOptional<EItemDropZone> ReturnedDropZone;
+
+					const TSharedPtr<FSkeletonTreeDragDropOp> DragDropOp = DragDropEvent.GetOperationAs<FSkeletonTreeDragDropOp>();
+					if (DragDropOp.IsValid())
+					{
+						ReturnedDropZone = EItemDropZone::OntoItem;
+					}
+
+					return ReturnedDropZone;
+				})
+			.OnAcceptDrop_Lambda([this](const FDragDropEvent& DragDropEvent, EItemDropZone DropZone, TSharedPtr<ITreeItem> TargetTreeItem) -> FReply
+				{
+					const TSharedPtr<FSkeletonTreeDragDropOp> DragDropOp = DragDropEvent.GetOperationAs<FSkeletonTreeDragDropOp>();
+					if (DragDropOp.IsValid())
+					{
+						const FScopedTransaction Transaction(LOCTEXT("AddAttributeToSet", "Add Attribute to Set"));
+						SetBindingWidget.SetBinding->Modify();
+
+						bool bTreeDirty = false;
+									
+						for (const TSharedPtr<SSetsSkeletonTree::FTreeItem>& Item : DragDropOp->Items)
+						{
+							bTreeDirty |= SetBindingWidget.SetBinding->RemoveBoneFromSet(Item->BoneName);
+							bTreeDirty |= SetBindingWidget.SetBinding->AddBoneToSet(Item->BoneName, Binding.SetName);
+						}
+							
+						if (bTreeDirty)
+						{
+							SetBindingWidget.RepopulateTreeData();
+						}
+
+						return FReply::Handled();
+					}
+					return FReply::Unhandled();
+				})
 			.ShowWires(true)
 			[
 				SNew(SHorizontalBox)
@@ -397,18 +463,18 @@ void SSetBinding::Construct(const FArguments& InArgs, TWeakObjectPtr<UAbstractSk
 					.Padding(4.0f, 0.0f, 0.0f, 0.0f)
 					[
 						SNew(STextBlock)
-							.Text(FText::FromName(Attribute.GetName()))
+							.Text(FText::FromName(Binding.Attribute.GetName()))
 					]
 
 			];
 	}
 
-	bool SSetBinding::FTreeItem_Set::GenerateContextMenu(SSetBinding& SetBindingWidget, FMenuBuilder& Menu)
+	bool SSetBinding::FTreeItem_Set::GenerateContextMenu(FMenuBuilder& Menu)
 	{
 		return false;
 	}
 
-	bool SSetBinding::FTreeItem_Bone::GenerateContextMenu(SSetBinding& SetBindingWidget, FMenuBuilder& Menu)
+	bool SSetBinding::FTreeItem_Bone::GenerateContextMenu(FMenuBuilder& Menu)
 	{
 		FUIAction UnbindBone(
 			FExecuteAction::CreateLambda([&]()
@@ -416,7 +482,7 @@ void SSetBinding::Construct(const FArguments& InArgs, TWeakObjectPtr<UAbstractSk
 				FScopedTransaction Transaction(LOCTEXT("RemoveBoneFromSet", "Remove Bone From Set"));
 				SetBindingWidget.SetBinding->Modify();
 
-				SetBindingWidget.SetBinding->RemoveBoneFromSet(BoneName);
+				SetBindingWidget.SetBinding->RemoveBoneFromSet(Binding.BoneName);
 				SetBindingWidget.RepopulateTreeData();
 			})
 		);
@@ -430,7 +496,7 @@ void SSetBinding::Construct(const FArguments& InArgs, TWeakObjectPtr<UAbstractSk
 		return true;
 	}
 	
-	bool SSetBinding::FTreeItem_Attribute::GenerateContextMenu(SSetBinding& SetBindingWidget, FMenuBuilder& Menu)
+	bool SSetBinding::FTreeItem_Attribute::GenerateContextMenu(FMenuBuilder& Menu)
 	{
 		FUIAction UnbindAttribute(
 			FExecuteAction::CreateLambda([&]()
@@ -438,7 +504,7 @@ void SSetBinding::Construct(const FArguments& InArgs, TWeakObjectPtr<UAbstractSk
 				FScopedTransaction Transaction(LOCTEXT("RemoveAttributeFromSet", "Remove Attribute From Set"));
 				SetBindingWidget.SetBinding->Modify();
 
-				SetBindingWidget.SetBinding->RemoveAttributeFromSet(Attribute);
+				SetBindingWidget.SetBinding->RemoveAttributeFromSet(Binding.Attribute);
 				SetBindingWidget.RepopulateTreeData();
 			})
 		);
diff --git a/Engine/Plugins/Experimental/UAF/UAF/Source/UAFEditor/Private/UAF/AbstractSkeleton/Sets/SSetBinding.h b/Engine/Plugins/Experimental/UAF/UAF/Source/UAFEditor/Private/UAF/AbstractSkeleton/Sets/SSetBinding.h
--- a/Engine/Plugins/Experimental/UAF/UAF/Source/UAFEditor/Private/UAF/AbstractSkeleton/Sets/SSetBinding.h
+++ b/Engine/Plugins/Experimental/UAF/UAF/Source/UAFEditor/Private/UAF/AbstractSkeleton/Sets/SSetBinding.h
@@ -5,6 +5,8 @@
 #include "Animation/AnimData/AttributeIdentifier.h"
 #include "CoreMinimal.h"
 #include "EditorUndoClient.h"
+#include "UAF/AbstractSkeleton/AbstractSkeletonSetBinding.h"
+#include "UAF/AbstractSkeleton/AbstractSkeletonSetCollection.h"
 #include "Widgets/SCompoundWidget.h"
 #include "Widgets/Views/STableRow.h"
 #include "Widgets/Views/STreeView.h"
@@ -24,15 +26,28 @@ namespace UE::UAF::Sets
 	public:
 		struct ITreeItem
 		{
+			ITreeItem(SSetBinding& InSetBindingWidget)
+				: SetBindingWidget(InSetBindingWidget)
+			{
+			}
+			
 			virtual ~ITreeItem() = default;
 			virtual FName GetType() const = 0;
 			virtual void GetChildren(TArray<TSharedPtr<ITreeItem>>& OutChildren) = 0;
 			virtual TSharedRef<ITableRow> GenerateRow(const TSharedRef<STableViewBase>& OwnerTable) = 0;
-			virtual bool GenerateContextMenu(SSetBinding& SetBindingWidget, FMenuBuilder& Menu) = 0;
+			virtual bool GenerateContextMenu(FMenuBuilder& Menu) = 0;
+			
+			SSetBinding& SetBindingWidget;
 		};
 
 		struct FTreeItem_Set : public ITreeItem
 		{
+			FTreeItem_Set(SSetBinding& InSetBindingWidget, FAbstractSkeletonSet InSet)
+				: ITreeItem(InSetBindingWidget)
+				, Set(InSet)
+			{
+			}
+			
 			virtual ~FTreeItem_Set() override = default;
 			
 			static FName StaticGetType() { return FName("Set"); }
@@ -40,17 +55,20 @@ namespace UE::UAF::Sets
 			virtual FName GetType() const override { return StaticGetType(); }
 			virtual void GetChildren(TArray<TSharedPtr<ITreeItem>>& OutChildren) override;
 			virtual TSharedRef<ITableRow> GenerateRow(const TSharedRef<STableViewBase>& OwnerTable) override;
-			virtual bool GenerateContextMenu(SSetBinding& SetBindingWidget, FMenuBuilder& Menu) override;
+			virtual bool GenerateContextMenu(FMenuBuilder& Menu) override;
 
-			FName Name;
-			FName ParentName;
+			FAbstractSkeletonSet Set;
 			TArray<TSharedPtr<ITreeItem>> Children;
-			STableRow<TSharedPtr<ITreeItem>>::FOnCanAcceptDrop OnCanAcceptDrop;
-			STableRow<TSharedPtr<ITreeItem>>::FOnAcceptDrop OnAcceptDrop;
 		};
 
 		struct FTreeItem_Bone : public ITreeItem
 		{
+			FTreeItem_Bone(SSetBinding& InSetBindingWidget, FAbstractSkeleton_BoneBinding InBinding)
+				: ITreeItem(InSetBindingWidget)
+				, Binding(InBinding)
+			{
+			}
+			
 			virtual ~FTreeItem_Bone() override = default;
 
 			static FName StaticGetType() { return FName("Bone"); }
@@ -58,13 +76,19 @@ namespace UE::UAF::Sets
 			virtual FName GetType() const override { return StaticGetType(); }
 			virtual void GetChildren(TArray<TSharedPtr<ITreeItem>>& OutChildren) override;
 			virtual TSharedRef<ITableRow> GenerateRow(const TSharedRef<STableViewBase>& OwnerTable) override;
-			virtual bool GenerateContextMenu(SSetBinding& SetBindingWidget, FMenuBuilder& Menu) override;
+			virtual bool GenerateContextMenu(FMenuBuilder& Menu) override;
 
-			FName BoneName;
+			FAbstractSkeleton_BoneBinding Binding;
 		};
 		
 		struct FTreeItem_Attribute : public ITreeItem
 		{
+			FTreeItem_Attribute(SSetBinding& InSetBindingWidget, FAbstractSkeleton_AttributeBinding InBinding)
+				: ITreeItem(InSetBindingWidget)
+				, Binding(InBinding)
+			{
+			}
+			
 			virtual ~FTreeItem_Attribute() override = default;
 
 			static FName StaticGetType() { return FName("Attribute"); }
@@ -72,9 +96,9 @@ namespace UE::UAF::Sets
 			virtual FName GetType() const override { return StaticGetType(); }
 			virtual void GetChildren(TArray<TSharedPtr<ITreeItem>>& OutChildren) override;
 			virtual TSharedRef<ITableRow> GenerateRow(const TSharedRef<STableViewBase>& OwnerTable) override;
-			virtual bool GenerateContextMenu(SSetBinding& SetBindingWidget, FMenuBuilder& Menu) override;
+			virtual bool GenerateContextMenu(FMenuBuilder& Menu) override;
 
-			FAnimationAttributeIdentifier Attribute;
+			FAbstractSkeleton_AttributeBinding Binding;
 		};
 
 		void Construct(const FArguments& InArgs, TWeakObjectPtr<UAbstractSkeletonSetBinding> InSetBinding);
diff --git a/Engine/Plugins/Experimental/UAF/UAF/Source/UAFEditor/Private/UAF/AbstractSkeleton/Sets/SSetSkeletonTree.cpp b/Engine/Plugins/Experimental/UAF/UAF/Source/UAFEditor/Private/UAF/AbstractSkeleton/Sets/SSetSkeletonTree.cpp
--- a/Engine/Plugins/Experimental/UAF/UAF/Source/UAFEditor/Private/UAF/AbstractSkeleton/Sets/SSetSkeletonTree.cpp
+++ b/Engine/Plugins/Experimental/UAF/UAF/Source/UAFEditor/Private/UAF/AbstractSkeleton/Sets/SSetSkeletonTree.cpp
@@ -184,7 +184,7 @@ namespace UE::UAF::Sets
 						const FScopedTransaction Transaction(LOCTEXT("AddAttributeToSet", "Add Attribute to Set"));
 						SetBinding->Modify();
 
-						const bool bSuccess = SetBinding->AddBoneToSet(TargetSetTreeItem->BoneName, DragDropOp->Item->Name);
+						const bool bSuccess = SetBinding->AddBoneToSet(TargetSetTreeItem->BoneName, DragDropOp->Item->Set.SetName);
 						check(bSuccess);
 						RepopulateTreeData();
 
diff --git a/Engine/Plugins/Experimental/UAF/UAF/Source/UAFEditor/Private/UAF/AbstractSkeleton/Sets/SetDragDrop.cpp b/Engine/Plugins/Experimental/UAF/UAF/Source/UAFEditor/Private/UAF/AbstractSkeleton/Sets/SetDragDrop.cpp
--- a/Engine/Plugins/Experimental/UAF/UAF/Source/UAFEditor/Private/UAF/AbstractSkeleton/Sets/SetDragDrop.cpp
+++ b/Engine/Plugins/Experimental/UAF/UAF/Source/UAFEditor/Private/UAF/AbstractSkeleton/Sets/SetDragDrop.cpp
@@ -21,7 +21,7 @@ namespace UE::UAF::Sets
 			.BorderImage(FAppStyle::GetBrush("Menu.Background"))
 			[
 				SNew(STextBlock)
-					.Text(FText::FromName(Item->Name))
+					.Text(FText::FromName(Item->Set.SetName))
 			];
 	}
 }
\ No newline at end of file
