Fix crash when attempting to add an existing curve attribute to abstract skeleton binding

diff --git a/Engine/Plugins/Experimental/UAF/UAF/Source/UAFEditor/Private/UAF/AbstractSkeleton/Sets/SAttributesList.cpp b/Engine/Plugins/Experimental/UAF/UAF/Source/UAFEditor/Private/UAF/AbstractSkeleton/Sets/SAttributesList.cpp
--- a/Engine/Plugins/Experimental/UAF/UAF/Source/UAFEditor/Private/UAF/AbstractSkeleton/Sets/SAttributesList.cpp
+++ b/Engine/Plugins/Experimental/UAF/UAF/Source/UAFEditor/Private/UAF/AbstractSkeleton/Sets/SAttributesList.cpp
@@ -257,38 +257,46 @@ void SAttributesList::RegisterMenus()
 
 		AttributeWidgetSection.AddDynamicEntry(
 			"AttributePicker",
-			FNewToolMenuSectionDelegate::CreateLambda([this](FToolMenuSection& Section)
+			FNewToolMenuSectionDelegate::CreateLambda([](FToolMenuSection& Section)
 			{
-				FPersonaModule& PersonaModule = FModuleManager::LoadModuleChecked<FPersonaModule>("Persona");
-				
-				if (SetBinding.IsValid())
-				{	
-					FToolMenuEntry WidgetEntry = FToolMenuEntry::InitWidget(
-						"AttributePicker",
-						PersonaModule.CreateAttributePicker(
-							SetBinding->GetSkeleton(),
-							FOnAttributesPicked::CreateLambda([this](const TConstArrayView<FAnimationAttributeIdentifier> InSelectedAttributes)
-								{
-									const FScopedTransaction Transaction(LOCTEXT("AddAttributeToSet", "Add Attribute to Set"));
-									SetBinding->Modify();
-
-									for (const FAnimationAttributeIdentifier& Attribute : InSelectedAttributes)
-									{
-										SetBinding->AddAttributeToSet(Attribute, NAME_None);
-									}
+				if (const UAttributesListMenuContext* Context = Section.FindContext<UAttributesListMenuContext>())
+				{
+					FPersonaModule& PersonaModule = FModuleManager::LoadModuleChecked<FPersonaModule>("Persona");
 
-									FSlateApplication::Get().DismissAllMenus();
-									RepopulateListData();
-								}),
+					TSharedPtr<SAttributesList> SetBindingWidget = Context->SetBindingWidget.Pin();
+					ensure(SetBindingWidget);
+					
+					const TWeakObjectPtr<UAbstractSkeletonSetBinding> SetBinding = SetBindingWidget->SetBinding;
+					
+					if (SetBinding.IsValid())
+					{	
+						FToolMenuEntry WidgetEntry = FToolMenuEntry::InitWidget(
+							"AttributePicker",
+							PersonaModule.CreateAttributePicker(
+								SetBinding->GetSkeleton(),
+								FOnAttributesPicked::CreateLambda([SetBinding, WeakSetBindingWidget = Context->SetBindingWidget](const TConstArrayView<FAnimationAttributeIdentifier> InSelectedAttributes)
+									{
+										const FScopedTransaction Transaction(LOCTEXT("AddAttributeToSet", "Add Attribute to Set"));
+										SetBinding->Modify();
+
+										for (const FAnimationAttributeIdentifier& Attribute : InSelectedAttributes)
+										{
+											SetBinding->AddAttributeToSet(Attribute, NAME_None);
+										}
+
+										FSlateApplication::Get().DismissAllMenus();
+										WeakSetBindingWidget.Pin()->RepopulateListData();
+									}),
+								true,
+								false),
+							FText(),
 							true,
-							false),
-						FText(),
-						true,
-						false,
-						true
-					);
-
-					Section.AddEntry(WidgetEntry);
+							false,
+							true
+						);
+
+						Section.AddEntry(WidgetEntry);
+					}
 				}
 			}));
 	}
@@ -298,34 +306,42 @@ void SAttributesList::RegisterMenus()
 
 		CurveWidgetSection.AddDynamicEntry(
 			"CurvePicker",
-			FNewToolMenuSectionDelegate::CreateLambda([this](FToolMenuSection& Section)
+			FNewToolMenuSectionDelegate::CreateLambda([](FToolMenuSection& Section)
 			{
-				FPersonaModule& PersonaModule = FModuleManager::LoadModuleChecked<FPersonaModule>("Persona");
-						
-				if (SetBinding.IsValid())
+				if (const UAttributesListMenuContext* Context = Section.FindContext<UAttributesListMenuContext>())
 				{
-					FToolMenuEntry WidgetEntry = FToolMenuEntry::InitWidget(
-						"CurvePicker",
-						PersonaModule.CreateCurvePicker(
-							SetBinding->GetSkeleton(),
-							FOnCurvePicked::CreateLambda([this](const FName& InCurve)
-								{
-									const FScopedTransaction Transaction(LOCTEXT("AddCurvesToSet", "Add Curve to Set"));
-									SetBinding->Modify();
-
-									FAnimationAttributeIdentifier CurveAttrId(InCurve, INDEX_NONE, NAME_None, FFloatAnimationAttribute::StaticStruct());
-									SetBinding->AddAttributeToSet(CurveAttrId, NAME_None);
-
-									FSlateApplication::Get().DismissAllMenus();
-									RepopulateListData();
-								})),
-						FText(),
-						true,
-						false,
-						true
-					);
-
-					Section.AddEntry(WidgetEntry);
+					FPersonaModule& PersonaModule = FModuleManager::LoadModuleChecked<FPersonaModule>("Persona");
+
+					TSharedPtr<SAttributesList> SetBindingWidget = Context->SetBindingWidget.Pin();
+					ensure(SetBindingWidget);
+					
+					const TWeakObjectPtr<UAbstractSkeletonSetBinding> SetBinding = SetBindingWidget->SetBinding;
+
+					if (SetBinding.IsValid())
+					{
+						FToolMenuEntry WidgetEntry = FToolMenuEntry::InitWidget(
+							"CurvePicker",
+							PersonaModule.CreateCurvePicker(
+								SetBinding->GetSkeleton(),
+								FOnCurvePicked::CreateLambda([SetBinding, SetBindingWidget = Context->SetBindingWidget](const FName& InCurve)
+									{
+										const FScopedTransaction Transaction(LOCTEXT("AddCurvesToSet", "Add Curve to Set"));
+										SetBinding->Modify();
+
+										FAnimationAttributeIdentifier CurveAttrId(InCurve, INDEX_NONE, NAME_None, FFloatAnimationAttribute::StaticStruct());
+										SetBinding->AddAttributeToSet(CurveAttrId, NAME_None);
+
+										FSlateApplication::Get().DismissAllMenus();
+										SetBindingWidget.Pin()->RepopulateListData();
+									})),
+							FText(),
+							true,
+							false,
+							true
+						);
+
+						Section.AddEntry(WidgetEntry);
+					}
 				}
 			}));
 	}
@@ -333,7 +349,11 @@ void SAttributesList::RegisterMenus()
 
 TSharedRef<SWidget> SAttributesList::CreateAddAttributeWidget()
 {
+	UAttributesListMenuContext* ListContext = NewObject<UAttributesListMenuContext>();
+	ListContext->SetBindingWidget = SharedThis(this);
+	
 	FToolMenuContext MenuContext;
+	MenuContext.AddObject(ListContext);
 
 	return UToolMenus::Get()->GenerateWidget(FMenus::AddAttributeId, MenuContext);
 }
diff --git a/Engine/Plugins/Experimental/UAF/UAF/Source/UAFEditor/Private/UAF/AbstractSkeleton/Sets/SAttributesList.h b/Engine/Plugins/Experimental/UAF/UAF/Source/UAFEditor/Private/UAF/AbstractSkeleton/Sets/SAttributesList.h
--- a/Engine/Plugins/Experimental/UAF/UAF/Source/UAFEditor/Private/UAF/AbstractSkeleton/Sets/SAttributesList.h
+++ b/Engine/Plugins/Experimental/UAF/UAF/Source/UAFEditor/Private/UAF/AbstractSkeleton/Sets/SAttributesList.h
@@ -7,6 +7,8 @@
 #include "Widgets/SCompoundWidget.h"
 #include "Widgets/Views/SListView.h"
 
+#include "SAttributesList.generated.h"
+
 class UAbstractSkeletonSetBinding;
 
 namespace UE::UAF::Sets
@@ -67,4 +69,13 @@ namespace UE::UAF::Sets
 		bool bRepopulating = false;
 	};
 
-}
\ No newline at end of file
+}
+
+UCLASS()
+class UAttributesListMenuContext : public UObject
+{
+	GENERATED_BODY()
+	
+public:
+	TWeakPtr<UE::UAF::Sets::SAttributesList> SetBindingWidget;
+};
\ No newline at end of file
