Fix set collection editor not updating when renaming a set and add rename to context menu

diff --git a/Engine/Plugins/Experimental/UAF/UAF/Source/UAFEditor/Private/UAF/AbstractSkeleton/Sets/SetCollectionEditor.cpp b/Engine/Plugins/Experimental/UAF/UAF/Source/UAFEditor/Private/UAF/AbstractSkeleton/Sets/SetCollectionEditor.cpp
--- a/Engine/Plugins/Experimental/UAF/UAF/Source/UAFEditor/Private/UAF/AbstractSkeleton/Sets/SetCollectionEditor.cpp
+++ b/Engine/Plugins/Experimental/UAF/UAF/Source/UAFEditor/Private/UAF/AbstractSkeleton/Sets/SetCollectionEditor.cpp
@@ -304,6 +304,7 @@ namespace UE::UAF
 	{
 		FMenuBuilder MenuBuilder(true, CommandList.ToSharedRef());
 
+		MenuBuilder.AddMenuEntry(FGenericCommands::Get().Rename);
 		MenuBuilder.AddMenuEntry(FGenericCommands::Get().Delete);
 
 		return MenuBuilder.MakeWidget();
@@ -417,6 +418,8 @@ namespace UE::UAF
 				TreeView->SetSelection(TreeItem);
 			}
 		}
+
+		TreeView->RequestTreeRefresh();
 	}
 
 	void SSetCollection::ExpandAllTreeItems()
@@ -431,6 +434,11 @@ namespace UE::UAF
 	{
 		CommandList = MakeShareable(new FUICommandList);
 
+		CommandList->MapAction(
+			FGenericCommands::Get().Rename,
+			FExecuteAction::CreateSP(this, &SSetCollection::OnRenameSet)
+		);
+	
 		CommandList->MapAction(
 			FGenericCommands::Get().Delete,
 			FExecuteAction::CreateSP(this, &SSetCollection::OnRemoveSet)
@@ -505,6 +513,17 @@ namespace UE::UAF
 
 		return nullptr;
 	}
+	
+	void SSetCollection::OnRenameSet()
+	{
+		TArray<TSharedPtr<FTreeItem>> Selected = TreeView->GetSelectedItems();
+		if (!Selected.IsEmpty())
+		{		
+			ItemToRename = Selected[0];
+        	TreeView->RequestScrollIntoView(ItemToRename);
+        }
+	}
+	
 }
 
 #undef LOCTEXT_NAMESPACE
\ No newline at end of file
diff --git a/Engine/Plugins/Experimental/UAF/UAF/Source/UAFEditor/Private/UAF/AbstractSkeleton/Sets/SetCollectionEditor.h b/Engine/Plugins/Experimental/UAF/UAF/Source/UAFEditor/Private/UAF/AbstractSkeleton/Sets/SetCollectionEditor.h
--- a/Engine/Plugins/Experimental/UAF/UAF/Source/UAFEditor/Private/UAF/AbstractSkeleton/Sets/SetCollectionEditor.h
+++ b/Engine/Plugins/Experimental/UAF/UAF/Source/UAFEditor/Private/UAF/AbstractSkeleton/Sets/SetCollectionEditor.h
@@ -73,6 +73,8 @@ namespace UE::UAF
 		
 		void TreeView_OnMouseButtonDoubleClick(TSharedPtr<FTreeItem> InItem);
 
+		void OnRenameSet();
+		
 		void OnRemoveSet();
 
 		void RepopulateTreeData();
