UAFEditor - fix for crash in FAnimNextModuleTrack::UpdateInternal when DetailsObject needs to be reinitialialized

diff --git a/Engine/Plugins/Experimental/UAF/UAF/Source/UAFEditor/Private/RewindDebugger/AnimNextModuleTrack.cpp b/Engine/Plugins/Experimental/UAF/UAF/Source/UAFEditor/Private/RewindDebugger/AnimNextModuleTrack.cpp
--- a/Engine/Plugins/Experimental/UAF/UAF/Source/UAFEditor/Private/RewindDebugger/AnimNextModuleTrack.cpp
+++ b/Engine/Plugins/Experimental/UAF/UAF/Source/UAFEditor/Private/RewindDebugger/AnimNextModuleTrack.cpp
@@ -129,6 +129,9 @@ bool FAnimNextModuleTrack::UpdateInternal()
 		{
 			// this should not happen unless the object was garbage collected (which should not happen since it's marked as Standalone)
 			Initialize();
+
+			DetailsObject = DetailsObjectWeakPtr.Get();
+			check(DetailsObject);
 		}
 
 		if(InstanceId != 0)
