Fix slow document navigation in workspace editor
* Avoid use of ExtendMenu as causes mass widget invalidation which causes a second long hitch.
* Fix compile button missing after using the state tree editor.

diff --git a/Engine/Plugins/Experimental/UAF/UAF/Source/UAFEditor/Private/AnimNextEditorContext.h b/Engine/Plugins/Experimental/UAF/UAF/Source/UAFEditor/Private/AnimNextEditorContext.h
deleted file mode 100644
diff --git a/Engine/Plugins/Experimental/UAF/UAF/Source/UAFEditor/Private/AnimNextEditorModule.cpp b/Engine/Plugins/Experimental/UAF/UAF/Source/UAFEditor/Private/AnimNextEditorModule.cpp
--- a/Engine/Plugins/Experimental/UAF/UAF/Source/UAFEditor/Private/AnimNextEditorModule.cpp
+++ b/Engine/Plugins/Experimental/UAF/UAF/Source/UAFEditor/Private/AnimNextEditorModule.cpp
@@ -33,7 +33,6 @@
 #include "Param/ObjectFunctionLocatorEditor.h"
 #include "Param/ObjectPropertyLocatorEditor.h"
 #include "AnimNextAssetWorkspaceAssetUserData.h"
-#include "AnimNextEditorContext.h"
 #include "AssetCompilationHandler.h"
 #include "ContentBrowserModule.h"
 #include "ContextObjectStore.h"
@@ -300,128 +299,129 @@ void FAnimNextEditorModule::StartupModule()
 		}
 	}));
 
-	WorkspaceEditorModule.OnExtendToolMenuContext().AddLambda([](TSharedPtr<UE::Workspace::IWorkspaceEditor> InWorkspaceEditor, FToolMenuContext& InContext)
+	UToolMenu* WorkspaceEditorToolbar = UToolMenus::Get()->ExtendMenu("AssetEditor.WorkspaceEditor.Toolbar");
 	{
-		UContextObjectStore* ContextStore = InWorkspaceEditor->GetEditorModeManager().GetInteractiveToolsContext()->ContextObjectStore;
-		UAnimNextEditorContext* AnimNextEditorContext = ContextStore->FindContext<UAnimNextEditorContext>();
-		if (!AnimNextEditorContext)
+		const auto Handler = FNewToolMenuDelegate::CreateLambda([](UToolMenu* InMenu)
 		{
-			AnimNextEditorContext = NewObject<UAnimNextEditorContext>();
-			AnimNextEditorContext->WeakWorkspaceEditor = InWorkspaceEditor;
-			ContextStore->AddContextObject(AnimNextEditorContext);
-		}
-
-		InContext.AddObject(AnimNextEditorContext);
-
-		// [TODO] This should be a dynamic extension
-		if (UToolMenu* Menu = UToolMenus::Get()->ExtendMenu(InWorkspaceEditor->GetToolMenuToolbarName()))
-		{
-			FToolMenuSection& RigVMOperationsSection = Menu->AddSection("RigVMOperations", TAttribute<FText>(), FToolMenuInsert("WorkspaceOperations", EToolMenuInsertType::After));
-
-			auto GetIcon = [WeakWorkspaceEditor = TWeakPtr<UE::Workspace::IWorkspaceEditor>(InWorkspaceEditor)]()
+			if (UAssetEditorToolkitMenuContext* Context = InMenu->FindContext<UAssetEditorToolkitMenuContext>())
 			{
-				static const FName CompileStatusBackground("AssetEditor.CompileStatus.Background");
-				static const FName CompileStatusUnknown("AssetEditor.CompileStatus.Overlay.Unknown");
-				static const FName CompileStatusError("AssetEditor.CompileStatus.Overlay.Error");
-				static const FName CompileStatusGood("AssetEditor.CompileStatus.Overlay.Good");
-				static const FName CompileStatusWarning("AssetEditor.CompileStatus.Overlay.Warning");
-
-				TSharedPtr<UE::Workspace::IWorkspaceEditor> WorkspaceEditor = WeakWorkspaceEditor.Pin();
-				if(!WorkspaceEditor.IsValid())
+				if (TSharedPtr<FAssetEditorToolkit> ToolkitShared = Context->Toolkit.Pin())
 				{
-					return FSlateIcon();
-				}
+					TSharedPtr<Workspace::IWorkspaceEditor> WorkspaceEditor = StaticCastSharedPtr<Workspace::IWorkspaceEditor>(ToolkitShared);
 
-				UAnimNextWorkspaceEditorMode* EditorMode = Cast<UAnimNextWorkspaceEditorMode>(WorkspaceEditor->GetEditorModeManager().GetActiveScriptableMode(UAnimNextWorkspaceEditorMode::EM_AnimNextWorkspace));
-				if(EditorMode == nullptr)
-				{
-					return FSlateIcon(FAppStyle::Get().GetStyleSetName(), CompileStatusBackground, NAME_None, CompileStatusUnknown);
-				}
-
-				ECompileStatus Status = EditorMode->GetLatestCompileStatus();
-				switch (Status)
-				{
-				default:
-				case ECompileStatus::Unknown:
-				case ECompileStatus::Dirty:
-					return FSlateIcon(FAppStyle::GetAppStyleSetName(), CompileStatusBackground, NAME_None, CompileStatusUnknown);
-				case ECompileStatus::Error:
-					return FSlateIcon(FAppStyle::GetAppStyleSetName(), CompileStatusBackground, NAME_None, CompileStatusError);
-				case ECompileStatus::UpToDate:
-					return FSlateIcon(FAppStyle::GetAppStyleSetName(), CompileStatusBackground, NAME_None, CompileStatusGood);
-				case ECompileStatus::Warning:
-					return FSlateIcon(FAppStyle::GetAppStyleSetName(), CompileStatusBackground, NAME_None, CompileStatusWarning);
-				}
-			};
+					auto GetIcon = [WeakWorkspaceEditor = TWeakPtr<UE::Workspace::IWorkspaceEditor>(WorkspaceEditor)]()
+					{
+						static const FName CompileStatusBackground("AssetEditor.CompileStatus.Background");
+						static const FName CompileStatusUnknown("AssetEditor.CompileStatus.Overlay.Unknown");
+						static const FName CompileStatusError("AssetEditor.CompileStatus.Overlay.Error");
+						static const FName CompileStatusGood("AssetEditor.CompileStatus.Overlay.Good");
+						static const FName CompileStatusWarning("AssetEditor.CompileStatus.Overlay.Warning");
+
+						TSharedPtr<UE::Workspace::IWorkspaceEditor> WorkspaceEditor = WeakWorkspaceEditor.Pin();
+						if(!WorkspaceEditor.IsValid())
+						{
+							return FSlateIcon();
+						}
 
-			auto GetTooltip = [WeakWorkspaceEditor = TWeakPtr<UE::Workspace::IWorkspaceEditor>(InWorkspaceEditor)]()
-			{
-				TSharedPtr<UE::Workspace::IWorkspaceEditor> WorkspaceEditor = WeakWorkspaceEditor.Pin();
-				if(!WorkspaceEditor.IsValid())
-				{
-					return LOCTEXT("CompileGenericTooltip", "Compile all relevant assets");
-				}
+						UAnimNextWorkspaceEditorMode* EditorMode = Cast<UAnimNextWorkspaceEditorMode>(WorkspaceEditor->GetEditorModeManager().GetActiveScriptableMode(UAnimNextWorkspaceEditorMode::EM_AnimNextWorkspace));
+						if(EditorMode == nullptr)
+						{
+							return FSlateIcon(FAppStyle::Get().GetStyleSetName(), CompileStatusBackground, NAME_None, CompileStatusUnknown);
+						}
 
-				UAnimNextWorkspaceEditorMode* EditorMode = Cast<UAnimNextWorkspaceEditorMode>(WorkspaceEditor->GetEditorModeManager().GetActiveScriptableMode(UAnimNextWorkspaceEditorMode::EM_AnimNextWorkspace));
-				if(EditorMode == nullptr)
-				{
-					return LOCTEXT("CompileGenericTooltip", "Compile all relevant assets");
-				}
+						ECompileStatus Status = EditorMode->GetLatestCompileStatus();
+						switch (Status)
+						{
+						default:
+						case ECompileStatus::Unknown:
+						case ECompileStatus::Dirty:
+							return FSlateIcon(FAppStyle::GetAppStyleSetName(), CompileStatusBackground, NAME_None, CompileStatusUnknown);
+						case ECompileStatus::Error:
+							return FSlateIcon(FAppStyle::GetAppStyleSetName(), CompileStatusBackground, NAME_None, CompileStatusError);
+						case ECompileStatus::UpToDate:
+							return FSlateIcon(FAppStyle::GetAppStyleSetName(), CompileStatusBackground, NAME_None, CompileStatusGood);
+						case ECompileStatus::Warning:
+							return FSlateIcon(FAppStyle::GetAppStyleSetName(), CompileStatusBackground, NAME_None, CompileStatusWarning);
+						}
+					};
 
-				FTextBuilder TextBuilder;
-				if(EditorMode->GetState().bCompileDirtyFiles)
-				{
-					if(EditorMode->GetState().bCompileWholeWorkspace)
-					{
-						TextBuilder.AppendLine(LOCTEXT("CompileWorkspaceDirtyTooltip", "Compile all assets in the workspace that are dirty or have errors"));
-					}
-					else
+					auto GetTooltip = [WeakWorkspaceEditor = TWeakPtr<UE::Workspace::IWorkspaceEditor>(WorkspaceEditor)]()
 					{
-						TextBuilder.AppendLine(LOCTEXT("CompileDirtyTooltip", "Compile the current asset if it is dirty or has errors"));
-					}
-				}
-				else
-				{
-					if(EditorMode->GetState().bCompileWholeWorkspace)
-					{
-						TextBuilder.AppendLine(LOCTEXT("CompileWorkspaceTooltip", "Compile all assets in the workspace"));
-					}
-					else
-					{
-						TextBuilder.AppendLine(LOCTEXT("CompileCurrentTooltip", "Compile the current asset"));
-					}
-				}
+						TSharedPtr<UE::Workspace::IWorkspaceEditor> WorkspaceEditor = WeakWorkspaceEditor.Pin();
+						if(!WorkspaceEditor.IsValid())
+						{
+							return LOCTEXT("CompileGenericTooltip", "Compile all relevant assets");
+						}
 
-				if(EditorMode->GetState().bAutoCompile)
-				{
-					TextBuilder.AppendLine(LOCTEXT("AutoCompileEnabledTooltip", "Auto-compilation is enabled"));
-				}
-				else
-				{
-					TextBuilder.AppendLine(LOCTEXT("AutoCompileDisabledTooltip", "Auto-compilation is disabled"));
-				}
+						UAnimNextWorkspaceEditorMode* EditorMode = Cast<UAnimNextWorkspaceEditorMode>(WorkspaceEditor->GetEditorModeManager().GetActiveScriptableMode(UAnimNextWorkspaceEditorMode::EM_AnimNextWorkspace));
+						if(EditorMode == nullptr)
+						{
+							return LOCTEXT("CompileGenericTooltip", "Compile all relevant assets");
+						}
 
-				return TextBuilder.ToText();
-			};
+						FTextBuilder TextBuilder;
+						if(EditorMode->GetState().bCompileDirtyFiles)
+						{
+							if(EditorMode->GetState().bCompileWholeWorkspace)
+							{
+								TextBuilder.AppendLine(LOCTEXT("CompileWorkspaceDirtyTooltip", "Compile all assets in the workspace that are dirty or have errors"));
+							}
+							else
+							{
+								TextBuilder.AppendLine(LOCTEXT("CompileDirtyTooltip", "Compile the current asset if it is dirty or has errors"));
+							}
+						}
+						else
+						{
+							if(EditorMode->GetState().bCompileWholeWorkspace)
+							{
+								TextBuilder.AppendLine(LOCTEXT("CompileWorkspaceTooltip", "Compile all assets in the workspace"));
+							}
+							else
+							{
+								TextBuilder.AppendLine(LOCTEXT("CompileCurrentTooltip", "Compile the current asset"));
+							}
+						}
 
-			RigVMOperationsSection.AddEntry(FToolMenuEntry::InitToolBarButton(FRigVMCommands::Get().Compile, TAttribute<FText>(), MakeAttributeLambda(GetTooltip), MakeAttributeLambda(GetIcon)));
-			RigVMOperationsSection.AddEntry(FToolMenuEntry::InitComboButton(
-				"CompileOptionsCombo",
-				FUIAction(),
-				FNewToolMenuDelegate::CreateLambda([](UToolMenu* InToolMenu)
-				{
-					FToolMenuSection& CompileOptionsSection = InToolMenu->AddSection("CompileOptionsSection", LOCTEXT("CompileOptionsComboLabel", "Compilation Options"));
-					CompileOptionsSection.AddEntry(FToolMenuEntry::InitMenuEntry(FRigVMCommands::Get().AutoCompile, TAttribute<FText>(), TAttribute<FText>(), FSlateIcon("RigVMEditorStyle", "RigVM.AutoCompileGraph")));
-					CompileOptionsSection.AddEntry(FToolMenuEntry::InitMenuEntry(FRigVMCommands::Get().CompileWholeWorkspace, TAttribute<FText>(), TAttribute<FText>(), FSlateIcon(FAppStyle::GetAppStyleSetName(), "Icons.World")));
-					CompileOptionsSection.AddEntry(FToolMenuEntry::InitMenuEntry(FRigVMCommands::Get().CompileDirtyFiles, TAttribute<FText>(), TAttribute<FText>(), FSlateIcon(FAppStyle::GetAppStyleSetName(), "Icons.DirtyBadge")));
-				}),
-				FText::GetEmpty(),
-				LOCTEXT("CompileOptionsComboTooltip", "Compilation Options"),
-				FSlateIcon(),
-				true));
-		}
-	});
+						if(EditorMode->GetState().bAutoCompile)
+						{
+							TextBuilder.AppendLine(LOCTEXT("AutoCompileEnabledTooltip", "Auto-compilation is enabled"));
+						}
+						else
+						{
+							TextBuilder.AppendLine(LOCTEXT("AutoCompileDisabledTooltip", "Auto-compilation is disabled"));
+						}
+
+						return TextBuilder.ToText();
+					};
+			
+					FToolMenuEntry CompileEntry = FToolMenuEntry::InitToolBarButton(FRigVMCommands::Get().Compile, TAttribute<FText>(), MakeAttributeLambda(GetTooltip), MakeAttributeLambda(GetIcon));
+
+					FToolMenuEntry CompileOptionsEntry = FToolMenuEntry::InitComboButton(
+						"CompileOptionsCombo",
+						FUIAction(),
+						FNewToolMenuDelegate::CreateLambda([](UToolMenu* InToolMenu)
+						{
+							FToolMenuSection& CompileOptionsSection = InToolMenu->AddSection("CompileOptionsSection", LOCTEXT("CompileOptionsComboLabel", "Compilation Options"));
+							CompileOptionsSection.AddEntry(FToolMenuEntry::InitMenuEntry(FRigVMCommands::Get().AutoCompile, TAttribute<FText>(), TAttribute<FText>(), FSlateIcon("RigVMEditorStyle", "RigVM.AutoCompileGraph")));
+							CompileOptionsSection.AddEntry(FToolMenuEntry::InitMenuEntry(FRigVMCommands::Get().CompileWholeWorkspace, TAttribute<FText>(), TAttribute<FText>(), FSlateIcon(FAppStyle::GetAppStyleSetName(), "Icons.World")));
+							CompileOptionsSection.AddEntry(FToolMenuEntry::InitMenuEntry(FRigVMCommands::Get().CompileDirtyFiles, TAttribute<FText>(), TAttribute<FText>(), FSlateIcon(FAppStyle::GetAppStyleSetName(), "Icons.DirtyBadge")));
+						}),
+						FText::GetEmpty(),
+						LOCTEXT("CompileOptionsComboTooltip", "Compilation Options"),
+						FSlateIcon(),
+						true);
+
+					FToolMenuSection& CompileSection = InMenu->FindOrAddSection("CompileOptions");
+					CompileSection.AddEntry(CompileEntry);
+					CompileSection.AddEntry(CompileOptionsEntry);
+				}
+			}
+		});
 	
+		WorkspaceEditorToolbar->AddDynamicSection("RigVMOperations", Handler, FToolMenuInsert("WorkspaceOperations", EToolMenuInsertType::After));
+	}
+
 	RegisterWorkspaceDocumentTypes(WorkspaceEditorModule);
 
 	WorkspaceEditorModule.OnRegisterWorkspaceDetailsCustomization().AddLambda([](const TWeakPtr<Workspace::IWorkspaceEditor>& InWorkspaceEditor, TSharedPtr<IDetailsView>& InDetailsView)
diff --git a/Engine/Plugins/Experimental/UAF/UAF/Source/UAFEditor/Private/Workspace/AnimNextWorkspaceEditorMode.cpp b/Engine/Plugins/Experimental/UAF/UAF/Source/UAFEditor/Private/Workspace/AnimNextWorkspaceEditorMode.cpp
--- a/Engine/Plugins/Experimental/UAF/UAF/Source/UAFEditor/Private/Workspace/AnimNextWorkspaceEditorMode.cpp
+++ b/Engine/Plugins/Experimental/UAF/UAF/Source/UAFEditor/Private/Workspace/AnimNextWorkspaceEditorMode.cpp
@@ -2,12 +2,12 @@
 
 #include "AnimNextWorkspaceEditorMode.h"
 
-#include "AnimNextEditorContext.h"
 #include "AnimNextEditorModule.h"
 #include "AnimNextRigVMAsset.h"
 #include "AnimNextRigVMAssetEditorData.h"
 #include "AnimNextScopedCompileJob.h"
 #include "ContextObjectStore.h"
+#include "EditorModeManager.h"
 #include "IAssetCompilationHandler.h"
 #include "InteractiveToolManager.h"
 #include "IWorkspaceEditor.h"
@@ -16,7 +16,9 @@
 #include "UncookedOnlyUtils.h"
 #include "Logging/MessageLog.h"
 #include "Modules/ModuleManager.h"
+#include "Toolkits/AssetEditorToolkitMenuContext.h"
 #include "Toolkits/BaseToolkit.h"
+#include "Tools/EdModeInteractiveToolsContext.h"
 
 #include UE_INLINE_GENERATED_CPP_BY_NAME(AnimNextWorkspaceEditorMode)
 
@@ -35,7 +37,7 @@ UAnimNextWorkspaceEditorMode::UAnimNextWorkspaceEditorMode()
 void UAnimNextWorkspaceEditorMode::Enter()
 {
 	Super::Enter();
-
+	
 	TSharedPtr<UE::Workspace::IWorkspaceEditor> WorkspaceEditor = GetWorkspaceEditor();
 	if(!WorkspaceEditor.IsValid())
 	{
@@ -363,11 +365,14 @@ void UAnimNextWorkspaceEditorMode::HandleRigVMModifiedEvent(ERigVMGraphNotifType
 
 TSharedPtr<UE::Workspace::IWorkspaceEditor> UAnimNextWorkspaceEditorMode::GetWorkspaceEditor() const
 {
-	if (const UContextObjectStore* ContextObjectStore = GetToolManager()->GetContextObjectStore())
+	if (const UContextObjectStore* ContextObjectStore = GetModeManager()->GetInteractiveToolsContext()->ContextObjectStore)
 	{
-		if (const UAnimNextEditorContext* Context = ContextObjectStore->FindContext<UAnimNextEditorContext>())
+		if (UAssetEditorToolkitMenuContext* Context = ContextObjectStore->FindContext<UAssetEditorToolkitMenuContext>())
 		{
-			return Context->WeakWorkspaceEditor.Pin();
+			if (TSharedPtr<FAssetEditorToolkit> ToolkitShared = Context->Toolkit.Pin())
+			{
+				return StaticCastSharedPtr<UE::Workspace::IWorkspaceEditor>(ToolkitShared);
+			}
 		}
 	}
 
diff --git a/Engine/Plugins/Experimental/Workspace/Source/WorkspaceEditor/Private/SWorkspaceTabWrapper.cpp b/Engine/Plugins/Experimental/Workspace/Source/WorkspaceEditor/Private/SWorkspaceTabWrapper.cpp
--- a/Engine/Plugins/Experimental/Workspace/Source/WorkspaceEditor/Private/SWorkspaceTabWrapper.cpp
+++ b/Engine/Plugins/Experimental/Workspace/Source/WorkspaceEditor/Private/SWorkspaceTabWrapper.cpp
@@ -33,40 +33,58 @@ void SWorkspaceTabWrapper::Construct( const FArguments& InArgs, TSharedPtr<class
     const FSlateBrush* BreadcrumbButtonImage = FAppStyle::GetBrush("BreadcrumbTrail.Delimiter");
 
 	UToolMenus* ToolMenus = UToolMenus::Get();
-	UToolMenu* ToolBarMenu = nullptr;
-	if(ToolMenus)
+
+	const FName ToolbarName = "WorkspaceTabWrapperToolbar";
+
+	if (!UToolMenus::Get()->IsMenuRegistered(ToolbarName))
 	{
-		const FName ToolbarName = "WorkspaceTabWrapperToolbar";		
-		if (ToolMenus->IsMenuRegistered(ToolbarName))
-		{
-			ToolBarMenu = ToolMenus->ExtendMenu(ToolbarName);
-		}
-		else
-		{
-			ToolBarMenu = ToolMenus->RegisterMenu(ToolbarName, NAME_None, EMultiBoxType::SlimHorizontalToolBar);
-			ToolBarMenu->StyleName = "AssetEditorToolbar";
-		}
+		UToolMenu* ToolBarMenu = UToolMenus::Get()->RegisterMenu(ToolbarName, NAME_None, EMultiBoxType::SlimHorizontalToolBar);
+		ToolBarMenu->StyleName = "AssetEditorToolbar";
 		
-		FToolMenuSection& Section = ToolBarMenu->AddSection("AssetActions");
-		Section.AddMenuEntry("Save", TAttribute<FText>(), LOCTEXT("SaveButtonTooltip", "Save Asset"), FSlateIcon(FAppStyle::Get().GetStyleSetName(), "AssetEditor.SaveAsset"),
-			FUIAction(FExecuteAction::CreateSP(this, &SWorkspaceTabWrapper::ExecuteSave), FCanExecuteAction::CreateSP(this, &SWorkspaceTabWrapper::CanExecuteSave), FIsActionChecked(), FIsActionButtonVisible::CreateSP(this, &SWorkspaceTabWrapper::IsSaveButtonVisible)));
-	
-		const FToolMenuEntry FindInContentBrowserEntry = FToolMenuEntry::InitToolBarButton(
-			"FindInContentBrowser",
-			FUIAction(FExecuteAction::CreateSPLambda(this, [this]() 
-				{
-					if (UObject* Asset = WeakDocumentObject.Get())
-					{
-						GEditor->SyncBrowserToObject(Asset);
-					}
-				})),
-			FText::GetEmpty(),
-			LOCTEXT("FindInContentBrowserTooltip", "Finds this asset in the content browser"),
-			FSlateIcon(FAppStyle::Get().GetStyleSetName(), "SystemWideCommands.FindInContentBrowser")
-		);
-		Section.AddEntry(FindInContentBrowserEntry);
+		auto Handler = FNewToolMenuDelegate::CreateLambda([](UToolMenu* InMenu)
+		{
+			if (UWorkspaceTabMenuContext* Context = InMenu->FindContext<UWorkspaceTabMenuContext>())
+			{
+				TSharedRef<SWorkspaceTabWrapper> This = Context->TabWrapper.Pin().ToSharedRef();
+				
+				const FToolMenuEntry SaveEntry = FToolMenuEntry::InitToolBarButton(
+					"Save",
+					FUIAction(
+						FExecuteAction::CreateSP(This, &SWorkspaceTabWrapper::ExecuteSave),
+						FCanExecuteAction::CreateSP(This, &SWorkspaceTabWrapper::CanExecuteSave),
+						FIsActionChecked(),
+						FIsActionButtonVisible::CreateSP(This, &SWorkspaceTabWrapper::IsSaveButtonVisible)),
+					FText::GetEmpty(),
+					FText::GetEmpty(),
+					FSlateIcon(FAppStyle::GetAppStyleSetName(), "Icons.Save"));
+
+				const FToolMenuEntry FindInContentBrowserEntry = FToolMenuEntry::InitToolBarButton(
+					"FindInContentBrowser",
+					FUIAction(FExecuteAction::CreateSPLambda(This, [This]() 
+						{
+							if (UObject* Asset = This->WeakDocumentObject.Get())
+							{
+								GEditor->SyncBrowserToObject(Asset);
+							}
+						})),
+					FText::GetEmpty(),
+					LOCTEXT("FindInContentBrowserTooltip", "Finds this asset in the content browser"),
+					FSlateIcon(FAppStyle::Get().GetStyleSetName(), "SystemWideCommands.FindInContentBrowser"));
+
+				FToolMenuSection& Section = InMenu->AddSection("AssetSection");
+				Section.AddEntry(SaveEntry);
+				Section.AddEntry(FindInContentBrowserEntry);
+			}
+		});
+
+		ToolBarMenu->AddDynamicSection("AssetActions", Handler);
 	}
 	
+	UWorkspaceTabMenuContext* WorkspaceTabContext = NewObject<UWorkspaceTabMenuContext>();
+	WorkspaceTabContext->TabWrapper = SharedThis(this);
+
+	const FToolMenuContext MenuContext(WorkspaceTabContext);
+	
 	ChildSlot
 	[
 		SNew(SVerticalBox)
@@ -201,7 +219,7 @@ void SWorkspaceTabWrapper::Construct( const FArguments& InArgs, TSharedPtr<class
 			+SHorizontalBox::Slot()
 			.AutoWidth()
 			[
-				ToolMenus && ToolBarMenu ? ToolMenus->GenerateWidget(ToolBarMenu) : SNullWidget::NullWidget	
+				UToolMenus::Get()->GenerateWidget(ToolbarName, MenuContext)
 			]
 		]
 		+ SVerticalBox::Slot()
diff --git a/Engine/Plugins/Experimental/Workspace/Source/WorkspaceEditor/Private/SWorkspaceTabWrapper.h b/Engine/Plugins/Experimental/Workspace/Source/WorkspaceEditor/Private/SWorkspaceTabWrapper.h
--- a/Engine/Plugins/Experimental/Workspace/Source/WorkspaceEditor/Private/SWorkspaceTabWrapper.h
+++ b/Engine/Plugins/Experimental/Workspace/Source/WorkspaceEditor/Private/SWorkspaceTabWrapper.h
@@ -9,6 +9,17 @@
 #include "WorkspaceEditor.h"
 #include "WorkspaceEditorModule.h"
 
+#include "SWorkspaceTabWrapper.generated.h"
+
+UCLASS()
+class UWorkspaceTabMenuContext : public UObject
+{
+	GENERATED_BODY()
+
+public:
+	TWeakPtr<SWorkspaceTabWrapper> TabWrapper;
+};
+
 class SWorkspaceTabWrapper : public SCompoundWidget
 {
 public:
diff --git a/Engine/Plugins/Experimental/Workspace/Source/WorkspaceEditor/Private/WorkspaceEditor.cpp b/Engine/Plugins/Experimental/Workspace/Source/WorkspaceEditor/Private/WorkspaceEditor.cpp
--- a/Engine/Plugins/Experimental/Workspace/Source/WorkspaceEditor/Private/WorkspaceEditor.cpp
+++ b/Engine/Plugins/Experimental/Workspace/Source/WorkspaceEditor/Private/WorkspaceEditor.cpp
@@ -26,6 +26,7 @@
 #include "WorkspaceSchema.h"
 #include "WorkspaceTabPayload.h"
 #include "AdvancedPreviewScene.h"
+#include "ContextObjectStore.h"
 #include "WorkspaceAssetViewportClient.h"
 #include "WorkspaceViewportController.h"
 #include "Viewports.h"
@@ -205,6 +206,17 @@ void FWorkspaceEditor::CreateWidgets()
 
 void FWorkspaceEditor::PostInitAssetEditor()
 {
+	if (UContextObjectStore* ContextObjectStore = GetEditorModeManager().GetInteractiveToolsContext()->ContextObjectStore)
+	{
+		UAssetEditorToolkitMenuContext* AssetEditorContext = ContextObjectStore->FindContext<UAssetEditorToolkitMenuContext>();
+		if (!AssetEditorContext)
+		{
+			AssetEditorContext = NewObject<UAssetEditorToolkitMenuContext>();
+			AssetEditorContext->Toolkit = AsShared();
+			ContextObjectStore->AddContextObject(AssetEditorContext);
+		}
+	}
+	
 	Workspace->LoadState();
 	GetSchema()->OnLoadWorkspaceState(SharedThis(this), Workspace->GetState()->UserState);
 	RestoreEditedObjectState();
@@ -723,10 +735,6 @@ FLinearColor FWorkspaceEditor::GetWorldCentricTabColorScale() const
 
 void FWorkspaceEditor::InitToolMenuContext(FToolMenuContext& InMenuContext)
 {
-	UAssetEditorToolkitMenuContext* ToolkitMenuContext = NewObject<UAssetEditorToolkitMenuContext>();
-	ToolkitMenuContext->Toolkit = AsShared();
-	InMenuContext.AddObject(ToolkitMenuContext);
-	
 	IWorkspaceEditor::InitToolMenuContext(InMenuContext);
 
 	Workspace::IWorkspaceEditorModule& WorkspaceEditorModule = FModuleManager::Get().LoadModuleChecked<Workspace::IWorkspaceEditorModule>("WorkspaceEditor");
diff --git a/Engine/Plugins/Runtime/StateTree/Source/StateTreeEditorModule/Private/StateTreeEditorModeToolkit.cpp b/Engine/Plugins/Runtime/StateTree/Source/StateTreeEditorModule/Private/StateTreeEditorModeToolkit.cpp
--- a/Engine/Plugins/Runtime/StateTree/Source/StateTreeEditorModule/Private/StateTreeEditorModeToolkit.cpp
+++ b/Engine/Plugins/Runtime/StateTree/Source/StateTreeEditorModule/Private/StateTreeEditorModeToolkit.cpp
@@ -256,7 +256,7 @@ void FStateTreeEditorModeToolkit::ExtendSecondaryModeToolbar(UToolMenu* ToolBar)
 
 	static const FToolMenuInsert InsertLast(NAME_None, EToolMenuInsertType::Last);
 
-	FToolMenuSection& CompileSection = ToolBar->AddSection("Compile", FText(), InsertLast);
+	FToolMenuSection& CompileSection = ToolBar->AddSection("Compile2", FText(), InsertLast);
 
 	auto GetToolkitFromAssetEditorContext = [](const UAssetEditorToolkitMenuContext* InContext) -> TSharedPtr<FStateTreeEditorModeToolkit>
 	{
@@ -296,6 +296,7 @@ void FStateTreeEditorModeToolkit::ExtendSecondaryModeToolbar(UToolMenu* ToolBar)
 					return FSlateIcon();
 				}))
 			);
+			CompileButton.Name = "StateTreeCompile";
 			CompileButton.StyleNameOverride = "CalloutToolbar";
 
 			FToolMenuEntry& CompileOptions = InSection.AddEntry(FToolMenuEntry::InitComboButton(
@@ -426,7 +427,7 @@ void FStateTreeEditorModeToolkit::ExtendSecondaryModeToolbar(UToolMenu* ToolBar)
 	const FName StateTreeEditModeProfile = TEXT("StateTreeEditModeDisabledProfile");
 	FToolMenuProfile* ToolbarProfile = UToolMenus::Get()->AddRuntimeMenuProfile(ToolBar->GetMenuName(), StateTreeEditModeProfile);
 	{
-		ToolbarProfile->MenuPermissions.AddDenyListItem("CompileCommands", Commands.Compile->GetCommandName());
+		ToolbarProfile->MenuPermissions.AddDenyListItem("CompileCommands", "StateTreeCompile");
 		ToolbarProfile->MenuPermissions.AddDenyListItem("CreateNewNodes", "CreateNewTaskComboButton");
 		ToolbarProfile->MenuPermissions.AddDenyListItem("CreateNewNodes", "CreateNewConditionComboButton");
 		ToolbarProfile->MenuPermissions.AddDenyListItem("CreateNewNodes", "CreateNewConsiderationComboButton");
