Force a variable referenced outside its owning asset to be publicly exported when adding a Variable Node to the referencing graph

diff --git a/Engine/Plugins/Experimental/UAF/UAF/Source/UAFEditor/Private/Common/GraphEditorSchemaActions.cpp b/Engine/Plugins/Experimental/UAF/UAF/Source/UAFEditor/Private/Common/GraphEditorSchemaActions.cpp
--- a/Engine/Plugins/Experimental/UAF/UAF/Source/UAFEditor/Private/Common/GraphEditorSchemaActions.cpp
+++ b/Engine/Plugins/Experimental/UAF/UAF/Source/UAFEditor/Private/Common/GraphEditorSchemaActions.cpp
@@ -259,6 +259,15 @@ UEdGraphNode* FAnimNextSchemaAction_Variable::PerformAction(UEdGraph* ParentGrap
 				{
 					if (OwnerEditorData != VariableOwnerEditorData)
 					{
+						if (IAnimNextRigVMExportInterface* VariableExport = Cast<IAnimNextRigVMExportInterface>(VariableOwnerEditorData->FindEntry(InName)))
+						{
+							if (VariableExport->GetExportAccessSpecifier() != EAnimNextExportAccessSpecifier::Public)
+							{
+								UE::UAF::UncookedOnly::FScopedCompileJob CompilerResults(LOCTEXT("ForceVariablePublicJobName", "Force Variable Public Access"), { SourceObject });							
+								VariableExport->SetExportAccessSpecifier(EAnimNextExportAccessSpecifier::Public);
+							}
+						}
+						
 						if (!OwnerEditorData->FindEntry(VariableOwningRigVMAsset->GetFName()))
 						{
 							UE::UAF::UncookedOnly::FScopedCompileJob CompilerResults(LOCTEXT("AddSharedAssetJobName", "Add Shared Asset dependency"), { ParentGraph->GetTypedOuter<UAnimNextRigVMAsset>() });
