Add transaction for setting/clearing variable binding

diff --git a/Engine/Plugins/Experimental/UAF/UAF/Source/UAFEditor/Private/Variables/Outliner/VariablesOutlinerColumns.cpp b/Engine/Plugins/Experimental/UAF/UAF/Source/UAFEditor/Private/Variables/Outliner/VariablesOutlinerColumns.cpp
--- a/Engine/Plugins/Experimental/UAF/UAF/Source/UAFEditor/Private/Variables/Outliner/VariablesOutlinerColumns.cpp
+++ b/Engine/Plugins/Experimental/UAF/UAF/Source/UAFEditor/Private/Variables/Outliner/VariablesOutlinerColumns.cpp
@@ -382,11 +382,13 @@ class SVariablesOutlinerValue : public SCompoundWidget, public FNotifyHook
 								{
 									if (InState == ECheckBoxState::Checked)
 									{
+										FScopedTransaction Transaction(LOCTEXT("SetVariableBinding", "Set Variable Binding"));
 										Entry->SetBindingType(FAnimNextUniversalObjectLocatorBindingData::StaticStruct(), true);
 										LocatorPropertyView->GetPropertyHandle()->NotifyPostChange(EPropertyChangeType::ValueSet);
 									}
 									else if (InState == ECheckBoxState::Unchecked)
 									{
+										FScopedTransaction Transaction(LOCTEXT("ClearVariableBinding", "Clear Variable Binding"));
 										Entry->SetBindingType(nullptr, true);
 										LocatorPropertyView->GetPropertyHandle()->NotifyPostChange(EPropertyChangeType::ValueSet);
 									}
