Ensure UAF instance migration of variable data takes into account added/removed referenced asset/structs

diff --git a/Engine/Plugins/Experimental/UAF/UAF/Source/UAF/Private/UAFAssetInstance.cpp b/Engine/Plugins/Experimental/UAF/UAF/Source/UAF/Private/UAFAssetInstance.cpp
--- a/Engine/Plugins/Experimental/UAF/UAF/Source/UAF/Private/UAFAssetInstance.cpp
+++ b/Engine/Plugins/Experimental/UAF/UAF/Source/UAF/Private/UAFAssetInstance.cpp
@@ -19,7 +19,7 @@ void FUAFAssetInstance::InitializeVariables(const TSharedPtr<const UE::UAF::FVar
 #if WITH_EDITOR
 void FUAFAssetInstance::MigrateVariables()
 {
-	Variables.Migrate();
+	Variables.Migrate(*this);
 }
 #endif
 
diff --git a/Engine/Plugins/Experimental/UAF/UAF/Source/UAF/Private/Variables/UAFInstanceVariableData.cpp b/Engine/Plugins/Experimental/UAF/UAF/Source/UAF/Private/Variables/UAFInstanceVariableData.cpp
--- a/Engine/Plugins/Experimental/UAF/UAF/Source/UAF/Private/Variables/UAFInstanceVariableData.cpp
+++ b/Engine/Plugins/Experimental/UAF/UAF/Source/UAF/Private/Variables/UAFInstanceVariableData.cpp
@@ -19,115 +19,19 @@ void FUAFInstanceVariableData::Initialize(const FUAFAssetInstance& InInstance, c
 
 	Reset();
 
-	auto AddAsset = [this, &InOverrides, &InInstance](const UAnimNextRigVMAsset* InAsset)
-	{
-		TSharedPtr<FUAFInstanceVariableContainer> NewVariables;
-
-		// First check to see we have an override
-		if(InOverrides.IsValid())
-		{
-			for (int32 OverrideIndex = 0; OverrideIndex < InOverrides->Collection.Num(); ++OverrideIndex)
-			{
-				const FVariableOverrides& Override = InOverrides->Collection[OverrideIndex];
-				if (Override.IsOverrideFor(InAsset))
-				{
-					NewVariables = MakeShared<FUAFInstanceVariableContainer>(InInstance, InAsset, InOverrides, OverrideIndex);
-					OwnedVariableContainers.Add(NewVariables.ToSharedRef());
-					break;
-				}
-			}
-		}
-
-		// Check to see if this asset exists in our host chain. If so, we take a reference.
-		if (!NewVariables.IsValid())
-		{
-			const FUAFAssetInstance* CurrentHostInstance = InInstance.GetHost();
-			while (CurrentHostInstance)
-			{
-				if (const int32* SetIndexPtr = CurrentHostInstance->Variables.AllSharedVariableContainersMap.Find(InAsset))
-				{
-					NewVariables = CurrentHostInstance->Variables.AllSharedVariableContainers[*SetIndexPtr].Pin();
-					break;
-				}
-				CurrentHostInstance = CurrentHostInstance->GetHost();
-			}
-		}
-
-		// Otherwise, create a new instance that we own
-		if (!NewVariables.IsValid())
-		{
-			NewVariables = MakeShared<FUAFInstanceVariableContainer>(InInstance, InAsset);
-			OwnedVariableContainers.Add(NewVariables.ToSharedRef());
-		}
-
-		const int32 ReferencedAssetIndex = AllSharedVariableContainers.Emplace(NewVariables);
-		InternalSharedVariableContainers.Add(NewVariables);
-		AllSharedVariableContainersMap.Add(InAsset, ReferencedAssetIndex);
-		NumAllVariables += NewVariables->NumVariables;
-		NumInternalVariables += NewVariables->NumVariables;
-	};
-
-	auto AddStruct = [this, &InOverrides, &InInstance](const UScriptStruct* InStruct)
-	{
-		TSharedPtr<FUAFInstanceVariableContainer> NewVariables;
-
-		// First check to see we have an override
-		if(InOverrides.IsValid())
-		{
-			for (int32 OverrideIndex = 0; OverrideIndex < InOverrides->Collection.Num(); ++OverrideIndex)
-			{
-				const FVariableOverrides& Override = InOverrides->Collection[OverrideIndex];
-				if (Override.IsOverrideFor(InStruct))
-				{
-					NewVariables = MakeShared<FUAFInstanceVariableContainer>(InInstance, InStruct, InOverrides, OverrideIndex);
-					OwnedVariableContainers.Add(NewVariables.ToSharedRef());
-					break;
-				}
-			}
-		}
-
-		// Check to see if this asset exists in our host chain. If so, we take a reference.
-		if (!NewVariables.IsValid())
-		{
-			const FUAFAssetInstance* CurrentHostInstance = InInstance.GetHost();
-			while (CurrentHostInstance)
-			{
-				if (const int32* SetIndexPtr = CurrentHostInstance->Variables.AllSharedVariableContainersMap.Find(InStruct))
-				{
-					NewVariables = CurrentHostInstance->Variables.AllSharedVariableContainers[*SetIndexPtr].Pin();
-					break;
-				}
-				CurrentHostInstance = CurrentHostInstance->GetHost();
-			}
-		}
-
-		// Otherwise, create a new instance
-		if (!NewVariables.IsValid())
-		{
-			NewVariables = MakeShared<FUAFInstanceVariableContainer>(InInstance, InStruct);
-			OwnedVariableContainers.Add(NewVariables.ToSharedRef());
-		}
-
-		const int32 ReferencedAssetIndex = AllSharedVariableContainers.Emplace(NewVariables);
-		InternalSharedVariableContainers.Add(NewVariables);
-		AllSharedVariableContainersMap.Add(InStruct, ReferencedAssetIndex);
-		NumAllVariables += NewVariables->NumVariables;
-		NumInternalVariables += NewVariables->NumVariables;
-	};
-
 	// Add the primary asset's variables
-	AddAsset(InInstance.GetAsset<UAnimNextRigVMAsset>());
+	AddVariablesContainerForAsset(InInstance.GetAsset<UAnimNextRigVMAsset>(), InInstance, InOverrides);
 
 	// Add referenced asset variables
 	for (const UAnimNextRigVMAsset* ReferencedAsset : InInstance.GetAsset<UAnimNextRigVMAsset>()->GetReferencedVariableAssets())
 	{
-		AddAsset(ReferencedAsset);
+		AddVariablesContainerForAsset(ReferencedAsset, InInstance, InOverrides);
 	}
 
 	// Add referenced struct variables
 	for (const UScriptStruct* Struct : InInstance.GetAsset<UAnimNextRigVMAsset>()->GetReferencedVariableStructs())
 	{
-		AddStruct(Struct);
+		AddVariablesContainerForStruct(Struct, InInstance, InOverrides);
 	}
 
 	// Now any overrides we havent already added
@@ -178,6 +82,106 @@ void FUAFInstanceVariableData::Initialize(const FUAFAssetInstance& InInstance, c
 #endif
 }
 
+void FUAFInstanceVariableData::AddVariablesContainerForAsset(const UAnimNextRigVMAsset* InAsset, const FUAFAssetInstance& InInstance, const TSharedPtr<const UE::UAF::FVariableOverridesCollection>& InOverrides)
+{
+	using namespace UE::UAF;
+
+	TSharedPtr<FUAFInstanceVariableContainer> NewVariables;
+
+	// First check to see we have an override
+	if(InOverrides.IsValid())
+	{
+		for (int32 OverrideIndex = 0; OverrideIndex < InOverrides->Collection.Num(); ++OverrideIndex)
+		{
+			const FVariableOverrides& Override = InOverrides->Collection[OverrideIndex];
+			if (Override.IsOverrideFor(InAsset))
+			{
+				NewVariables = MakeShared<FUAFInstanceVariableContainer>(InInstance, InAsset, InOverrides, OverrideIndex);
+				OwnedVariableContainers.Add(NewVariables.ToSharedRef());
+				break;
+			}
+		}
+	}
+
+	// Check to see if this asset exists in our host chain. If so, we take a reference.
+	if (!NewVariables.IsValid())
+	{
+		const FUAFAssetInstance* CurrentHostInstance = InInstance.GetHost();
+		while (CurrentHostInstance)
+		{
+			if (const int32* SetIndexPtr = CurrentHostInstance->Variables.AllSharedVariableContainersMap.Find(InAsset))
+			{
+				NewVariables = CurrentHostInstance->Variables.AllSharedVariableContainers[*SetIndexPtr].Pin();
+				break;
+			}
+			CurrentHostInstance = CurrentHostInstance->GetHost();
+		}
+	}
+
+	// Otherwise, create a new instance that we own
+	if (!NewVariables.IsValid())
+	{
+		NewVariables = MakeShared<FUAFInstanceVariableContainer>(InInstance, InAsset);
+		OwnedVariableContainers.Add(NewVariables.ToSharedRef());
+	}
+
+	const int32 ReferencedAssetIndex = AllSharedVariableContainers.Emplace(NewVariables);
+	InternalSharedVariableContainers.Add(NewVariables);
+	AllSharedVariableContainersMap.Add(InAsset, ReferencedAssetIndex);
+	NumAllVariables += NewVariables->NumVariables;
+	NumInternalVariables += NewVariables->NumVariables;
+}
+
+void FUAFInstanceVariableData::AddVariablesContainerForStruct(const UScriptStruct* InStruct, const FUAFAssetInstance& InInstance, const TSharedPtr<const UE::UAF::FVariableOverridesCollection>& InOverrides)
+{
+	using namespace UE::UAF;
+
+	TSharedPtr<FUAFInstanceVariableContainer> NewVariables;
+
+	// First check to see we have an override
+	if(InOverrides.IsValid())
+	{
+		for (int32 OverrideIndex = 0; OverrideIndex < InOverrides->Collection.Num(); ++OverrideIndex)
+		{
+			const FVariableOverrides& Override = InOverrides->Collection[OverrideIndex];
+			if (Override.IsOverrideFor(InStruct))
+			{
+				NewVariables = MakeShared<FUAFInstanceVariableContainer>(InInstance, InStruct, InOverrides, OverrideIndex);
+				OwnedVariableContainers.Add(NewVariables.ToSharedRef());
+				break;
+			}
+		}
+	}
+
+	// Check to see if this asset exists in our host chain. If so, we take a reference.
+	if (!NewVariables.IsValid())
+	{
+		const FUAFAssetInstance* CurrentHostInstance = InInstance.GetHost();
+		while (CurrentHostInstance)
+		{
+			if (const int32* SetIndexPtr = CurrentHostInstance->Variables.AllSharedVariableContainersMap.Find(InStruct))
+			{
+				NewVariables = CurrentHostInstance->Variables.AllSharedVariableContainers[*SetIndexPtr].Pin();
+				break;
+			}
+			CurrentHostInstance = CurrentHostInstance->GetHost();
+		}
+	}
+
+	// Otherwise, create a new instance
+	if (!NewVariables.IsValid())
+	{
+		NewVariables = MakeShared<FUAFInstanceVariableContainer>(InInstance, InStruct);
+		OwnedVariableContainers.Add(NewVariables.ToSharedRef());
+	}
+
+	const int32 ReferencedAssetIndex = AllSharedVariableContainers.Emplace(NewVariables);
+	InternalSharedVariableContainers.Add(NewVariables);
+	AllSharedVariableContainersMap.Add(InStruct, ReferencedAssetIndex);
+	NumAllVariables += NewVariables->NumVariables;
+	NumInternalVariables += NewVariables->NumVariables;
+}
+
 void FUAFInstanceVariableData::RebuildNameMaps()
 {
 	using namespace UE::UAF;
@@ -248,20 +252,110 @@ void FUAFInstanceVariableData::Reset()
 }
 
 #if WITH_EDITOR
-void FUAFInstanceVariableData::Migrate()
+void FUAFInstanceVariableData::Migrate(const FUAFAssetInstance& InInstance)
 {
 	// Higher level logic needs to check to see if we have already been initialized, as overrides need to be cached first
 	check(bHasBeenInitialized);
 
-	NumAllVariables = 0;
-	NumInternalVariables = 0;
-	AllSharedVariableContainersMap.Reset();
+	const UAnimNextRigVMAsset* InstanceAsset = InInstance.GetAsset<UAnimNextRigVMAsset>();
+	TConstArrayView<TObjectPtr<const UAnimNextRigVMAsset>> ReferencedVariableAssets = InstanceAsset->GetReferencedVariableAssets();
+	TConstArrayView<TObjectPtr<const UScriptStruct>> ReferencedVariableStructs = InstanceAsset->GetReferencedVariableStructs();
+
+	// Referenced assets + struct + self-container
+	const int32 NumExpectedContainers = ReferencedVariableAssets.Num() + ReferencedVariableStructs.Num() + 1;
+	const int32 NumExistingContainers = AllSharedVariableContainersMap.Num();
+
+		// Verify the number of containers, new ones might have been added or existing ones removed
+	const bool bMissingContainers = NumExpectedContainers > NumExistingContainers;
+	const bool bOutOfDataContainers = NumExistingContainers > NumExpectedContainers;
+
+	if (bMissingContainers)
+	{
+		// Referenced asset variables
+		for (const UAnimNextRigVMAsset* ReferencedAsset : ReferencedVariableAssets)
+		{
+			const bool bHasVariableContainerForAsset = OwnedVariableContainers.ContainsByPredicate([ReferencedAsset](const TSharedRef<FUAFInstanceVariableContainer>& Container)
+			{
+				if (const FUAFInstanceVariableContainer::FAssetType* VariantPtr = Container->AssetOrStructData.TryGet<TObjectPtr<const UAnimNextRigVMAsset>>())
+				{
+					return (*VariantPtr) == ReferencedAsset;
+				}
+				return false;
+			});
+		
+			if (!bHasVariableContainerForAsset)
+			{
+				AddVariablesContainerForAsset(ReferencedAsset, InInstance, nullptr);
+			}
+		}
+
+		// Referenced struct variables
+		for (const UScriptStruct* Struct : ReferencedVariableStructs)
+		{
+			const bool bHasVariableContainerForStruct = OwnedVariableContainers.ContainsByPredicate([Struct](const TSharedRef<FUAFInstanceVariableContainer>& Container)
+			{
+				if (const FUAFInstanceVariableContainer::FStructType* VariantPtr = Container->AssetOrStructData.TryGet<FUAFInstanceVariableContainer::FStructType>())
+				{
+					return (*VariantPtr) == Struct;
+				}
+				return false;
+			});
+
+			if (!bHasVariableContainerForStruct)
+			{
+				AddVariablesContainerForStruct(Struct, InInstance, nullptr);
+			}
+		}
+	}
+	else if (bOutOfDataContainers)
+	{		
+		for (int32 Index = 0; Index < AllSharedVariableContainers.Num(); ++Index)
+		{
+			const TWeakPtr<FUAFInstanceVariableContainer>& WeakContainer = AllSharedVariableContainers[Index];
+			if (TSharedPtr<FUAFInstanceVariableContainer> SharedVariableContainer = WeakContainer.Pin())
+			{
+				const UObject* Object = nullptr;
+				switch (SharedVariableContainer->AssetOrStructData.GetIndex())
+				{
+				case FUAFInstanceVariableContainer::FAssetOrStructType::IndexOfType<FUAFInstanceVariableContainer::FAssetType>():
+					Object = SharedVariableContainer->AssetOrStructData.Get<FUAFInstanceVariableContainer::FAssetType>();
+					break;
+				case FUAFInstanceVariableContainer::FAssetOrStructType::IndexOfType<FUAFInstanceVariableContainer::FStructType>():
+					Object = SharedVariableContainer->AssetOrStructData.Get<FUAFInstanceVariableContainer::FStructType>();
+					break;
+				default:
+					checkNoEntry()
+					break;
+				}
+
+				// Referenced shared asset, or object corresponds to the instanced asset
+				const bool bIsAssetEntry = ReferencedVariableAssets.Contains(Object) || InstanceAsset == Object;
+				const bool bIsStructEntry = ReferencedVariableStructs.Contains(Object);
+
+				// Remove variable containers for asset/struct which are no longer referenced/part of this instance
+				if (!bIsAssetEntry && !bIsStructEntry)
+				{
+					AllSharedVariableContainers.RemoveAtSwap(Index);
 
+					OwnedVariableContainers.Remove(SharedVariableContainer.ToSharedRef());
+					InternalSharedVariableContainers.Remove(WeakContainer);
+					
+					--Index;
+				}
+			}
+		}
+	}
+
+	// Only migrate required/valid containers
 	for (TSharedRef<FUAFInstanceVariableContainer>& OwnedVariableContainer : OwnedVariableContainers)
 	{
 		OwnedVariableContainer->Migrate();
 	}
 
+	NumAllVariables = 0;
+	NumInternalVariables = 0;
+	AllSharedVariableContainersMap.Reset();
+
 	int32 ContainerIndex = 0;
 	for (TWeakPtr<FUAFInstanceVariableContainer>& WeakSharedVariableContainer : AllSharedVariableContainers)
 	{
@@ -294,7 +388,6 @@ void FUAFInstanceVariableData::Migrate()
 		TSharedPtr<FUAFInstanceVariableContainer> InternalSharedVariableContainer = WeakInternalSharedVariableContainer.Pin();
 		check(InternalSharedVariableContainer.IsValid());
 		
-		NumAllVariables += InternalSharedVariableContainer->NumVariables;
 		NumInternalVariables += InternalSharedVariableContainer->NumVariables;
 	}
 }
diff --git a/Engine/Plugins/Experimental/UAF/UAF/Source/UAF/Public/Variables/UAFInstanceVariableData.h b/Engine/Plugins/Experimental/UAF/UAF/Source/UAF/Public/Variables/UAFInstanceVariableData.h
--- a/Engine/Plugins/Experimental/UAF/UAF/Source/UAF/Public/Variables/UAFInstanceVariableData.h
+++ b/Engine/Plugins/Experimental/UAF/UAF/Source/UAF/Public/Variables/UAFInstanceVariableData.h
@@ -52,6 +52,9 @@ private:
 	// Initialize from a host's asset and its references
 	void Initialize(const FUAFAssetInstance& InHostInstance, const TSharedPtr<const UE::UAF::FVariableOverridesCollection>& InOverrides);
 
+	void AddVariablesContainerForAsset(const UAnimNextRigVMAsset* InAsset, const FUAFAssetInstance& InInstance, const TSharedPtr<const UE::UAF::FVariableOverridesCollection>& InOverrides);
+	void AddVariablesContainerForStruct(const UScriptStruct* InStruct, const FUAFAssetInstance& InInstance, const TSharedPtr<const UE::UAF::FVariableOverridesCollection>& InOverrides);
+
 	// Reset to the default state
 	void Reset();
 
@@ -60,7 +63,7 @@ private:
 
 #if WITH_EDITOR
 	// Used during compilation, migrate all variables to new bags according to new defaults.
-	void Migrate();
+	void Migrate(const FUAFAssetInstance& InHostInstance);
 #endif
 
 	// Resolve any VariableOverrides to FUAFInstanceVariableContainer::ResolvedOverrides
