Keep track of processed Assets and Structs so they only get added once to the treeview items list (having two assets reference the same native struct would duplicate the Struct property entries)
#misc renamed some variables for clarity, added missing log and removed duplicate item add

diff --git a/Engine/Plugins/Experimental/UAF/UAF/Source/UAFEditor/Private/Variables/Outliner/VariablesOutlinerHierarchy.cpp b/Engine/Plugins/Experimental/UAF/UAF/Source/UAFEditor/Private/Variables/Outliner/VariablesOutlinerHierarchy.cpp
--- a/Engine/Plugins/Experimental/UAF/UAF/Source/UAFEditor/Private/Variables/Outliner/VariablesOutlinerHierarchy.cpp
+++ b/Engine/Plugins/Experimental/UAF/UAF/Source/UAFEditor/Private/Variables/Outliner/VariablesOutlinerHierarchy.cpp
@@ -30,6 +30,9 @@ void FVariablesOutlinerHierarchy::CreateItems(TArray<FSceneOutlinerTreeItemPtr>&
 {
 	SVariablesOutliner* Outliner = static_cast<FVariablesOutlinerMode* const>(Mode)->GetOutliner();
 
+	WeakProcessedAssets.Reset();
+	WeakProcessedStructs.Reset();
+
 	if(Outliner->Assets.Num())
 	{
 		// Process first asset only for now
@@ -124,6 +127,13 @@ void FVariablesOutlinerHierarchy::ProcessAsset(const TSoftObjectPtr<UAnimNextRig
 		return;
 	}
 
+	if (WeakProcessedAssets.Contains(Asset))
+	{
+		return;
+	}
+
+	WeakProcessedAssets.Add(Asset);
+
 	UE_CLOG(GOutputHierarchyDebugInformationCvar.GetValueOnAnyThread(),LogAnimation, Warning, TEXT("ProcessAsset: %s %i"), *InSoftAsset.GetAssetName(), Depth);
 		
 	TArray<FSceneOutlinerTreeItemPtr> ParentItems;
@@ -153,6 +163,13 @@ void FVariablesOutlinerHierarchy::ProcessAsset(const TSoftObjectPtr<UAnimNextRig
 			{				
 				if (const UScriptStruct* Struct = InSharedVariablesEntry->GetStruct())
 				{
+					if (WeakProcessedStructs.Contains(Struct))
+					{
+						return true;
+					}
+
+					WeakProcessedStructs.Add(Struct);
+
 					if (FSceneOutlinerTreeItemPtr Item = Mode->CreateItemFor<FVariablesOutlinerStructSharedVariablesItem>(InSharedVariablesEntry))
 					{
 						OutItems.Add(Item);
@@ -171,9 +188,11 @@ void FVariablesOutlinerHierarchy::ProcessAsset(const TSoftObjectPtr<UAnimNextRig
 							EntryItem->WeakSharedVariablesEntry = InSharedVariablesEntry;
 							EntryItem->SortValue = Index++;
 							OutItems.Add(Item);
+
+							UE_CLOG(GOutputHierarchyDebugInformationCvar.GetValueOnAnyThread(), LogAnimation, Warning, TEXT("\t- FVariablesOutlinerEntryItem: %s %i"), *Property->GetName(), EntryItem->SortValue);
 						}
 					}
-				}				
+				}
 			}
 			break;
 		}		
@@ -235,9 +254,9 @@ void FVariablesOutlinerHierarchy::ProcessAsset(const TSoftObjectPtr<UAnimNextRig
 		}
 	};
 
-	if (FSceneOutlinerTreeItemPtr Item = Mode->CreateItemFor<FVariablesOutlinerAssetItem>(FVariablesOutlinerAssetItem::FItemData{InSoftAsset, AssetDepth, OutSharedVariableSourcePaths}))
+	if (FSceneOutlinerTreeItemPtr AssetItem = Mode->CreateItemFor<FVariablesOutlinerAssetItem>(FVariablesOutlinerAssetItem::FItemData{InSoftAsset, AssetDepth, OutSharedVariableSourcePaths}))
 	{
-		OutItems.Add(Item);
+		OutItems.Add(AssetItem);
 
 		UE_CLOG(GOutputHierarchyDebugInformationCvar.GetValueOnAnyThread(), LogAnimation, Warning, TEXT("* FVariablesOutlinerAssetItem: %s %i"), *InSoftAsset.GetAssetName(), AssetDepth);
 
@@ -251,7 +270,6 @@ void FVariablesOutlinerHierarchy::ProcessAsset(const TSoftObjectPtr<UAnimNextRig
 				EntryItem->SortValue = (CategoryIndex != INDEX_NONE ? ((CategoryIndex + 1) * NumberOfVariableSortValueEntriesInCategory) : 0) + Index++;
 				OutItems.Add(Item);
 				
-				OutItems.Add(Item);
 				UE_CLOG(GOutputHierarchyDebugInformationCvar.GetValueOnAnyThread(), LogAnimation, Warning, TEXT("\t- FVariablesOutlinerEntryItem: %s %i %s"), *InVariable->GetEntryName().ToString(), Item->CastTo<FVariablesOutlinerEntryItem>()->SortValue, InVariable->GetVariableCategory().GetData());
 
 				// Add category item on demand (to ensure we do not end up with empty category item hierarchies)
diff --git a/Engine/Plugins/Experimental/UAF/UAF/Source/UAFEditor/Private/Variables/Outliner/VariablesOutlinerHierarchy.h b/Engine/Plugins/Experimental/UAF/UAF/Source/UAFEditor/Private/Variables/Outliner/VariablesOutlinerHierarchy.h
--- a/Engine/Plugins/Experimental/UAF/UAF/Source/UAFEditor/Private/Variables/Outliner/VariablesOutlinerHierarchy.h
+++ b/Engine/Plugins/Experimental/UAF/UAF/Source/UAFEditor/Private/Variables/Outliner/VariablesOutlinerHierarchy.h
@@ -25,6 +25,9 @@ public:
 	// End ISceneOutlinerHierarchy overrides
 protected:
 	void ProcessAsset(const TSoftObjectPtr<UAnimNextRigVMAsset>& InSoftAsset, int32 Depth, TArray<FSceneOutlinerTreeItemPtr>& OutItems, TArray<const FSoftObjectPath>& OutSharedVariableSourcePaths) const;
+
+	mutable TSet<TWeakObjectPtr<const UAnimNextRigVMAsset>> WeakProcessedAssets;
+	mutable TSet<TWeakObjectPtr<const UScriptStruct>> WeakProcessedStructs;
 };
 
 }
\ No newline at end of file
