Output compiler error for duplicate variable names within asset context

diff --git a/Engine/Plugins/Experimental/UAF/UAF/Source/UAFUncookedOnly/Private/AnimNextRigVMAssetEditorData.cpp b/Engine/Plugins/Experimental/UAF/UAF/Source/UAFUncookedOnly/Private/AnimNextRigVMAssetEditorData.cpp
--- a/Engine/Plugins/Experimental/UAF/UAF/Source/UAFUncookedOnly/Private/AnimNextRigVMAssetEditorData.cpp
+++ b/Engine/Plugins/Experimental/UAF/UAF/Source/UAFUncookedOnly/Private/AnimNextRigVMAssetEditorData.cpp
@@ -3030,13 +3030,32 @@ FInstancedPropertyBag UAnimNextRigVMAssetEditorData::GenerateCombinedPropertyBag
 
 	FInstancedPropertyBag CombinedPropertyBag;
 
-	auto AddInstancedPropertyBag = [this, &CombinedPropertyBag](const FInstancedPropertyBag& InPropertyBag)
+	auto AddInstancedPropertyBag = [this, &CombinedPropertyBag, &InCompileContext, Asset](const FInstancedPropertyBag& InPropertyBag, const UAnimNextRigVMAsset* SourceAsset, bool bValidateProperties = false)
 	{
 		if(const UPropertyBag* PropertyBag = InPropertyBag.GetPropertyBagStruct())
 		{
 			TConstArrayView<FPropertyBagPropertyDesc> VariableDescs = PropertyBag->GetPropertyDescs();
 			if(VariableDescs.Num() != 0)
 			{
+				if (bValidateProperties)
+				{
+					bool bInvalidProperties = false;
+					for (const FPropertyBagPropertyDesc& Desc : VariableDescs)
+					{
+						if (const FPropertyBagPropertyDesc* ExistingDesc = CombinedPropertyBag.FindPropertyDescByName(Desc.Name))
+						{
+							InCompileContext.GetAssetCompileContext().Error(Asset, LOCTEXT("DuplicateVariableEntries", "Variable with overlapping Name {0} found in {1}"), FText::FromName(Desc.Name), FText::FromString(SourceAsset->GetPathName()));
+
+							bInvalidProperties = true;
+						}
+					}
+
+					if (bInvalidProperties)
+					{
+						return;
+					}
+				}
+				
 				CombinedPropertyBag.AddProperties(VariableDescs);
 			}
 		}
@@ -3047,7 +3066,7 @@ FInstancedPropertyBag UAnimNextRigVMAssetEditorData::GenerateCombinedPropertyBag
 
 	// First add internal variables
 	FInstancedPropertyBag AssetBag = UE::UAF::UncookedOnly::FUtils::MakePropertyBagForEditorData(this, InCompileContext);
-	AddInstancedPropertyBag(AssetBag);
+	AddInstancedPropertyBag(AssetBag, Asset, false);
 
 	// Next add shared variables
 	for (const UAnimNextRigVMAsset* ReferencedVariableAsset : Asset->ReferencedVariableAssets)
@@ -3061,7 +3080,7 @@ FInstancedPropertyBag UAnimNextRigVMAssetEditorData::GenerateCombinedPropertyBag
 
 		FAnimNextRigVMAssetCompileContext CompileContext = { ReferencedEditorData };
 		AssetBag = UE::UAF::UncookedOnly::FUtils::MakePropertyBagForEditorData(ReferencedEditorData, ReferencedEditorData->GetVariableCompileContext(InSettings, CompileContext));
-		AddInstancedPropertyBag(AssetBag);
+		AddInstancedPropertyBag(AssetBag, ReferencedVariableAsset, true);
 	}
 
 	// Next add native structs
