Fix UAF Crash on Drag Shared Variables due to UAF attempting to recompile mid variable add / remove and getting and inccorect num variable compare. This can result in a permenantly corrupt asset if the drag does not lead to a crash (IE: when nothing is running the graph), and the variable drag change is partially commit.

#tests PIE, editor drag drop

diff --git a/Engine/Plugins/Experimental/UAF/UAF/Source/UAFUncookedOnly/Private/UncookedOnlyUtils.cpp b/Engine/Plugins/Experimental/UAF/UAF/Source/UAFUncookedOnly/Private/UncookedOnlyUtils.cpp
--- a/Engine/Plugins/Experimental/UAF/UAF/Source/UAFUncookedOnly/Private/UncookedOnlyUtils.cpp
+++ b/Engine/Plugins/Experimental/UAF/UAF/Source/UAFUncookedOnly/Private/UncookedOnlyUtils.cpp
@@ -1499,9 +1499,11 @@ void FUtils::DeleteVariable(UAnimNextVariableEntry* VariableEntry, bool bSetupUn
 {
 	if (VariableEntry)
 	{
-		const UAnimNextRigVMAsset* OuterRigVMAsset = VariableEntry->GetTypedOuter<UAnimNextRigVMAsset>();
+		UAnimNextRigVMAsset* OuterRigVMAsset = VariableEntry->GetTypedOuter<UAnimNextRigVMAsset>();
 		check(OuterRigVMAsset);
 		
+		FScopedCompileJob CompilerResults(LOCTEXT("ModifiedAssets_DeleteVariable", "Modified Assets Delete Variable"), { OuterRigVMAsset });
+		
 		const FAnimNextSoftVariableReference FindReference(VariableEntry->GetEntryName(), OuterRigVMAsset);
 
 		// Replace references with None
@@ -1529,6 +1531,8 @@ void FUtils::RenameVariable(UAnimNextVariableEntry* VariableEntry, const FName N
 	{
 		UAnimNextRigVMAsset* OuterRigVMAsset = VariableEntry->GetTypedOuter<UAnimNextRigVMAsset>();
 		check(OuterRigVMAsset);
+
+		FScopedCompileJob CompilerResults(LOCTEXT("ModifiedAssets_RenameVariable", "Modified Assets Rename Variable"), { OuterRigVMAsset });
 		
 		const FAnimNextSoftVariableReference FindReference(VariableEntry->GetEntryName(), OuterRigVMAsset);
 		
@@ -1557,6 +1561,8 @@ void FUtils::MoveVariableToAsset(UAnimNextVariableEntry* VariableEntry, UAnimNex
 	{
 		if (UAnimNextRigVMAsset* CurrentOuter = VariableEntry->GetTypedOuter<UAnimNextRigVMAsset>())
 		{
+			FScopedCompileJob CompilerResults(LOCTEXT("ModifiedAssets_MoveVariableToAsset", "Modified Assets Move Variable"), { CurrentOuter, NewOuter });
+
 			UAnimNextRigVMAssetEditorData* CurrentEditorData = GetEditorData(CurrentOuter);
 			UAnimNextRigVMAssetEditorData* NewEditorData = GetEditorData(NewOuter);
 						
@@ -1591,6 +1597,8 @@ void FUtils::SetVariableType(UAnimNextVariableEntry* VariableEntry, const FAnimN
 	{
 		UAnimNextRigVMAsset* OuterRigVMAsset = VariableEntry->GetTypedOuter<UAnimNextRigVMAsset>();
 		check(OuterRigVMAsset);
+
+		FScopedCompileJob CompilerResults(LOCTEXT("ModifiedAssets_SetVariableType", "Modified Assets Set Variable Type"), { OuterRigVMAsset });
 		
 		const FAnimNextSoftVariableReference FindReference(VariableEntry->GetEntryName(), OuterRigVMAsset);
 		
