[Animation] added public helper functions UE::UAF::SetModuleVariable/WriteModuleVariable to allow external frameworks to read/write variable(s) from/to AnimationModule instances.
Added Mass Trait, Fragment and StateTree Task to write variables
Added thread access detectors to AnimationFramework subsystem

diff --git a/Engine/Plugins/Experimental/UAF/UAF/Source/UAF/Internal/Component/AnimNextWorldSubsystem.h b/Engine/Plugins/Experimental/UAF/UAF/Source/UAF/Internal/Component/AnimNextWorldSubsystem.h
--- a/Engine/Plugins/Experimental/UAF/UAF/Source/UAF/Internal/Component/AnimNextWorldSubsystem.h
+++ b/Engine/Plugins/Experimental/UAF/UAF/Source/UAF/Internal/Component/AnimNextWorldSubsystem.h
@@ -2,29 +2,31 @@
 
 #pragma once
 
-#include "CoreMinimal.h"
 #include "AnimNextPool.h"
-#include "Subsystems/WorldSubsystem.h"
+#include "EngineDefines.h"
+#include "Misc/TVariant.h"
 #include "Module/ModuleHandle.h"
-#include "Module/ModuleTickFunction.h"
-#include "Module/AnimNextModuleInstance.h"
-#include "Module/ModuleTaskContext.h"
+#include "Module/ModuleInstanceUtils.h"
 #include "Module/TaskRunLocation.h"
+#include "Subsystems/WorldSubsystem.h"
 #include "AnimNextWorldSubsystem.generated.h"
 
 #define UE_API UAF_API
 
-class UAnimNextComponent;
-struct FRigUnit_AnimNextInitializeEvent;
-
-namespace UE::UAF
-{
-	struct FModuleEventTickFunction;
-	struct FInjectionRequest;
-}
+class UAnimNextModule;
+class UAnimNextRigVMAsset;
+enum class EAnimNextModuleInitMethod : uint8;
+enum class EPropertyBagResult : uint8;
+struct FAnimNextParamType;
+struct FAnimNextTraitEvent;
+struct FAnimNextVariableReference;
 
 namespace UE::UAF
 {
+struct FInjectionRequest;
+struct FModuleContext;
+struct FModuleEventTickFunction;
+struct FModuleTaskContext;
 
 // A queued action to complete next frame
 struct FModulePendingAction
@@ -81,7 +83,7 @@ public:
 
 	// Get the subsystem for the specified object's world
 	static UE_API UAnimNextWorldSubsystem* Get(UObject* InObject);
-	
+
 	static UE_API void AddReferencedObjects(UObject* InThis, FReferenceCollector& Collector);
 
 	// Dependency type passed to AddDependencyHandle etc.
@@ -135,7 +137,7 @@ protected:
 	UE_API void QueueTaskHandle(UE::UAF::FModuleHandle InHandle, FName InModuleEventName, TUniqueFunction<void(const UE::UAF::FModuleTaskContext&)>&& InTaskFunction, UE::UAF::ETaskRunLocation InLocation = UE::UAF::ETaskRunLocation::Before);
 
 	// Queue an input trait event on the supplied handle
-	UE_API void QueueInputTraitEventHandle(UE::UAF::FModuleHandle InHandle, FAnimNextTraitEventPtr Event);
+	UE_API void QueueInputTraitEventHandle(UE::UAF::FModuleHandle InHandle, TSharedPtr<FAnimNextTraitEvent> Event);
 
 	// Find the const tick function for the specified event. Useful to analyze/log prerequisites. Use AddDependencyHandle, RemoveDependencyHandle to actually modify dependencies
 	// @param	InHandle			The handle of the module
@@ -149,7 +151,7 @@ protected:
 	// @param	InEventName			The event to add the dependency to
 	// @param	InDependency		The kind of dependency to add
 	UE_API void AddDependencyHandle(UE::UAF::FModuleHandle InHandle, UObject* InObject, FTickFunction& InTickFunction, FName InEventName, EDependency InDependency);
-	
+
 	// Remove a dependency on a tick function from the specified event
 	// @param	InHandle			The handle of the module whose dependencies are being modified
 	// @param	InObject			The object that owns the tick function
@@ -184,7 +186,7 @@ protected:
 
 	// Accesses the variable of the specified name for writing.
 	// General read access is not provided via this API due to the current double-buffering strategy used to communicate variable writes to worker threads.
-	// This is intended to allow for copy-free writing of larger data structures & arrays, rather than read access. 
+	// This is intended to allow for copy-free writing of larger data structures & arrays, rather than read access.
 	// @param	InHandle			The handle of the module whose variable we are setting
 	// @param	InVariable			The variable we want to access
 	// @param	InType				The type of the variable we want to access
@@ -196,6 +198,8 @@ protected:
 	friend struct FRigUnit_AnimNextInitializeEvent;
 	friend struct UE::UAF::FModuleEventTickFunction;
 	friend struct UE::UAF::FInjectionRequest;
+	friend EPropertyBagResult UE::UAF::SetModuleVariable(UE::UAF::FModuleContext&&, const FAnimNextVariableReference&, const FAnimNextParamType&, TConstArrayView<uint8>);
+	friend EPropertyBagResult UE::UAF::WriteModuleVariable(UE::UAF::FModuleContext&&, const FAnimNextVariableReference&, const FAnimNextParamType&, TFunctionRef<void(TArrayView<uint8>)>);
 
 	// Currently running instances, pooled
 	UE::UAF::TPool<FAnimNextModuleInstance> Instances;
@@ -203,9 +207,11 @@ protected:
 	// Queued actions
 	TArray<UE::UAF::FModulePendingAction> PendingActions;
 
-	// Locks for concurrent modifications
+	// Locks and validators for concurrent modifications
 	FRWLock InstancesLock;
 	FRWLock PendingLock;
+	UE_MT_DECLARE_RW_FULLY_RECURSIVE_ACCESS_DETECTOR(InstancesAccessDetector);
+	UE_MT_DECLARE_RW_ACCESS_DETECTOR(PendingActionsAccessDetector);
 
 #if WITH_EDITOR
 	// Handle used to hook asset compilation
@@ -219,4 +225,4 @@ protected:
 	float DeltaTime;
 };
 
-#undef UE_API
+#undef UE_API
\ No newline at end of file
diff --git a/Engine/Plugins/Experimental/UAF/UAF/Source/UAF/Private/Component/AnimNextWorldSubsystem.cpp b/Engine/Plugins/Experimental/UAF/UAF/Source/UAF/Private/Component/AnimNextWorldSubsystem.cpp
--- a/Engine/Plugins/Experimental/UAF/UAF/Source/UAF/Private/Component/AnimNextWorldSubsystem.cpp
+++ b/Engine/Plugins/Experimental/UAF/UAF/Source/UAF/Private/Component/AnimNextWorldSubsystem.cpp
@@ -2,19 +2,18 @@
 
 #include "Component/AnimNextWorldSubsystem.h"
 
+#include "AnimNextDebugDraw.h"
+#include "AnimNextRigVMAsset.h"
 #include "Engine/World.h"
 #include "Logging/StructuredLog.h"
-#include "Module/AnimNextModule.h"
-#include "Module/ModuleHandle.h"
+#include "Module/AnimNextModuleInstance.h"
 #include "Module/ModuleTaskContext.h"
-#include "AnimNextDebugDraw.h"
-#include "Variables/AnimNextVariableReference.h"
 
 #include UE_INLINE_GENERATED_CPP_BY_NAME(AnimNextWorldSubsystem)
 
 UAnimNextWorldSubsystem::UAnimNextWorldSubsystem()
 {
-	if(!HasAnyFlags(RF_ClassDefaultObject))
+	if (!HasAnyFlags(RF_ClassDefaultObject))
 	{
 #if WITH_EDITOR
 		OnCompileJobFinishedHandle = UAnimNextRigVMAsset::OnCompileJobFinished().AddUObject(this, &UAnimNextWorldSubsystem::OnCompileJobFinished);
@@ -47,6 +46,8 @@ void UAnimNextWorldSubsystem::BeginDestroy()
 
 		{
 			FRWScopeLock InstancesLockScope(InstancesLock, SLT_Write);
+			UE_MT_SCOPED_WRITE_ACCESS(InstancesAccessDetector);
+
 			for (FAnimNextModuleInstance& Instance : Instances)
 			{
 				Instance.RemoveAllTickDependencies();
@@ -59,13 +60,13 @@ void UAnimNextWorldSubsystem::BeginDestroy()
 
 UAnimNextWorldSubsystem* UAnimNextWorldSubsystem::Get(UObject* InObject)
 {
-	if(InObject == nullptr)
+	if (InObject == nullptr)
 	{
 		return nullptr;
 	}
 
 	UWorld* World = InObject->GetWorld();
-	if(World == nullptr)
+	if (World == nullptr)
 	{
 		return nullptr;
 	}
@@ -78,7 +79,9 @@ void UAnimNextWorldSubsystem::AddReferencedObjects(UObject* InThis, FReferenceCo
 	Super::AddReferencedObjects(InThis, Collector);
 
 	UAnimNextWorldSubsystem* This = CastChecked<UAnimNextWorldSubsystem>(InThis);
-	for(FAnimNextModuleInstance& Instance : This->Instances)
+	UE_MT_SCOPED_READ_ACCESS(This->InstancesAccessDetector);
+
+	for (FAnimNextModuleInstance& Instance : This->Instances)
 	{
 		Collector.AddPropertyReferencesWithStructARO(FAnimNextModuleInstance::StaticStruct(), &Instance, InThis);
 	}
@@ -101,6 +104,7 @@ bool UAnimNextWorldSubsystem::DoesSupportWorldType(const EWorldType::Type WorldT
 
 bool UAnimNextWorldSubsystem::IsValidHandle(UE::UAF::FModuleHandle InHandle) const
 {
+	UE_MT_SCOPED_READ_ACCESS(InstancesAccessDetector);
 	return Instances.IsValidHandle(InHandle);
 }
 
@@ -109,11 +113,10 @@ void UAnimNextWorldSubsystem::FlushPendingActions()
 	using namespace UE::UAF;
 
 	FRWScopeLock PendingLockScope(PendingLock, SLT_Write);
+	UE_MT_SCOPED_WRITE_ACCESS(PendingActionsAccessDetector);
 
 	if (PendingActions.Num() > 0)
 	{
-		FRWScopeLock InstancesLockScope(InstancesLock, SLT_Write);
-
 		for (FModulePendingAction& PendingAction : PendingActions)
 		{
 			switch (PendingAction.Type)
@@ -121,12 +124,15 @@ void UAnimNextWorldSubsystem::FlushPendingActions()
 			case FModulePendingAction::EType::ReleaseHandle:
 				if (IsValidHandle(PendingAction.Handle))
 				{
+					FRWScopeLock InstancesLockScope(InstancesLock, SLT_Write);
+					UE_MT_SCOPED_WRITE_ACCESS(InstancesAccessDetector);
 					Instances.Release(PendingAction.Handle);
 				}
 				break;
 			case FModulePendingAction::EType::EnableHandle:
 				if (IsValidHandle(PendingAction.Handle))
 				{
+					UE_MT_SCOPED_READ_ACCESS(InstancesAccessDetector);
 					FAnimNextModuleInstance& Instance = Instances.Get(PendingAction.Handle);
 					Instance.Enable(PendingAction.Payload.Get<bool>());
 				}
@@ -135,6 +141,7 @@ void UAnimNextWorldSubsystem::FlushPendingActions()
 #if UE_ENABLE_DEBUG_DRAWING
 				if (IsValidHandle(PendingAction.Handle))
 				{
+					UE_MT_SCOPED_READ_ACCESS(InstancesAccessDetector);
 					FAnimNextModuleInstance& Instance = Instances.Get(PendingAction.Handle);
 					Instance.ShowDebugDrawing(PendingAction.Payload.Get<bool>());
 				}
@@ -154,12 +161,19 @@ void UAnimNextWorldSubsystem::RegisterHandle(UE::UAF::FModuleHandle& InOutHandle
 {
 	using namespace UE::UAF;
 	check(IsInGameThread());
-	FWriteScopeLock InstancesLockScope(InstancesLock);
+	{
+		FWriteScopeLock InstancesLockScope(InstancesLock);
+		UE_MT_SCOPED_WRITE_ACCESS(InstancesAccessDetector);
 
-	InOutHandle = Instances.Emplace(InModule, InObject, &Instances, InitMethod);
-	FAnimNextModuleInstance& Instance = Instances.Get(InOutHandle);
-	Instance.Handle = InOutHandle;
-	Instance.Initialize();
+		InOutHandle = Instances.Emplace(InModule, InObject, &Instances, InitMethod);
+	}
+
+	{
+		UE_MT_SCOPED_READ_ACCESS(InstancesAccessDetector);
+		FAnimNextModuleInstance& Instance = Instances.Get(InOutHandle);
+		Instance.Handle = InOutHandle;
+		Instance.Initialize();
+	}
 }
 
 void UAnimNextWorldSubsystem::UnregisterHandle(UE::UAF::FModuleHandle& InOutHandle)
@@ -167,9 +181,9 @@ void UAnimNextWorldSubsystem::UnregisterHandle(UE::UAF::FModuleHandle& InOutHand
 	using namespace UE::UAF;
 	check(IsInGameThread());
 
-	if(IsValidHandle(InOutHandle))
+	if (IsValidHandle(InOutHandle))
 	{
-		FWriteScopeLock InstancesLockScope(InstancesLock);
+		UE_MT_SCOPED_READ_ACCESS(InstancesAccessDetector);
 		FAnimNextModuleInstance& Instance = Instances.Get(InOutHandle);
 
 #if UE_ENABLE_DEBUG_DRAWING
@@ -180,7 +194,11 @@ void UAnimNextWorldSubsystem::UnregisterHandle(UE::UAF::FModuleHandle& InOutHand
 		// Remove all tick dependencies immediately, as once the handle has been invalidated there is no way for external systems to remove their dependencies
 		Instance.RemoveAllTickDependencies();
 
-		PendingActions.Emplace(InOutHandle, FModulePendingAction::EType::ReleaseHandle);
+		{
+			FRWScopeLock PendingLockScope(PendingLock, SLT_Write);
+			UE_MT_SCOPED_WRITE_ACCESS(PendingActionsAccessDetector);
+			PendingActions.Emplace(InOutHandle, FModulePendingAction::EType::ReleaseHandle);
+		}
 		InOutHandle.Reset();
 	}
 }
@@ -192,16 +210,22 @@ bool UAnimNextWorldSubsystem::IsHandleEnabled(UE::UAF::FModuleHandle InHandle) c
 	if (IsValidHandle(InHandle))
 	{
 		// The last pending action takes precedent here if present
-		int32 PendingActionIndex = PendingActions.FindLastByPredicate([&InHandle](const FModulePendingAction& PendingAction) { return PendingAction.Handle == InHandle && PendingAction.Type == FModulePendingAction::EType::EnableHandle;});
-
-		if (PendingActionIndex != INDEX_NONE)
 		{
-			return PendingActions[PendingActionIndex].Payload.Get<bool>();
+			UE_MT_SCOPED_READ_ACCESS(PendingActionsAccessDetector);
+			int32 PendingActionIndex = PendingActions.FindLastByPredicate([&InHandle](const FModulePendingAction& PendingAction) { return PendingAction.Handle == InHandle && PendingAction.Type == FModulePendingAction::EType::EnableHandle;});
+
+			if (PendingActionIndex != INDEX_NONE)
+			{
+				return PendingActions[PendingActionIndex].Payload.Get<bool>();
+			}
 		}
 
 		// Otherwise return the current value on the instance
-		const FAnimNextModuleInstance& Instance = Instances.Get(InHandle);
-		return Instance.IsEnabled();
+		{
+			UE_MT_SCOPED_READ_ACCESS(InstancesAccessDetector);
+			const FAnimNextModuleInstance& Instance = Instances.Get(InHandle);
+			return Instance.IsEnabled();
+		}
 	}
 	return false;
 }
@@ -212,6 +236,7 @@ void UAnimNextWorldSubsystem::EnableHandle(UE::UAF::FModuleHandle InHandle, bool
 	check(IsInGameThread());
 	if (IsValidHandle(InHandle))
 	{
+		UE_MT_SCOPED_WRITE_ACCESS(PendingActionsAccessDetector);
 		PendingActions.Emplace(InHandle, FModulePendingAction::EType::EnableHandle, bInEnabled);
 	}
 }
@@ -223,6 +248,7 @@ void UAnimNextWorldSubsystem::ShowDebugDrawingHandle(UE::UAF::FModuleHandle InHa
 	check(IsInGameThread());
 	if (IsValidHandle(InHandle))
 	{
+		UE_MT_SCOPED_READ_ACCESS(PendingActionsAccessDetector);
 		PendingActions.Emplace(InHandle, FModulePendingAction::EType::EnableDebugDrawing, bInShowDebugDrawing);
 	}
 }
@@ -234,6 +260,7 @@ void UAnimNextWorldSubsystem::QueueTaskHandle(UE::UAF::FModuleHandle InOutHandle
 	check(IsInGameThread());
 	if (IsValidHandle(InOutHandle))
 	{
+		UE_MT_SCOPED_READ_ACCESS(InstancesAccessDetector);
 		FAnimNextModuleInstance& Instance = Instances.Get(InOutHandle);
 		Instance.QueueTask(InModuleEventName, MoveTemp(InTaskFunction), InLocation);
 	}
@@ -244,11 +271,12 @@ void UAnimNextWorldSubsystem::QueueTaskHandle(UE::UAF::FModuleHandle InOutHandle
 void UAnimNextWorldSubsystem::OnCompileJobFinished(UAnimNextRigVMAsset* InAsset)
 {
 	// Cant do this while we are running in a world tick
-	check(!GetWorld()->bInTick); 
+	check(!GetWorld()->bInTick);
+	UE_MT_SCOPED_READ_ACCESS(InstancesAccessDetector);
 
-	for(FAnimNextModuleInstance& Instance : Instances)
+	for (FAnimNextModuleInstance& Instance : Instances)
 	{
-		if(Instance.GetModule() == InAsset)
+		if (Instance.GetModule() == InAsset)
 		{
 			Instance.OnCompileJobFinished();
 		}
@@ -274,15 +302,17 @@ const FTickFunction* UAnimNextWorldSubsystem::FindTickFunctionHandle(UE::UAF::FM
 	check(IsInGameThread());
 	if (IsValidHandle(InHandle))
 	{
+		UE_MT_SCOPED_READ_ACCESS(InstancesAccessDetector);
+
 		const FAnimNextModuleInstance& Instance = Instances.Get(InHandle);
 		const FModuleEventTickFunction* TickFunction = Instance.FindTickFunctionByName(InEventName);
-		if(TickFunction == nullptr)
+		if (TickFunction == nullptr)
 		{
 			UE_LOGFMT(LogAnimation, Warning, "FindTickFunctionHandle: Could not find event '{EventName}' in module '{ModuleName}'", InEventName, Instance.GetAssetName());
 			return nullptr;
 		}
-		
-		if(TickFunction->bUserEvent)
+
+		if (TickFunction->bUserEvent)
 		{
 			return TickFunction;
 		}
@@ -298,15 +328,17 @@ void UAnimNextWorldSubsystem::AddDependencyHandle(UE::UAF::FModuleHandle InHandl
 	check(IsInGameThread());
 	if (IsValidHandle(InHandle))
 	{
+		UE_MT_SCOPED_READ_ACCESS(InstancesAccessDetector);
+
 		FAnimNextModuleInstance& Instance = Instances.Get(InHandle);
 		FModuleEventTickFunction* TickFunction = Instance.FindTickFunctionByName(InEventName);
-		if(TickFunction == nullptr)
+		if (TickFunction == nullptr)
 		{
 			UE_LOGFMT(LogAnimation, Warning, "AddDependencyHandle: Could not find event '{EventName}' in module '{ModuleName}'", InEventName, Instance.GetAssetName());
 			return;
 		}
 
-		if(InDependency == EDependency::Prerequisite)
+		if (InDependency == EDependency::Prerequisite)
 		{
 			TickFunction->AddPrerequisite(InObject, InTickFunction);
 		}
@@ -323,15 +355,17 @@ void UAnimNextWorldSubsystem::RemoveDependencyHandle(UE::UAF::FModuleHandle InHa
 	check(IsInGameThread());
 	if (IsValidHandle(InHandle))
 	{
+		UE_MT_SCOPED_READ_ACCESS(InstancesAccessDetector);
+
 		FAnimNextModuleInstance& Instance = Instances.Get(InHandle);
 		FModuleEventTickFunction* TickFunction = Instance.FindTickFunctionByName(InEventName);
-		if(TickFunction == nullptr)
+		if (TickFunction == nullptr)
 		{
 			UE_LOGFMT(LogAnimation, Warning, "RemoveDependencyHandle: Could not find event '{EventName}' in module '{ModuleName}'", InEventName, Instance.GetAssetName());
 			return;
 		}
 
-		if(InDependency == EDependency::Prerequisite)
+		if (InDependency == EDependency::Prerequisite)
 		{
 			TickFunction->RemovePrerequisite(InObject, InTickFunction);
 		}
@@ -349,9 +383,11 @@ void UAnimNextWorldSubsystem::AddModuleEventDependencyHandle(UE::UAF::FModuleHan
 	check(InHandle != OtherHandle);
 	if (IsValidHandle(OtherHandle))
 	{
+		UE_MT_SCOPED_READ_ACCESS(InstancesAccessDetector);
+
 		FAnimNextModuleInstance& Instance = Instances.Get(OtherHandle);
 		FModuleEventTickFunction* TickFunction = Instance.FindTickFunctionByName(OtherEventName);
-		if(TickFunction == nullptr)
+		if (TickFunction == nullptr)
 		{
 			UE_LOGFMT(LogAnimation, Warning, "AddDependencyHandle: Could not find event '{EventName}' in module '{ModuleName}'", OtherEventName, Instance.GetAssetName());
 			return;
@@ -368,9 +404,11 @@ void UAnimNextWorldSubsystem::RemoveModuleEventDependencyHandle(UE::UAF::FModule
 	check(InHandle != OtherHandle);
 	if (IsValidHandle(InHandle))
 	{
+		UE_MT_SCOPED_READ_ACCESS(InstancesAccessDetector);
+
 		FAnimNextModuleInstance& Instance = Instances.Get(InHandle);
 		FModuleEventTickFunction* TickFunction = Instance.FindTickFunctionByName(InEventName);
-		if(TickFunction == nullptr)
+		if (TickFunction == nullptr)
 		{
 			UE_LOGFMT(LogAnimation, Warning, "RemoveDependencyHandle: Could not find event '{EventName}' in module '{ModuleName}'", OtherEventName, Instance.GetAssetName());
 			return;
@@ -385,6 +423,8 @@ EPropertyBagResult UAnimNextWorldSubsystem::SetVariableHandle(UE::UAF::FModuleHa
 	using namespace UE::UAF;
 	if (IsValidHandle(InHandle))
 	{
+		UE_MT_SCOPED_READ_ACCESS(InstancesAccessDetector);
+
 		FAnimNextModuleInstance& Instance = Instances.Get(InHandle);
 		return Instance.SetProxyVariable(InVariable, InType, InData);
 	}
@@ -396,6 +436,8 @@ EPropertyBagResult UAnimNextWorldSubsystem::WriteVariableHandle(UE::UAF::FModule
 	using namespace UE::UAF;
 	if (IsValidHandle(InHandle))
 	{
+		UE_MT_SCOPED_READ_ACCESS(InstancesAccessDetector);
+
 		FAnimNextModuleInstance& Instance = Instances.Get(InHandle);
 		return Instance.WriteProxyVariable(InVariable, InType, InFunction);
 	}
diff --git a/Engine/Plugins/Experimental/UAF/UAF/Source/UAF/Private/Module/ModuleInstanceUtils.cpp b/Engine/Plugins/Experimental/UAF/UAF/Source/UAF/Private/Module/ModuleInstanceUtils.cpp
new file mode 100644
--- /dev/null
+++ b/Engine/Plugins/Experimental/UAF/UAF/Source/UAF/Private/Module/ModuleInstanceUtils.cpp
@@ -0,0 +1,47 @@
+// Copyright Epic Games, Inc. All Rights Reserved.
+
+#include "Module/ModuleInstanceUtils.h"
+
+#include "Component/AnimNextComponentWorldSubsystem.h"
+#include "Engine/Engine.h"
+#include "StructUtils/PropertyBag.h"
+#include "Variables/AnimNextVariableReference.h"
+
+namespace UE::UAF
+{
+
+EPropertyBagResult SetModuleVariable(FModuleContext&& InModuleContext, const FAnimNextVariableReference& InVariable
+	, const FAnimNextParamType& InType, TConstArrayView<uint8> InData)
+{
+	if (const UWorld* World = GEngine->GetWorldFromContextObject(InModuleContext.ContextObject, EGetWorldErrorMode::ReturnNull))
+	{
+		if (UAnimNextComponentWorldSubsystem* Subsystem = World->GetSubsystem<UAnimNextComponentWorldSubsystem>())
+		{
+			return Subsystem->SetVariableHandle(InModuleContext.Handle, InVariable, InType, InData);
+		}
+	}
+
+	UE_LOG(LogAnimation, Error, TEXT("Unable to set variable '%s'. Cause: unable to access world system from the provided context object '%s'")
+		, *InVariable.GetName().ToString()
+		, *GetFullNameSafe(InModuleContext.ContextObject));
+	return EPropertyBagResult::PropertyNotFound;
+}
+
+EPropertyBagResult WriteModuleVariable(FModuleContext&& InModuleContext, const FAnimNextVariableReference& InVariable
+	, const FAnimNextParamType& InType, TFunctionRef<void(TArrayView<uint8>)> InFunction)
+{
+	if (const UWorld* World = GEngine->GetWorldFromContextObject(InModuleContext.ContextObject, EGetWorldErrorMode::ReturnNull))
+	{
+		if (UAnimNextComponentWorldSubsystem* Subsystem = World->GetSubsystem<UAnimNextComponentWorldSubsystem>())
+		{
+			return Subsystem->WriteVariableHandle(InModuleContext.Handle, InVariable, InType, InFunction);
+		}
+	}
+
+	UE_LOG(LogAnimation, Error, TEXT("Unable to set variable '%s'. Cause: unable to access world system from the provided context object '%s'")
+		, *InVariable.GetName().ToString()
+		, *GetFullNameSafe(InModuleContext.ContextObject));
+	return EPropertyBagResult::PropertyNotFound;
+}
+
+} // UE::UAF
\ No newline at end of file
diff --git a/Engine/Plugins/Experimental/UAF/UAF/Source/UAF/Public/Module/ModuleInstanceUtils.h b/Engine/Plugins/Experimental/UAF/UAF/Source/UAF/Public/Module/ModuleInstanceUtils.h
new file mode 100644
--- /dev/null
+++ b/Engine/Plugins/Experimental/UAF/UAF/Source/UAF/Public/Module/ModuleInstanceUtils.h
@@ -0,0 +1,48 @@
+// Copyright Epic Games, Inc. All Rights Reserved.
+
+#pragma once
+
+#include "CoreTypes.h"
+#include "ModuleHandle.h"
+
+struct FAnimNextParamType;
+struct FAnimNextVariableReference;
+struct FAnimNextModuleInstance;
+enum class EPropertyBagResult : uint8;
+
+#define UE_API UAF_API
+
+namespace UE::UAF
+{
+
+struct FModuleContext
+{
+	FModuleHandle Handle;
+	TObjectPtr<UObject> ContextObject;
+};
+
+/**
+ * Sets the value of the specified variable in the specified animation module
+ * @param InModuleContext	The context providing the handle of the module whose variable we are setting and some additional information to access the module.
+ * @param InVariable		The variable to set
+ * @param InType			The type of the variable to set
+ * @param InData			The data to set the variable with
+ * @return Whether the variable was set successfully
+ */
+UE_API EPropertyBagResult SetModuleVariable(FModuleContext&& InModuleContext, const FAnimNextVariableReference& InVariable, const FAnimNextParamType& InType, TConstArrayView<uint8> InData);
+
+/**
+ * Accesses the variable of the specified name for writing.
+ * General read access is not provided via this API due to the current double-buffering strategy used to communicate variable writes to worker threads.
+ * This is intended to allow for copy-free writing of larger data structures & arrays, rather than read access.
+ * @param InModuleContext	The context providing the handle of the module whose variable we are setting and some additional information to access the module.
+ * @param InVariable		The variable we want to access
+ * @param InType			The type of the variable we want to access
+ * @param InFunction		Function that will be called to allow modification of the variables
+ * @return Whether the variable was accessed successfully
+ */
+UE_API EPropertyBagResult WriteModuleVariable(FModuleContext&& InModuleContext, const FAnimNextVariableReference& InVariable, const FAnimNextParamType& InType, TFunctionRef<void(TArrayView<uint8>)> InFunction);
+
+} // UE::UAF
+
+#undef UE_API
\ No newline at end of file
