Fix non-unity build (Issue 1019001)

#rnx

diff --git a/Engine/Plugins/Experimental/UAF/UAF/Source/UAF/Private/Module/ModuleTickFunction.cpp b/Engine/Plugins/Experimental/UAF/UAF/Source/UAF/Private/Module/ModuleTickFunction.cpp
--- a/Engine/Plugins/Experimental/UAF/UAF/Source/UAF/Private/Module/ModuleTickFunction.cpp
+++ b/Engine/Plugins/Experimental/UAF/UAF/Source/UAF/Private/Module/ModuleTickFunction.cpp
@@ -8,6 +8,7 @@
 #include "Algo/TopologicalSort.h"
 #include "Component/AnimNextWorldSubsystem.h"
 #include "Module/AnimNextModule.h"
+#include "Module/AnimNextModuleInstance.h"
 #include "Module/ModuleTaskContext.h"
 #include "Module/RigUnit_AnimNextModuleEvents.h"
 #include "Module/ModuleEvents.h"
