UAF: Fix crashes undoing/redoing function creation

Additionally fix crash attempting to compile an asset after a function (the last edited document) has been removed/closed. This deactivates the editor mode that provides the ability for the compile button to function.

diff --git a/Engine/Plugins/Experimental/UAF/UAF/Source/UAFEditor/Internal/Common/AnimNextAssetItemDetails.cpp b/Engine/Plugins/Experimental/UAF/UAF/Source/UAFEditor/Internal/Common/AnimNextAssetItemDetails.cpp
--- a/Engine/Plugins/Experimental/UAF/UAF/Source/UAFEditor/Internal/Common/AnimNextAssetItemDetails.cpp
+++ b/Engine/Plugins/Experimental/UAF/UAF/Source/UAFEditor/Internal/Common/AnimNextAssetItemDetails.cpp
@@ -287,7 +287,6 @@ void FAnimNextAssetItemDetails::RegisterToolMenuExtensions()
 			FUIAction(
 				FExecuteAction::CreateWeakLambda(WorkspaceItemContext, [WeakWorkspaceEditor = TWeakPtr<UE::Workspace::IWorkspaceEditor>(WorkspaceEditor), EditorData]()
 				{
-					FScopedTransaction Transaction(LOCTEXT("AddFunctionTransaction", "Add Function"));
 					URigVMLibraryNode* NewFunction = EditorData->AddFunction(TEXT("NewFunction"), true);
 
 					// Open the new function's graph
diff --git a/Engine/Plugins/Experimental/UAF/UAF/Source/UAFEditor/Private/AnimNextEditorModule.cpp b/Engine/Plugins/Experimental/UAF/UAF/Source/UAFEditor/Private/AnimNextEditorModule.cpp
--- a/Engine/Plugins/Experimental/UAF/UAF/Source/UAFEditor/Private/AnimNextEditorModule.cpp
+++ b/Engine/Plugins/Experimental/UAF/UAF/Source/UAFEditor/Private/AnimNextEditorModule.cpp
@@ -1218,6 +1218,47 @@ void FAnimNextEditorModule::RegisterWorkspaceDocumentTypes(Workspace::IWorkspace
 			}
 		}
 	});
+	GraphArgs.OnGraphDocumentCreated = Workspace::FOnGraphDocumentCreated::CreateLambda([](const Workspace::FWorkspaceEditorContext& InContext, TSharedPtr<SWidget> InGraphWidget)
+	{
+		if (UAnimNextEdGraph* EdGraph = InContext.Document.GetTypedObject<UAnimNextEdGraph>())
+		{
+			if(UAnimNextRigVMAssetEditorData* EditorData = EdGraph->GetTypedOuter<UAnimNextRigVMAssetEditorData>())
+			{
+				EditorData->RigVMGraphModifiedEvent.AddSPLambda(InGraphWidget.Get(), [EditorData, WeakWorkspaceEditor = TWeakPtr<Workspace::IWorkspaceEditor>(InContext.WorkspaceEditor)](ERigVMGraphNotifType InNotifType, URigVMGraph* InGraph, UObject* InSubject)
+				{
+					switch (InNotifType)
+					{
+						case ERigVMGraphNotifType::NodeRemoved:
+						{
+							if (URigVMCollapseNode* CollapseNode = Cast<URigVMCollapseNode>(InSubject))
+							{
+								if (UEdGraph* EdGraph = Cast<UEdGraph>(EditorData->GetEditorObjectForRigVMGraph(CollapseNode->GetContainedGraph())))
+								{
+									if (TSharedPtr<Workspace::IWorkspaceEditor> WorkspaceEditor = WeakWorkspaceEditor.Pin())
+									{
+										WorkspaceEditor->CloseObjects({ EdGraph });
+										WorkspaceEditor->SetDetailsObjects({});
+									}
+								}
+							}
+							else if(URigVMFunctionReferenceNode* FunctionRefNode = Cast<URigVMFunctionReferenceNode>(InSubject))
+							{
+								if (TSharedPtr<Workspace::IWorkspaceEditor> WorkspaceEditor = WeakWorkspaceEditor.Pin())
+								{
+									WorkspaceEditor->SetDetailsObjects({});
+								}
+							}
+							break;
+						}
+						default:
+						{
+							break;
+						}
+					}
+				});
+			}
+		}
+	});
 
 	Workspace::FObjectDocumentArgs GraphDocumentArgs = WorkspaceEditorModule.CreateGraphDocumentArgs(GraphArgs);
 	Workspace::FOnMakeDocumentWidget WorkspaceMakeDocumentWidgetDelegate = GraphDocumentArgs.OnMakeDocumentWidget;
diff --git a/Engine/Plugins/Experimental/UAF/UAF/Source/UAFEditor/Private/Workspace/AnimNextWorkspaceEditorMode.cpp b/Engine/Plugins/Experimental/UAF/UAF/Source/UAFEditor/Private/Workspace/AnimNextWorkspaceEditorMode.cpp
--- a/Engine/Plugins/Experimental/UAF/UAF/Source/UAFEditor/Private/Workspace/AnimNextWorkspaceEditorMode.cpp
+++ b/Engine/Plugins/Experimental/UAF/UAF/Source/UAFEditor/Private/Workspace/AnimNextWorkspaceEditorMode.cpp
@@ -66,6 +66,17 @@ void UAnimNextWorkspaceEditorMode::Enter()
 
 void UAnimNextWorkspaceEditorMode::Exit()
 {
+	TSharedPtr<UE::Workspace::IWorkspaceEditor> WorkspaceEditor = GetWorkspaceEditor();
+	if(WorkspaceEditor.IsValid())
+	{
+		const TSharedRef<FUICommandList>& CommandList = WorkspaceEditor->GetToolkitCommands();
+		const UE::UAF::FRigVMCommands& RigVMCommands = UE::UAF::FRigVMCommands::Get();
+		CommandList->UnmapAction(RigVMCommands.Compile);
+		CommandList->UnmapAction(RigVMCommands.AutoCompile);
+		CommandList->UnmapAction(RigVMCommands.CompileWholeWorkspace);
+		CommandList->UnmapAction(RigVMCommands.CompileDirtyFiles);
+	}
+
 	Super::Exit();
 
 	for (TPair<FObjectKey, TSharedRef<UE::UAF::Editor::IAssetCompilationHandler>>& AssetCompiler : AssetCompilers)
@@ -107,21 +118,22 @@ void UAnimNextWorkspaceEditorMode::BindCommands()
 
 	const UE::UAF::FRigVMCommands& RigVMCommands = UE::UAF::FRigVMCommands::Get();
 	ToolkitCommands->MapAction(RigVMCommands.Compile,
-		FExecuteAction::CreateUObject(this, &UAnimNextWorkspaceEditorMode::HandleCompile));
+		FExecuteAction::CreateUObject(this, &UAnimNextWorkspaceEditorMode::HandleCompile),
+		FCanExecuteAction::CreateUObject(this, &UAnimNextWorkspaceEditorMode::CanCompile));
 
 	ToolkitCommands->MapAction(RigVMCommands.AutoCompile,
 		FExecuteAction::CreateUObject(this, &UAnimNextWorkspaceEditorMode::HandleAutoCompile),
-		FCanExecuteAction(),
+		FCanExecuteAction::CreateUObject(this, &UAnimNextWorkspaceEditorMode::CanCompile),
 		FIsActionChecked::CreateUObject(this, &UAnimNextWorkspaceEditorMode::IsAutoCompileChecked));
 
 	ToolkitCommands->MapAction(RigVMCommands.CompileWholeWorkspace,
 		FExecuteAction::CreateUObject(this, &UAnimNextWorkspaceEditorMode::HandleCompileWholeWorkspace),
-		FCanExecuteAction(),
+		FCanExecuteAction::CreateUObject(this, &UAnimNextWorkspaceEditorMode::CanCompile),
 		FIsActionChecked::CreateUObject(this, &UAnimNextWorkspaceEditorMode::IsCompileWholeWorkspaceChecked));
 
 	ToolkitCommands->MapAction(RigVMCommands.CompileDirtyFiles,
 		FExecuteAction::CreateUObject(this, &UAnimNextWorkspaceEditorMode::HandleCompileDirtyFiles),
-		FCanExecuteAction(),
+		FCanExecuteAction::CreateUObject(this, &UAnimNextWorkspaceEditorMode::CanCompile),
 		FIsActionChecked::CreateUObject(this, &UAnimNextWorkspaceEditorMode::IsCompileDirtyFilesChecked));
 }
 
@@ -204,6 +216,12 @@ void UAnimNextWorkspaceEditorMode::HandleCompile()
 	}
 }
 
+bool UAnimNextWorkspaceEditorMode::CanCompile() const
+{
+	TSharedPtr<UE::Workspace::IWorkspaceEditor> WorkspaceEditor = GetWorkspaceEditor();
+	return WorkspaceEditor.IsValid();
+}
+
 void UAnimNextWorkspaceEditorMode::HandleAutoCompile()
 {
 	State.bAutoCompile = !State.bAutoCompile;
diff --git a/Engine/Plugins/Experimental/UAF/UAF/Source/UAFEditor/Private/Workspace/AnimNextWorkspaceEditorMode.h b/Engine/Plugins/Experimental/UAF/UAF/Source/UAFEditor/Private/Workspace/AnimNextWorkspaceEditorMode.h
--- a/Engine/Plugins/Experimental/UAF/UAF/Source/UAFEditor/Private/Workspace/AnimNextWorkspaceEditorMode.h
+++ b/Engine/Plugins/Experimental/UAF/UAF/Source/UAFEditor/Private/Workspace/AnimNextWorkspaceEditorMode.h
@@ -48,6 +48,7 @@ private:
 	virtual void BindCommands() override;
 	virtual bool UsesToolkits() const override { return false; }
 	void HandleCompile();
+	bool CanCompile() const;
 	void HandleAutoCompile();
 	bool IsAutoCompileChecked() const;
 	void HandleCompileWholeWorkspace();
diff --git a/Engine/Plugins/Experimental/UAF/UAF/Source/UAFUncookedOnly/Internal/AnimNextRigVMAssetEditorData.h b/Engine/Plugins/Experimental/UAF/UAF/Source/UAFUncookedOnly/Internal/AnimNextRigVMAssetEditorData.h
--- a/Engine/Plugins/Experimental/UAF/UAF/Source/UAFUncookedOnly/Internal/AnimNextRigVMAssetEditorData.h
+++ b/Engine/Plugins/Experimental/UAF/UAF/Source/UAFUncookedOnly/Internal/AnimNextRigVMAssetEditorData.h
@@ -340,7 +340,7 @@ protected:
 	UE_API virtual void CreateEdGraphForCollapseNode(URigVMCollapseNode* InNode, bool bForce);
 
 	// Destroy a UEdGraph that corresponds to a URigVMCollapseNode
-	UE_API virtual void RemoveEdGraphForCollapseNode(URigVMCollapseNode* InNode, bool bNotify);
+	UE_API virtual bool RemoveEdGraphForCollapseNode(URigVMCollapseNode* InNode, bool bNotify);
 
 	// Remove the UEdGraph that corresponds to a URigVMGraph
 	UE_API virtual bool RemoveEdGraph(URigVMGraph* InModel);
diff --git a/Engine/Plugins/Experimental/UAF/UAF/Source/UAFUncookedOnly/Private/AnimNextRigVMAssetEditorData.cpp b/Engine/Plugins/Experimental/UAF/UAF/Source/UAFUncookedOnly/Private/AnimNextRigVMAssetEditorData.cpp
--- a/Engine/Plugins/Experimental/UAF/UAF/Source/UAFUncookedOnly/Private/AnimNextRigVMAssetEditorData.cpp
+++ b/Engine/Plugins/Experimental/UAF/UAF/Source/UAFUncookedOnly/Private/AnimNextRigVMAssetEditorData.cpp
@@ -1363,7 +1363,7 @@ void UAnimNextRigVMAssetEditorData::HandleModifiedEvent(ERigVMGraphNotifType InN
 	{
 		if (URigVMCollapseNode* CollapseNode = Cast<URigVMCollapseNode>(InSubject))
 		{
-			RemoveEdGraphForCollapseNode(CollapseNode, false);
+			bNotifForOthersPending = !RemoveEdGraphForCollapseNode(CollapseNode, true);
 			break;
 		}
 		RequestAutoVMRecompilation();
@@ -1985,7 +1985,7 @@ void UAnimNextRigVMAssetEditorData::CreateEdGraphForCollapseNode(URigVMCollapseN
 	}
 }
 
-void UAnimNextRigVMAssetEditorData::RemoveEdGraphForCollapseNode(URigVMCollapseNode* InNode, bool bNotify)
+bool UAnimNextRigVMAssetEditorData::RemoveEdGraphForCollapseNode(URigVMCollapseNode* InNode, bool bNotify)
 {
 	check(InNode);
 
@@ -2012,7 +2012,7 @@ void UAnimNextRigVMAssetEditorData::RemoveEdGraphForCollapseNode(URigVMCollapseN
 						FunctionEdGraphs.Remove(RigFunctionGraph);
 						RigFunctionGraph->Rename(nullptr, GetTransientPackage(), REN_ForceNoResetLoaders | REN_DontCreateRedirectors);
 						RigFunctionGraph->MarkAsGarbage();
-						break;
+						return bNotify;
 					}
 				}
 			}
@@ -2041,12 +2041,14 @@ void UAnimNextRigVMAssetEditorData::RemoveEdGraphForCollapseNode(URigVMCollapseN
 						RigGraph->SubGraphs.Remove(SubRigGraph);
 						SubRigGraph->Rename(nullptr, GetTransientPackage(), REN_ForceNoResetLoaders | REN_DontCreateRedirectors);
 						SubRigGraph->MarkAsGarbage();
-						break;
+						return bNotify;
 					}
 				}
 			}
 		}
 	}
+
+	return false;
 }
 
 UEdGraph* UAnimNextRigVMAssetEditorData::CreateEdGraph(URigVMGraph* InRigVMGraph, bool bForce)
@@ -2534,11 +2536,6 @@ bool UAnimNextRigVMAssetLibrary::RenameCategory(UAnimNextRigVMAsset* InAsset, co
 
 URigVMLibraryNode* UAnimNextRigVMAssetEditorData::AddFunction(FName InFunctionName, bool bInMutable, bool bInSetupUndoRedo, bool bPrintPythonCommand)
 {
-	if (bInSetupUndoRedo)
-	{
-		Modify();
-	}
-
 	URigVMController* Controller = RigVMClient.GetOrCreateController(GetLocalFunctionLibrary());
 	URigVMLibraryNode* Node = Controller->AddFunctionToLibrary(InFunctionName, bInMutable, FVector2D::ZeroVector, bInSetupUndoRedo, false);
 
diff --git a/Engine/Plugins/Experimental/Workspace/Source/WorkspaceEditor/Private/WorkspaceEditor.cpp b/Engine/Plugins/Experimental/Workspace/Source/WorkspaceEditor/Private/WorkspaceEditor.cpp
--- a/Engine/Plugins/Experimental/Workspace/Source/WorkspaceEditor/Private/WorkspaceEditor.cpp
+++ b/Engine/Plugins/Experimental/Workspace/Source/WorkspaceEditor/Private/WorkspaceEditor.cpp
@@ -776,11 +776,24 @@ void FWorkspaceEditor::SaveAsset_Execute()
 void FWorkspaceEditor::CloseDocumentTab(const UObject* DocumentID) const
 {
 	UWorkspaceState* State = Workspace->GetState();
-	FWorkspaceOutlinerItemExport Export;
-	State->DocumentStates.Remove(TInstancedStruct<FWorkspaceDocumentState>::Make(DocumentID, Export));
 
-	const TSharedRef<FTabPayload_WorkspaceDocument> Payload = FTabPayload_WorkspaceDocument::Make(DocumentID);
-	DocumentManager->CloseTab(Payload);
+	// Close all exports that reference the document's object
+	FSoftObjectPath ObjectPath(DocumentID);
+	for (const TInstancedStruct<FWorkspaceDocumentState>& InstancedStruct : State->DocumentStates)
+	{
+		const FWorkspaceDocumentState& DocumentState = InstancedStruct.Get<FWorkspaceDocumentState>();
+		if (DocumentState.Object == ObjectPath)
+		{
+			const TSharedRef<FTabPayload_WorkspaceDocument> Payload = FTabPayload_WorkspaceDocument::Make(DocumentID, DocumentState.Export);
+			DocumentManager->CloseTab(Payload);
+		}
+	}
+
+	// Clean up document states
+	State->DocumentStates.RemoveAll([&ObjectPath](const TInstancedStruct<FWorkspaceDocumentState>& InDocumentState)
+	{
+		return InDocumentState.Get<FWorkspaceDocumentState>().Object == ObjectPath;
+	});
 }
 
 bool FWorkspaceEditor::InEditingMode() const 
diff --git a/Engine/Plugins/Experimental/Workspace/Source/WorkspaceEditor/Private/WorkspaceEditorModule.cpp b/Engine/Plugins/Experimental/Workspace/Source/WorkspaceEditor/Private/WorkspaceEditorModule.cpp
--- a/Engine/Plugins/Experimental/Workspace/Source/WorkspaceEditor/Private/WorkspaceEditorModule.cpp
+++ b/Engine/Plugins/Experimental/Workspace/Source/WorkspaceEditor/Private/WorkspaceEditorModule.cpp
@@ -94,7 +94,7 @@ FObjectDocumentArgs FWorkspaceEditorModule::CreateGraphDocumentArgs(const FGraph
 	Args.OnMakeDocumentWidget = FOnMakeDocumentWidget::CreateLambda([InArgs](const FWorkspaceEditorContext& InContext)
 	{
 		TWeakPtr<FWorkspaceEditor> WeakWorkspaceEditor = StaticCastSharedRef<FWorkspaceEditor>(InContext.WorkspaceEditor);
-		return SNew(SGraphDocument, StaticCastSharedRef<FWorkspaceEditor>(InContext.WorkspaceEditor), InContext.Document)
+		TSharedRef<SWidget> Widget = SNew(SGraphDocument, StaticCastSharedRef<FWorkspaceEditor>(InContext.WorkspaceEditor), InContext.Document)
 			.OnCreateActionMenu_Lambda([OnCreateActionMenu = InArgs.OnCreateActionMenu](const FWorkspaceEditorContext& InContext, UEdGraph* InGraph, const FVector2D& InNodePosition, const TArray<UEdGraphPin*>& InDraggedPins, bool bInAutoExpand, SGraphEditor::FActionMenuClosed InOnMenuClosed)
 			{
 				if(OnCreateActionMenu.IsBound())
@@ -118,6 +118,11 @@ FObjectDocumentArgs FWorkspaceEditorModule::CreateGraphDocumentArgs(const FGraph
 			.OnNavigateHistoryForward_Lambda([WeakWorkspaceEditor](){ if(const TSharedPtr<FWorkspaceEditor> SharedWorkspaceEditor = WeakWorkspaceEditor.Pin()) SharedWorkspaceEditor->NavigateForward(); })
 			.OnNavigateHistoryBack_Lambda([WeakWorkspaceEditor](){ if(const TSharedPtr<FWorkspaceEditor> SharedWorkspaceEditor = WeakWorkspaceEditor.Pin()) SharedWorkspaceEditor->NavigateBack(); })
 			.OnNodeDoubleClicked(InArgs.OnNodeDoubleClicked);
+
+		// Let listeners know that the document widget has been created 
+		InArgs.OnGraphDocumentCreated.ExecuteIfBound(InContext, Widget);
+
+		return Widget;
 	});
 	Args.OnGetTabIcon = FOnGetTabIcon::CreateLambda([](const FWorkspaceEditorContext& InContext)
 	{
diff --git a/Engine/Plugins/Experimental/Workspace/Source/WorkspaceEditor/Public/IWorkspaceEditorModule.h b/Engine/Plugins/Experimental/Workspace/Source/WorkspaceEditor/Public/IWorkspaceEditorModule.h
--- a/Engine/Plugins/Experimental/Workspace/Source/WorkspaceEditor/Public/IWorkspaceEditorModule.h
+++ b/Engine/Plugins/Experimental/Workspace/Source/WorkspaceEditor/Public/IWorkspaceEditorModule.h
@@ -187,6 +187,8 @@ using FOnGetWorkspaceDetailCustomizationInstance = TDelegate<TSharedRef<IDetailC
 
 using FOnNodeDoubleClicked = TDelegate<void(const FWorkspaceEditorContext&, const UEdGraphNode* InNode)>;
 
+using FOnGraphDocumentCreated = TDelegate<void(const FWorkspaceEditorContext&, TSharedPtr<SWidget> InGraphWidget)>;
+
 // Arguments used to make document widgets for graphs
 struct FGraphDocumentWidgetArgs
 {
@@ -218,6 +220,8 @@ struct FGraphDocumentWidgetArgs
 	FOnGraphSelectionChanged OnGraphSelectionChanged;
 
 	FOnNodeDoubleClicked OnNodeDoubleClicked;
+
+	FOnGraphDocumentCreated OnGraphDocumentCreated;
 };
 
 // Enum describing how to open a workspace
