Fix crash adding 2+ SetVariable nodes for two different variables in UAF function graph.

#tests Editor

diff --git a/Engine/Plugins/Experimental/UAF/UAF/Source/UAFEditor/Private/Workspace/AnimNextWorkspaceEditorMode.cpp b/Engine/Plugins/Experimental/UAF/UAF/Source/UAFEditor/Private/Workspace/AnimNextWorkspaceEditorMode.cpp
--- a/Engine/Plugins/Experimental/UAF/UAF/Source/UAFEditor/Private/Workspace/AnimNextWorkspaceEditorMode.cpp
+++ b/Engine/Plugins/Experimental/UAF/UAF/Source/UAFEditor/Private/Workspace/AnimNextWorkspaceEditorMode.cpp
@@ -13,6 +13,8 @@
 #include "IWorkspaceEditor.h"
 #include "IWorkspaceEditorModule.h"
 #include "RigVMCommands.h"
+#include "RigVMCore/RigVMGraphFunctionHost.h"
+#include "RigVMModel/Nodes/RigVMFunctionReferenceNode.h"
 #include "UncookedOnlyUtils.h"
 #include "Logging/MessageLog.h"
 #include "Modules/ModuleManager.h"
@@ -375,6 +377,26 @@ void UAnimNextWorkspaceEditorMode::HandleRigVMModifiedEvent(ERigVMGraphNotifType
 	}
 
 	UAnimNextRigVMAssetEditorData* EditorData = InGraph->GetTypedOuter<UAnimNextRigVMAssetEditorData>();
+
+	// Our RigVM Graph may be a function, in this case the graph has a temp outer not the editor data. 
+	// So search for the host object of the function instead.
+	if (!EditorData)
+	{
+		if (URigVMFunctionReferenceNode* FunctionReferenceNode = Cast<URigVMFunctionReferenceNode>(InSubject))
+		{
+			constexpr bool bLoadIfNecessary = true;
+			if (UAnimNextRigVMAssetEditorData* FunctionHost = Cast<UAnimNextRigVMAssetEditorData>(FunctionReferenceNode->GetReferencedFunctionHeader().GetFunctionHost(bLoadIfNecessary)))
+			{
+				EditorData = FunctionHost;
+			}
+		}
+	}
+	
+	if (!ensureMsgf(EditorData, TEXT("Failed to find UAF asset for RigVM modification, asset may not update correctly. It also may fail to propagate dirty status to other assets.")))
+	{
+		return;
+	}
+
 	if(EditorData->IsDirtyForRecompilation())
 	{
 		CompileStatus = UE::UAF::Editor::ECompileStatus::Dirty;
