Fix missing drag-drop op for adding attributes to sets

diff --git a/Engine/Plugins/Experimental/UAF/UAF/Source/UAFEditor/Private/UAF/AbstractSkeleton/Sets/AttributeDragDrop.cpp b/Engine/Plugins/Experimental/UAF/UAF/Source/UAFEditor/Private/UAF/AbstractSkeleton/Sets/AttributeDragDrop.cpp
new file mode 100644
--- /dev/null
+++ b/Engine/Plugins/Experimental/UAF/UAF/Source/UAFEditor/Private/UAF/AbstractSkeleton/Sets/AttributeDragDrop.cpp
@@ -0,0 +1,29 @@
+// Copyright Epic Games, Inc. All Rights Reserved.
+
+#include "UAF/AbstractSkeleton/Sets/AttributeDragDrop.h"
+
+namespace UE::UAF::Sets
+{
+	TSharedRef<FAttributeDragDropOp> FAttributeDragDropOp::New(TArray<FAnimationAttributeIdentifier>&& InAttributes)
+	{
+		TSharedRef<FAttributeDragDropOp> Operation = MakeShared<FAttributeDragDropOp>();
+		Operation->Attributes = InAttributes;
+
+		Operation->Construct();
+
+		return Operation;
+	}
+
+	TSharedPtr<SWidget> FAttributeDragDropOp::GetDefaultDecorator() const
+	{
+		return SNew(SBorder)
+			.Visibility(EVisibility::Visible)
+			.BorderImage(FAppStyle::GetBrush("Menu.Background"))
+			[
+				SNew(STextBlock)
+				.Text(Attributes.Num() == 1
+					? FText::FromName(Attributes[0].GetName())
+					: FText::Format(NSLOCTEXT("UE::UAF::Sets::FAttributeDragDropOp", "MultipleAttributesTooltip", "{0} attributes"), FText::AsNumber(Attributes.Num())))
+			];
+	}
+}
\ No newline at end of file
diff --git a/Engine/Plugins/Experimental/UAF/UAF/Source/UAFEditor/Private/UAF/AbstractSkeleton/Sets/AttributeDragDrop.h b/Engine/Plugins/Experimental/UAF/UAF/Source/UAFEditor/Private/UAF/AbstractSkeleton/Sets/AttributeDragDrop.h
new file mode 100644
--- /dev/null
+++ b/Engine/Plugins/Experimental/UAF/UAF/Source/UAFEditor/Private/UAF/AbstractSkeleton/Sets/AttributeDragDrop.h
@@ -0,0 +1,22 @@
+// Copyright Epic Games, Inc. All Rights Reserved.
+
+#pragma once
+
+#include "CoreMinimal.h"
+#include "DragAndDrop/DecoratedDragDropOp.h"
+#include "Animation/AnimData/AttributeIdentifier.h"
+
+namespace UE::UAF::Sets
+{
+	class FAttributeDragDropOp : public FDecoratedDragDropOp
+	{
+	public:
+		DRAG_DROP_OPERATOR_TYPE(FAttributeDragDropOp, FDecoratedDragDropOp)
+
+		static TSharedRef<FAttributeDragDropOp> New(TArray<FAnimationAttributeIdentifier>&& InAttributes);
+
+		TSharedPtr<SWidget> GetDefaultDecorator() const override;
+
+		TArray<FAnimationAttributeIdentifier> Attributes;
+	};
+}
diff --git a/Engine/Plugins/Experimental/UAF/UAF/Source/UAFEditor/Private/UAF/AbstractSkeleton/Sets/SAttributesList.cpp b/Engine/Plugins/Experimental/UAF/UAF/Source/UAFEditor/Private/UAF/AbstractSkeleton/Sets/SAttributesList.cpp
--- a/Engine/Plugins/Experimental/UAF/UAF/Source/UAFEditor/Private/UAF/AbstractSkeleton/Sets/SAttributesList.cpp
+++ b/Engine/Plugins/Experimental/UAF/UAF/Source/UAFEditor/Private/UAF/AbstractSkeleton/Sets/SAttributesList.cpp
@@ -2,6 +2,7 @@
 
 #include "SAttributesList.h"
 
+#include "AttributeDragDrop.h"
 #include "Animation/BuiltInAttributeTypes.h"
 #include "Framework/Application/SlateApplication.h"
 #include "Framework/MultiBox/MultiBoxBuilder.h"
@@ -116,6 +117,7 @@ public:
 		SLATE_ARGUMENT(SAttributesList::FListItemPtr, Item)
 		SLATE_EVENT(FOnCanAcceptDrop, OnCanAcceptDrop)
 		SLATE_EVENT(FOnAcceptDrop, OnAcceptDrop)
+		SLATE_EVENT(FOnDragDetected, OnDragDetected)
 	SLATE_END_ARGS()
 
 	void Construct(const FArguments& InArgs, const TSharedRef<STableViewBase>& OwnerTableView)
@@ -123,6 +125,7 @@ public:
 		Item = InArgs._Item;
 
 		const FSuperRowType::FArguments Args = FSuperRowType::FArguments()
+			.OnDragDetected(InArgs._OnDragDetected)
 			.OnCanAcceptDrop(InArgs._OnCanAcceptDrop)
 			.OnAcceptDrop(InArgs._OnAcceptDrop)
 			.Style(FAppStyle::Get(), "TableView.AlternatingRow");
@@ -235,6 +238,24 @@ TSharedRef<ITableRow> SAttributesList::ListView_OnGenerateRow(FListItemPtr Item,
 					return FReply::Handled();
 				}
 
+				return FReply::Unhandled();
+			})
+		.OnDragDetected_Lambda([this](const FGeometry& MyGeometry, const FPointerEvent& MouseEvent)
+			{
+				if (MouseEvent.IsMouseButtonDown(EKeys::LeftMouseButton))
+				{
+					TArray<FAnimationAttributeIdentifier> SelectedAttributes;
+					SelectedAttributes.Reserve(ListView->GetNumItemsSelected());
+					
+					for (const FListItemPtr& Selected : ListView->GetSelectedItems())
+					{
+						SelectedAttributes.Add(Selected->Attribute);
+					}
+					
+					const TSharedRef<FAttributeDragDropOp> DragDropOp = FAttributeDragDropOp::New(MoveTemp(SelectedAttributes));
+					return FReply::Handled().BeginDragDrop(DragDropOp);
+				}
+
 				return FReply::Unhandled();
 			});
 }
diff --git a/Engine/Plugins/Experimental/UAF/UAF/Source/UAFEditor/Private/UAF/AbstractSkeleton/Sets/SSetBinding.cpp b/Engine/Plugins/Experimental/UAF/UAF/Source/UAFEditor/Private/UAF/AbstractSkeleton/Sets/SSetBinding.cpp
--- a/Engine/Plugins/Experimental/UAF/UAF/Source/UAFEditor/Private/UAF/AbstractSkeleton/Sets/SSetBinding.cpp
+++ b/Engine/Plugins/Experimental/UAF/UAF/Source/UAFEditor/Private/UAF/AbstractSkeleton/Sets/SSetBinding.cpp
@@ -2,6 +2,7 @@
 
 #include "UAF/AbstractSkeleton/Sets/SSetBinding.h"
 
+#include "AttributeDragDrop.h"
 #include "Framework/Commands/UICommandList.h"
 #include "Framework/MultiBox/MultiBoxBuilder.h"
 #include "ScopedTransaction.h"
@@ -266,8 +267,12 @@ void SSetBinding::Construct(const FArguments& InArgs, TWeakObjectPtr<UAbstractSk
 
 					TOptional<EItemDropZone> ReturnedDropZone;
 
-					const TSharedPtr<FSkeletonTreeDragDropOp> DragDropOp = DragDropEvent.GetOperationAs<FSkeletonTreeDragDropOp>();
-					if (DragDropOp.IsValid())
+					if (const TSharedPtr<FSkeletonTreeDragDropOp> DragDropOp = DragDropEvent.GetOperationAs<FSkeletonTreeDragDropOp>())
+					{
+						ReturnedDropZone = EItemDropZone::OntoItem;
+					}
+
+					if (const TSharedPtr<FAttributeDragDropOp> DragDropOp = DragDropEvent.GetOperationAs<FAttributeDragDropOp>())
 					{
 						ReturnedDropZone = EItemDropZone::OntoItem;
 					}
@@ -276,8 +281,7 @@ void SSetBinding::Construct(const FArguments& InArgs, TWeakObjectPtr<UAbstractSk
 				})
 			.OnAcceptDrop_Lambda([this](const FDragDropEvent& DragDropEvent, EItemDropZone DropZone, TSharedPtr<ITreeItem> TargetTreeItem) -> FReply
 				{
-					const TSharedPtr<FSkeletonTreeDragDropOp> DragDropOp = DragDropEvent.GetOperationAs<FSkeletonTreeDragDropOp>();
-					if (DragDropOp.IsValid())
+					if (const TSharedPtr<FSkeletonTreeDragDropOp> DragDropOp = DragDropEvent.GetOperationAs<FSkeletonTreeDragDropOp>())
 					{
 						TSharedPtr<FTreeItem_Set> TargetSetTreeItem = StaticCastSharedPtr<FTreeItem_Set>(TargetTreeItem);
 
@@ -299,6 +303,30 @@ void SSetBinding::Construct(const FArguments& InArgs, TWeakObjectPtr<UAbstractSk
 
 						return FReply::Handled();
 					}
+
+					if (const TSharedPtr<FAttributeDragDropOp> DragDropOp = DragDropEvent.GetOperationAs<FAttributeDragDropOp>())
+					{
+						TSharedPtr<FTreeItem_Set> TargetSetTreeItem = StaticCastSharedPtr<FTreeItem_Set>(TargetTreeItem);
+						
+						const FScopedTransaction Transaction(LOCTEXT("AddAttributeToSet", "Add Attribute to Set"));
+						SetBindingWidget.SetBinding->Modify();
+
+						bool bTreeDirty = false;
+									
+						for (const FAnimationAttributeIdentifier& Attribute : DragDropOp->Attributes)
+						{
+							bTreeDirty |= SetBindingWidget.SetBinding->RemoveAttributeFromSet(Attribute);
+							bTreeDirty |= SetBindingWidget.SetBinding->AddAttributeToSet(Attribute, TargetSetTreeItem->Set.SetName);
+						}
+							
+						if (bTreeDirty)
+						{
+							SetBindingWidget.RepopulateTreeData();
+						}
+
+						return FReply::Handled();
+					}
+					
 					return FReply::Unhandled();
 				})
 			.ShowWires(true)
