[AnimFramework] fixed trace analysis re-entrance issue

diff --git a/Engine/Plugins/Experimental/UAF/UAF/Source/UAFEditor/Private/RewindDebugger/AnimNextModuleTrack.cpp b/Engine/Plugins/Experimental/UAF/UAF/Source/UAFEditor/Private/RewindDebugger/AnimNextModuleTrack.cpp
--- a/Engine/Plugins/Experimental/UAF/UAF/Source/UAFEditor/Private/RewindDebugger/AnimNextModuleTrack.cpp
+++ b/Engine/Plugins/Experimental/UAF/UAF/Source/UAFEditor/Private/RewindDebugger/AnimNextModuleTrack.cpp
@@ -214,26 +214,33 @@ bool FAnimNextModuleTrack::UpdateInternal()
 
 bool FAnimNextModuleTrack::HandleDoubleClickInternal()
 {
+	if (InstanceId == RewindDebugger::FObjectId::InvalidId)
+	{
+		return false;
+	}
+
 	const IRewindDebugger* RewindDebugger = IRewindDebugger::Instance();
 	if (const TraceServices::IAnalysisSession* AnalysisSession = RewindDebugger->GetAnalysisSession())
 	{
-		TraceServices::FAnalysisSessionReadScope SessionReadScope(*AnalysisSession);
+		const FObjectInfo* AssetInfo = nullptr;
+		{
+			TraceServices::FAnalysisSessionReadScope SessionReadScope(*AnalysisSession);
 
-		const IGameplayProvider* GameplayProvider = AnalysisSession->ReadProvider<IGameplayProvider>("GameplayProvider");
-		const FAnimNextProvider* AnimNextProvider = IRewindDebugger::Instance()->GetAnalysisSession()->ReadProvider<FAnimNextProvider>("AnimNextProvider");
+			const IGameplayProvider* GameplayProvider = AnalysisSession->ReadProvider<IGameplayProvider>("GameplayProvider");
+			const FAnimNextProvider* AnimNextProvider = IRewindDebugger::Instance()->GetAnalysisSession()->ReadProvider<FAnimNextProvider>("AnimNextProvider");
 
-		if (InstanceId != 0)
-		{
 			if (const FDataInterfaceData* ModuleData = AnimNextProvider->GetDataInterfaceData(InstanceId))
 			{
-				const FObjectInfo& AssetInfo = GameplayProvider->GetObjectInfo(ModuleData->AssetId);
-				if (!EnumHasAnyFlags(AssetInfo.Flags, EObjectInfoFlags::TransientObject))
-				{
-					GEditor->GetEditorSubsystem<UAssetEditorSubsystem>()->OpenEditorForAsset(AssetInfo.PathName);
-					return true;
-				}
+				AssetInfo = &GameplayProvider->GetObjectInfo(ModuleData->AssetId);
 			}
 		}
+
+		if (AssetInfo
+			&& !EnumHasAnyFlags(AssetInfo->Flags, EObjectInfoFlags::TransientObject))
+		{
+			GEditor->GetEditorSubsystem<UAssetEditorSubsystem>()->OpenEditorForAsset(AssetInfo->PathName);
+			return true;
+		}
 	}
 	return false;
 }
