Fix variables outliner miss-aligned with Native Shared Variables

#tests Editor

diff --git a/Engine/Plugins/Experimental/UAF/UAF/Source/UAFEditor/Private/Variables/Outliner/VariablesOutlinerColumns.cpp b/Engine/Plugins/Experimental/UAF/UAF/Source/UAFEditor/Private/Variables/Outliner/VariablesOutlinerColumns.cpp
--- a/Engine/Plugins/Experimental/UAF/UAF/Source/UAFEditor/Private/Variables/Outliner/VariablesOutlinerColumns.cpp
+++ b/Engine/Plugins/Experimental/UAF/UAF/Source/UAFEditor/Private/Variables/Outliner/VariablesOutlinerColumns.cpp
@@ -13,6 +13,7 @@
 #include "ScopedTransaction.h"
 #include "SPinTypeSelector.h"
 #include "UncookedOnlyUtils.h"
+#include "VariablesOutlinerAssetItem.h"
 #include "VariablesOutlinerEntryItem.h"
 #include "Entries/AnimNextSharedVariablesEntry.h"
 #include "Entries/AnimNextRigVMAssetEntry.h"
@@ -121,10 +122,17 @@ SHeaderRow::FColumn::FArguments FVariablesOutlinerTypeColumn::ConstructHeaderRow
 
 const TSharedRef<SWidget> FVariablesOutlinerTypeColumn::ConstructRowWidget(FSceneOutlinerTreeItemRef Item, const STableRow<FSceneOutlinerTreeItemPtr>& Row)
 {
+	const FVariablesOutlinerAssetItem* AssetItem = Item->CastTo<FVariablesOutlinerAssetItem>();
+	if (AssetItem)
+	{
+		return SNullWidget::NullWidget;
+	}
+
+	// Return a dummy box widget, with ShouldGenerateEmptyWidgetForSpacing settings this will prevent us from skipping outside of the header, avoiding alignment issues
 	const FVariablesOutlinerEntryItem* TreeItem = Item->CastTo<FVariablesOutlinerEntryItem>();
 	if (TreeItem == nullptr)
 	{
-		return SNullWidget::NullWidget;
+		return SNew(SBox);
 	}
 
 	return
@@ -464,16 +472,22 @@ class SVariablesOutlinerValue : public SCompoundWidget, public FNotifyHook
 
 const TSharedRef<SWidget> FVariablesOutlinerValueColumn::ConstructRowWidget(FSceneOutlinerTreeItemRef Item, const STableRow<FSceneOutlinerTreeItemPtr>& Row)
 {
-	FVariablesOutlinerEntryItem* TreeItem = Item->CastTo<FVariablesOutlinerEntryItem>();
-	if (TreeItem == nullptr)
+	const FVariablesOutlinerAssetItem* AssetItem = Item->CastTo<FVariablesOutlinerAssetItem>();
+	if (AssetItem)
 	{
 		return SNullWidget::NullWidget;
 	}
 
+	// Return a dummy box widget, with ShouldGenerateEmptyWidgetForSpacing settings this will prevent us from skipping outside of the header, avoiding alignment issues
+	FVariablesOutlinerEntryItem* TreeItem = Item->CastTo<FVariablesOutlinerEntryItem>();
+	if (TreeItem == nullptr)
+	{
+		return SNew(SBox);
+	}
 	IAnimNextRigVMVariableInterface* Variable = Cast<IAnimNextRigVMVariableInterface>(TreeItem->WeakEntry.Get());
 	if(Variable == nullptr)
 	{
-		return SNullWidget::NullWidget;
+		return SNew(SBox);
 	}
 	
 	return SNew(SVariablesOutlinerValue, *TreeItem, WeakSceneOutliner.Pin().ToSharedRef().Get(), Row);
