[AnimFramework] make sure session struct is not used outside the read scope

diff --git a/Engine/Plugins/Experimental/UAF/UAF/Source/UAFEditor/Private/RewindDebugger/AnimNextModuleTrack.cpp b/Engine/Plugins/Experimental/UAF/UAF/Source/UAFEditor/Private/RewindDebugger/AnimNextModuleTrack.cpp
--- a/Engine/Plugins/Experimental/UAF/UAF/Source/UAFEditor/Private/RewindDebugger/AnimNextModuleTrack.cpp
+++ b/Engine/Plugins/Experimental/UAF/UAF/Source/UAFEditor/Private/RewindDebugger/AnimNextModuleTrack.cpp
@@ -11,12 +11,12 @@
 #include "Editor.h"
 #include "IGameplayProvider.h"
 #include "IRewindDebugger.h"
-#include "Subsystems/AssetEditorSubsystem.h"
 #include "Modules/ModuleManager.h"
-#include "PropertyEditorModule.h"
 #include "ObjectAsTraceIdProxyArchiveReader.h"
-#include "StructUtils/PropertyBag.h"
+#include "PropertyEditorModule.h"
 #include "Serialization/MemoryReader.h"
+#include "StructUtils/PropertyBag.h"
+#include "Subsystems/AssetEditorSubsystem.h"
 
 #define LOCTEXT_NAMESPACE "AnimNextModuleTrack"
 
@@ -25,36 +25,35 @@ namespace UE::UAF::Editor
 
 static const FName AnimNextModulesName("AnimNextModules");
 
- FName FAnimNextModuleTrackCreator::GetTargetTypeNameInternal() const
- {
- 	static const FName ObjectName("UAFAssetInstance");
+FName FAnimNextModuleTrackCreator::GetTargetTypeNameInternal() const
+{
+	static const FName ObjectName("UAFAssetInstance");
 	return ObjectName;
 }
 
-	
 FText FAnimNextModuleTrack::GetDisplayNameInternal() const
 {
- 	if (DisplayNameCache.IsEmpty())
- 	{
- 		TraceServices::FAnalysisSessionReadScope SessionReadScope(*IRewindDebugger::Instance()->GetAnalysisSession());
- 		const FAnimNextProvider* AnimNextProvider = IRewindDebugger::Instance()->GetAnalysisSession()->ReadProvider<FAnimNextProvider>("AnimNextProvider");
- 		const IGameplayProvider* GameplayProvider = IRewindDebugger::Instance()->GetAnalysisSession()->ReadProvider<IGameplayProvider>("GameplayProvider");
-
- 		if (InstanceId != 0)
- 		{
- 			if (const FDataInterfaceData* Data = AnimNextProvider->GetDataInterfaceData(InstanceId))
- 			{
- 				const FObjectInfo& ModuleInfo = GameplayProvider->GetObjectInfo(Data->AssetId);
- 				DisplayNameCache = FText::FromString(ModuleInfo.Name);
- 				return DisplayNameCache;
- 			}
- 		}
-	
+	if (DisplayNameCache.IsEmpty())
+	{
+		TraceServices::FAnalysisSessionReadScope SessionReadScope(*IRewindDebugger::Instance()->GetAnalysisSession());
+		const FAnimNextProvider* AnimNextProvider = IRewindDebugger::Instance()->GetAnalysisSession()->ReadProvider<FAnimNextProvider>("AnimNextProvider");
+		const IGameplayProvider* GameplayProvider = IRewindDebugger::Instance()->GetAnalysisSession()->ReadProvider<IGameplayProvider>("GameplayProvider");
+
+		if (InstanceId != 0)
+		{
+			if (const FDataInterfaceData* Data = AnimNextProvider->GetDataInterfaceData(InstanceId))
+			{
+				const FObjectInfo& ModuleInfo = GameplayProvider->GetObjectInfo(Data->AssetId);
+				DisplayNameCache = FText::FromString(ModuleInfo.Name);
+				return DisplayNameCache;
+			}
+		}
+
 		// Don't cache this since it is a placeholder name, will be replaced once the real one is sent
- 		return NSLOCTEXT("RewindDebugger", "AnimNextSystemTrackName", "System");
- 	}
+		return NSLOCTEXT("RewindDebugger", "AnimNextSystemTrackName", "System");
+	}
 
- 	return DisplayNameCache;
+	return DisplayNameCache;
 }
 
 void FAnimNextModuleTrackCreator::GetTrackTypesInternal(TArray<RewindDebugger::FRewindDebuggerTrackType>& Types) const
@@ -70,38 +69,38 @@ TSharedPtr<RewindDebugger::FRewindDebuggerTrack> FAnimNextModuleTrackCreator::Cr
 FAnimNextModuleTrack::FAnimNextModuleTrack(uint64 InInstanceId) :
 	InstanceId(InInstanceId)
 {
-   	Initialize();
+	Initialize();
 }
-	
-	
+
+
 FAnimNextModuleTrack::~FAnimNextModuleTrack()
 {
- 	if (UPropertyBagDetailsObject* DetailsObject = DetailsObjectWeakPtr.Get())
-   	{
-   		DetailsObject->ClearFlags(RF_Standalone);
-   	}
+	if (UPropertyBagDetailsObject* DetailsObject = DetailsObjectWeakPtr.Get())
+	{
+		DetailsObject->ClearFlags(RF_Standalone);
+	}
 }
-	
+
 void FAnimNextModuleTrack::Initialize()
 {
 	ExistenceRange = MakeShared<SEventTimelineView::FTimelineEventData>();
 	ExistenceRange->Windows.Add({0,0, GetDisplayNameInternal(), GetDisplayNameInternal(), FLinearColor(0.1f,0.15f,0.11f)});
-	
-   	FPropertyEditorModule& PropertyEditorModule = FModuleManager::GetModuleChecked<FPropertyEditorModule>("PropertyEditor");
-   	FDetailsViewArgs DetailsViewArgs;
-   	DetailsViewArgs.NameAreaSettings = FDetailsViewArgs::HideNameArea;
-   	DetailsView = PropertyEditorModule.CreateDetailView(DetailsViewArgs);
+
+	FPropertyEditorModule& PropertyEditorModule = FModuleManager::GetModuleChecked<FPropertyEditorModule>("PropertyEditor");
+	FDetailsViewArgs DetailsViewArgs;
+	DetailsViewArgs.NameAreaSettings = FDetailsViewArgs::HideNameArea;
+	DetailsView = PropertyEditorModule.CreateDetailView(DetailsViewArgs);
 
 	InitializeDetailsObject();
 }
-	
+
 UPropertyBagDetailsObject* FAnimNextModuleTrack::InitializeDetailsObject()
 {
 	UPropertyBagDetailsObject* DetailsObject = NewObject<UPropertyBagDetailsObject>();
 	DetailsObject->SetFlags(RF_Standalone);
 	DetailsObjectWeakPtr = MakeWeakObjectPtr(DetailsObject);
 	DetailsView->SetObject(DetailsObject);
- 	return DetailsObject;
+	return DetailsObject;
 }
 
 
@@ -109,19 +108,20 @@ TSharedPtr<SWidget> FAnimNextModuleTrack::GetTimelineViewInternal()
 {
 	return SNew(SEventTimelineView)
 		.ViewRange_Lambda([]() { return IRewindDebugger::Instance()->GetCurrentViewRange(); })
-		.EventData_Raw(this, &FAnimNextModuleTrack::GetExistenceRange);}
+		.EventData_Raw(this, &FAnimNextModuleTrack::GetExistenceRange);
+}
 
 bool FAnimNextModuleTrack::UpdateInternal()
 {
 	TRACE_CPUPROFILER_EVENT_SCOPE(FAnimNextModuleTrack::UpdateInternal);
- 	
+
 	IRewindDebugger* RewindDebugger = IRewindDebugger::Instance();
- 	const TRange<double> ViewRange = RewindDebugger->GetCurrentViewRange();
-	
+	const TRange<double> ViewRange = RewindDebugger->GetCurrentViewRange();
+
 	const TraceServices::IAnalysisSession* AnalysisSession = RewindDebugger->GetAnalysisSession();
 
- 	bool bChanged = false;
-	
+	bool bChanged = false;
+
 	if (const FAnimNextProvider* AnimNextProvider = IRewindDebugger::Instance()->GetAnalysisSession()->ReadProvider<FAnimNextProvider>("AnimNextProvider"))
 	{
 		double CurrentScrubTime = IRewindDebugger::Instance()->CurrentTraceTime();
@@ -136,23 +136,23 @@ bool FAnimNextModuleTrack::UpdateInternal()
 			check(DetailsObject);
 		}
 
-		if(InstanceId != 0)
+		if (InstanceId != 0)
 		{
 			if (const FDataInterfaceData* Data = AnimNextProvider->GetDataInterfaceData(InstanceId))
 			{
 				if (PreviousScrubTime != CurrentScrubTime)
 				{
 					PreviousScrubTime = CurrentScrubTime;
-					
+
 					const TraceServices::IFrameProvider& FramesProvider = TraceServices::ReadFrameProvider(*AnalysisSession);
 					TraceServices::FAnalysisSessionReadScope SessionReadScope(*AnalysisSession);
 
-					
+
 					DetailsObject->Properties.SetNum(0, EAllowShrinking::No);
 					DetailsObject->NativeProperties.SetNum(0, EAllowShrinking::No);
 
 					TraceServices::FFrame MarkerFrame;
-					if(FramesProvider.GetFrameFromTime(ETraceFrameType::TraceFrameType_Game, CurrentScrubTime, MarkerFrame))
+					if (FramesProvider.GetFrameFromTime(ETraceFrameType::TraceFrameType_Game, CurrentScrubTime, MarkerFrame))
 					{
 						Data->VariablesTimeline.EnumerateEvents(MarkerFrame.StartTime, MarkerFrame.EndTime, [AnimNextProvider, DetailsObject, AnalysisSession, this](double InStartTime, double InEndTime, uint32 InDepth, const FPropertyVariableData& VariableListData)
 						{
@@ -166,16 +166,16 @@ bool FAnimNextModuleTrack::UpdateInternal()
 									// Can't do anything without a description
 									return TraceServices::EEventEnumerate::Continue;
 								}
-								
+
 								// Load the property descriptions
 								FMemoryReader DescriptionReader(DescriptionData->Data);
 								FObjectAsTraceIdProxyArchiveReader DescriptionReaderProxy(DescriptionReader, GameplayProvider);
 								DescriptionReaderProxy.UsingCustomVersion(UE::UAF::FAnimNextTrace::CustomVersionGUID);
 								DescriptionReaderProxy << PropertyDescriptions;
-								DetailsObject->Properties.SetNum(DetailsObject->Properties.Num()+1);
-								FInstancedPropertyBag &PropertyBag = DetailsObject->Properties.Last();
+								DetailsObject->Properties.SetNum(DetailsObject->Properties.Num() + 1);
+								FInstancedPropertyBag& PropertyBag = DetailsObject->Properties.Last();
 								PropertyBag.AddProperties(PropertyDescriptions);
-							
+
 								// Load the property values
 								FMemoryReader Reader(VariableListData.ValueData);
 								FObjectAsTraceIdProxyArchiveReader ReaderProxy(Reader, GameplayProvider);
@@ -190,10 +190,10 @@ bool FAnimNextModuleTrack::UpdateInternal()
 							{
 								FMemoryReader Reader(VariableListData.ValueData);
 								FObjectAsTraceIdProxyArchiveReader ReaderProxy(Reader, GameplayProvider);
-								
-								DetailsObject->NativeProperties.SetNum(DetailsObject->NativeProperties.Num()+1);
+
+								DetailsObject->NativeProperties.SetNum(DetailsObject->NativeProperties.Num() + 1);
 								FInstancedStruct& InstancedStruct = DetailsObject->NativeProperties.Last();
-								
+
 								FInstancedStruct::StaticStruct()->SerializeItem(ReaderProxy, &InstancedStruct, nullptr);
 							}
 
@@ -201,15 +201,15 @@ bool FAnimNextModuleTrack::UpdateInternal()
 						});
 					}
 				}
-		
-				ExistenceRange->Windows.SetNum(1,EAllowShrinking::No);
+
+				ExistenceRange->Windows.SetNum(1, EAllowShrinking::No);
 				ExistenceRange->Windows[0].TimeStart = Data->StartTime;
 				ExistenceRange->Windows[0].TimeEnd = Data->EndTime;
 			}
 		}
 	}
-	
- 	return bChanged;
+
+	return bChanged;
 }
 
 bool FAnimNextModuleTrack::HandleDoubleClickInternal()
@@ -222,7 +222,7 @@ bool FAnimNextModuleTrack::HandleDoubleClickInternal()
 	const IRewindDebugger* RewindDebugger = IRewindDebugger::Instance();
 	if (const TraceServices::IAnalysisSession* AnalysisSession = RewindDebugger->GetAnalysisSession())
 	{
-		const FObjectInfo* AssetInfo = nullptr;
+		FString AssetPath;
 		{
 			TraceServices::FAnalysisSessionReadScope SessionReadScope(*AnalysisSession);
 
@@ -231,28 +231,29 @@ bool FAnimNextModuleTrack::HandleDoubleClickInternal()
 
 			if (const FDataInterfaceData* ModuleData = AnimNextProvider->GetDataInterfaceData(InstanceId))
 			{
-				AssetInfo = &GameplayProvider->GetObjectInfo(ModuleData->AssetId);
+				const FObjectInfo& AssetInfo = GameplayProvider->GetObjectInfo(ModuleData->AssetId);
+				if (!EnumHasAnyFlags(AssetInfo.Flags, EObjectInfoFlags::TransientObject))
+				{
+					AssetPath = AssetInfo.PathName;
+				}
 			}
 		}
 
-		if (AssetInfo
-			&& !EnumHasAnyFlags(AssetInfo->Flags, EObjectInfoFlags::TransientObject))
+		if (!AssetPath.IsEmpty())
 		{
-			GEditor->GetEditorSubsystem<UAssetEditorSubsystem>()->OpenEditorForAsset(AssetInfo->PathName);
+			GEditor->GetEditorSubsystem<UAssetEditorSubsystem>()->OpenEditorForAsset(AssetPath);
 			return true;
 		}
 	}
 	return false;
 }
-	
-	
+
 bool FAnimNextModuleTrackCreator::HasDebugInfoInternal(const RewindDebugger::FObjectId& ObjectId) const
 {
 	return true;
 }
 
-
-}
+} // UE::UAF::Editor
 
 #undef LOCTEXT_NAMESPACE
 
