Fix UAF RigVM override variable asset persistency

diff --git a/Engine/Plugins/Experimental/UAF/UAF/Source/UAF/Private/AnimNextModuleImpl.cpp b/Engine/Plugins/Experimental/UAF/UAF/Source/UAF/Private/AnimNextModuleImpl.cpp
--- a/Engine/Plugins/Experimental/UAF/UAF/Source/UAF/Private/AnimNextModuleImpl.cpp
+++ b/Engine/Plugins/Experimental/UAF/UAF/Source/UAF/Private/AnimNextModuleImpl.cpp
@@ -2,6 +2,7 @@
 
 #include "AnimNextModuleImpl.h"
 #include "AnimNextConfig.h"
+#include "AnimNextRigVMAsset.h"
 #include "DataRegistry.h"
 #include "IUniversalObjectLocatorModule.h"
 #include "RigVMRuntimeDataRegistry.h"
@@ -11,6 +12,7 @@
 #include "Animation/BlendProfile.h"
 #include "Component/AnimNextComponent.h"
 #include "Curves/CurveFloat.h"
+#include "Graph/AnimNext_LODPose.h"
 #include "Module/AnimNextModuleInstance.h"
 #include "Module/ModuleHandle.h"
 #include "Modules/ModuleManager.h"
@@ -51,6 +53,7 @@ namespace UE::UAF
 			{ UCurveFloat::StaticClass(), FRigVMRegistry::ERegisterObjectOperation::Class },
 			{ UAnimNextComponent::StaticClass(), FRigVMRegistry::ERegisterObjectOperation::Class },
 			{ UAnimNextModule::StaticClass(), FRigVMRegistry::ERegisterObjectOperation::Class },
+			{ UAnimNextRigVMAsset::StaticClass(), FRigVMRegistry::ERegisterObjectOperation::Class },
 			{ UHierarchyTable::StaticClass(), FRigVMRegistry::ERegisterObjectOperation::Class },
 		};
 
@@ -67,6 +70,8 @@ namespace UE::UAF
 			TBaseStructure<FGuid>::Get(),
 			FRigVMVariant::StaticStruct(),
 			FAnimNextModuleHandle::StaticStruct(),
+			FAnimNextGraphLODPose::StaticStruct(),
+			FAnimNextGraphReferencePose::StaticStruct(),
 		};
 
 		RigVMRegistry.RegisterStructTypes(AllowedStructTypes);
diff --git a/Engine/Plugins/Experimental/UAF/UAFAnimGraph/Source/UAFAnimGraph/Private/AnimNextAnimGraphModule.cpp b/Engine/Plugins/Experimental/UAF/UAFAnimGraph/Source/UAFAnimGraph/Private/AnimNextAnimGraphModule.cpp
--- a/Engine/Plugins/Experimental/UAF/UAFAnimGraph/Source/UAFAnimGraph/Private/AnimNextAnimGraphModule.cpp
+++ b/Engine/Plugins/Experimental/UAF/UAFAnimGraph/Source/UAFAnimGraph/Private/AnimNextAnimGraphModule.cpp
@@ -13,7 +13,6 @@
 #include "Factory/AnimGraphFactory.h"
 #include "TraitInterfaces/IUpdate.h"
 #include "Graph/AnimNextAnimationGraph.h"
-#include "Graph/AnimNext_LODPose.h"
 #include "Graph/RigVMTrait_AnimNextPublicVariables.h"
 #include "Injection/InjectionRequest.h"
 #include "TraitCore/NodeTemplateRegistry.h"
@@ -65,10 +64,9 @@ void FAnimNextAnimGraphModule::StartupModule()
 	
 	static UScriptStruct* const AllowedStructTypes[] =
 	{
-		FAnimNextGraphReferencePose::StaticStruct(),
 		FRigVMTrait_AnimNextPublicVariables::StaticStruct(),
 		FAnimNextInjectionBlendSettings::StaticStruct(),
-		FModifyCurveParameters::StaticStruct()
+		FModifyCurveParameters::StaticStruct(),
 	};
 
 	RigVMRegistry.RegisterStructTypes(AllowedStructTypes);
