Avoid invalid memory access on editor exit

diff --git a/Engine/Plugins/Experimental/UAF/UAF/Source/UAF/Private/DataRegistry.cpp b/Engine/Plugins/Experimental/UAF/UAF/Source/UAF/Private/DataRegistry.cpp
--- a/Engine/Plugins/Experimental/UAF/UAF/Source/UAF/Private/DataRegistry.cpp
+++ b/Engine/Plugins/Experimental/UAF/UAF/Source/UAF/Private/DataRegistry.cpp
@@ -2,6 +2,7 @@
 
 #include "DataRegistry.h"
 
+#include "Misc/CoreDelegates.h"
 #include "Misc/ScopeRWLock.h"
 #include "DataRegistryTypes.h"
 #include "GenerationTools.h"
@@ -17,6 +18,7 @@ namespace // Private
 UE::UAF::FDataRegistry* GAnimationDataRegistry = nullptr;
 constexpr int32 BASIC_TYPE_ALLOC_BLOCK = 1000;
 FDelegateHandle PostGarbageCollectHandle;
+FDelegateHandle DataRegistryPreExitHandle;
 
 static TAutoConsoleVariable<float> CVar_GCCleanupTimeBudget(TEXT("a.AnimNext.GCCleanupTimeBudget"), 100.0f / (1000.0f * 1000.0f), TEXT("Amount of time in seconds that we can use when reclaiming memory post-GC"));
 
@@ -32,6 +34,7 @@ namespace UE::UAF
 		GAnimationDataRegistry = new FDataRegistry();
 
 		PostGarbageCollectHandle = FCoreUObjectDelegates::GetPostGarbageCollect().AddStatic(&FDataRegistry::HandlePostGarbageCollect);
+		DataRegistryPreExitHandle = FCoreDelegates::OnEnginePreExit.AddStatic(&FDataRegistry::OnPreExit);
 	}
 }
 
@@ -40,6 +43,10 @@ namespace UE::UAF
 	if (GAnimationDataRegistry != nullptr)
 	{
 		FCoreUObjectDelegates::GetPostGarbageCollect().Remove(PostGarbageCollectHandle);
+		PostGarbageCollectHandle.Reset();
+
+		FCoreDelegates::OnEnginePreExit.Remove(DataRegistryPreExitHandle);
+		DataRegistryPreExitHandle.Reset();
 
 		if (!GAnimationDataRegistry->GCKeysToCleanup.IsEmpty())
 		{
@@ -48,14 +55,6 @@ namespace UE::UAF
 			FTSTicker::RemoveTicker(GAnimationDataRegistry->GCCleanupTickerHandle);
 		}
 
-		// When the module unloads, we cannot use ScriptStructs and our ParamTypes to free the memory
-		// we allocated. To avoid crashing on an invalid access, we clear the type map which means that
-		// destructors won't run for allocated data. If they were to free memory, it will leak on app exit.
-		{
-			FRWScopeLock ScopedDataTypeDefsLock(GAnimationDataRegistry->DataTypeDefsLock, SLT_Write);
-			GAnimationDataRegistry->DataTypeDefs.Empty();
-		}
-
 		GAnimationDataRegistry->ReleaseReferencePoseData(); // release any registered poses
 
 		check(GAnimationDataRegistry->AllocatedBlocks.Num() == 0); // any other data should have been released at this point
@@ -66,6 +65,20 @@ namespace UE::UAF
 	}
 }
 
+void FDataRegistry::OnPreExit()
+{
+	checkf(GAnimationDataRegistry, TEXT("Animation Data Registry is not instanced. It is only valid to access this while the engine module is loaded."));
+
+	if (!GAnimationDataRegistry->GCKeysToCleanup.IsEmpty())
+	{
+		GAnimationDataRegistry->RunGCCleanup(INDEX_NONE);
+
+		FTSTicker::RemoveTicker(GAnimationDataRegistry->GCCleanupTickerHandle);
+	}
+
+	GAnimationDataRegistry->ReleaseReferencePoseData(); // release any registered poses
+}
+
 FDataRegistry* FDataRegistry::Get()
 {
 	checkf(GAnimationDataRegistry, TEXT("Animation Data Registry is not instanced. It is only valid to access this while the engine module is loaded."));
diff --git a/Engine/Plugins/Experimental/UAF/UAF/Source/UAF/Public/DataRegistry.h b/Engine/Plugins/Experimental/UAF/UAF/Source/UAF/Public/DataRegistry.h
--- a/Engine/Plugins/Experimental/UAF/UAF/Source/UAF/Public/DataRegistry.h
+++ b/Engine/Plugins/Experimental/UAF/UAF/Source/UAF/Public/DataRegistry.h
@@ -246,6 +246,9 @@ private:
 	// Shutdown the global registry
 	static UE_API void Destroy();
 
+	// Called before engine shutdown to clear internal state while UObjects are still valid
+	static void OnPreExit();
+
 	volatile int32 HandleCounter = 0;
 
 
diff --git a/Engine/Plugins/Experimental/UAF/UAFAnimGraph/Source/UAFAnimGraph/Private/AnimNextAnimGraphModule.cpp b/Engine/Plugins/Experimental/UAF/UAFAnimGraph/Source/UAFAnimGraph/Private/AnimNextAnimGraphModule.cpp
--- a/Engine/Plugins/Experimental/UAF/UAFAnimGraph/Source/UAFAnimGraph/Private/AnimNextAnimGraphModule.cpp
+++ b/Engine/Plugins/Experimental/UAF/UAFAnimGraph/Source/UAFAnimGraph/Private/AnimNextAnimGraphModule.cpp
@@ -74,6 +74,7 @@ void FAnimNextAnimGraphModule::StartupModule()
 	FTraitRegistry::Init();
 	FTraitInterfaceRegistry::Init();
 	FNodeTemplateRegistry::Init();
+	FAnimGraphFactory::Init();
 
 	// Register asset class -> graph factory mappings
 	{
@@ -167,6 +168,7 @@ void FAnimNextAnimGraphModule::ShutdownModule()
 	ConsoleCommands.Empty();
 #endif
 
+	FAnimGraphFactory::Destroy();
 	FNodeTemplateRegistry::Destroy();
 	FTraitInterfaceRegistry::Destroy();
 	FTraitRegistry::Destroy();
diff --git a/Engine/Plugins/Experimental/UAF/UAFAnimGraph/Source/UAFAnimGraph/Private/Factory/AnimGraphFactory.cpp b/Engine/Plugins/Experimental/UAF/UAFAnimGraph/Source/UAFAnimGraph/Private/Factory/AnimGraphFactory.cpp
--- a/Engine/Plugins/Experimental/UAF/UAFAnimGraph/Source/UAFAnimGraph/Private/Factory/AnimGraphFactory.cpp
+++ b/Engine/Plugins/Experimental/UAF/UAFAnimGraph/Source/UAFAnimGraph/Private/Factory/AnimGraphFactory.cpp
@@ -1,4 +1,4 @@
-﻿// Copyright Epic Games, Inc. All Rights Reserved.
+// Copyright Epic Games, Inc. All Rights Reserved.
 
 #include "Factory/AnimGraphFactory.h"
 
@@ -6,11 +6,14 @@
 #include "Containers/StripedMap.h"
 #include "Factory/AnimNextAnimGraphBuilder.h"
 #include "Graph/AnimNextAnimationGraph.h"
+#include "Misc/CoreDelegates.h"
 #include "TraitCore/TraitWriter.h"
 
 namespace UE::UAF::Private
 {
 
+static FDelegateHandle GAnimGraphFactoryPreExitHandle;
+
 // References to all generated factory graphs, keyed by method
 static TStripedMap<32, uint64, TWeakObjectPtr<const UAnimNextAnimationGraph>> GAnimationGraphMap;
 
@@ -141,5 +144,22 @@ TConstArrayView<UClass*> FAnimGraphFactory::GetRegisteredClasses()
 	return Private::GAllFactoryClasses;
 }
 
+void FAnimGraphFactory::Init()
+{
+	Private::GAnimGraphFactoryPreExitHandle = FCoreDelegates::OnEnginePreExit.AddStatic(&FAnimGraphFactory::OnPreExit);
+}
+
+void FAnimGraphFactory::Destroy()
+{
+	FCoreDelegates::OnEnginePreExit.Remove(Private::GAnimGraphFactoryPreExitHandle);
+	Private::GAnimGraphFactoryPreExitHandle.Reset();
+}
+
+void FAnimGraphFactory::OnPreExit()
+{
+	Private::GAnimationGraphMap.Empty();
+	Private::GDefaultParamsForObject.Empty();
+	Private::GAllFactoryClasses.Empty();
+}
 
 }
diff --git a/Engine/Plugins/Experimental/UAF/UAFAnimGraph/Source/UAFAnimGraph/Public/Factory/AnimGraphFactory.h b/Engine/Plugins/Experimental/UAF/UAFAnimGraph/Source/UAFAnimGraph/Public/Factory/AnimGraphFactory.h
--- a/Engine/Plugins/Experimental/UAF/UAFAnimGraph/Source/UAFAnimGraph/Public/Factory/AnimGraphFactory.h
+++ b/Engine/Plugins/Experimental/UAF/UAFAnimGraph/Source/UAFAnimGraph/Public/Factory/AnimGraphFactory.h
@@ -1,4 +1,4 @@
-﻿// Copyright Epic Games, Inc. All Rights Reserved.
+// Copyright Epic Games, Inc. All Rights Reserved.
 
 #pragma once
 
@@ -13,6 +13,11 @@ namespace UE::UAF
 	class FTraitWriter;
 	struct IAnimGraphBuilder;
 	struct FAnimGraphFactory;
+
+	namespace AnimGraph
+	{
+		class FAnimNextAnimGraphModule;
+	}
 }
 
 namespace UE::UAF
@@ -81,6 +86,18 @@ struct FAnimGraphFactory
 
 	// Get all registered classes used to generate factory graphs
 	UAFANIMGRAPH_API static TConstArrayView<UClass*> GetRegisteredClasses();
+
+private:
+	// Called on module init
+	UAFANIMGRAPH_API static void Init();
+
+	// Called on module shutdown to unload everything
+	UAFANIMGRAPH_API static void Destroy();
+
+	// Called before engine shutdown to clear internal state while UObjects are still valid
+	static void OnPreExit();
+
+	friend AnimGraph::FAnimNextAnimGraphModule;
 };
 
 }
\ No newline at end of file
diff --git a/Engine/Plugins/VirtualProduction/Takes/Source/TakeRecorder/Private/TakeRecorderModule.cpp b/Engine/Plugins/VirtualProduction/Takes/Source/TakeRecorder/Private/TakeRecorderModule.cpp
--- a/Engine/Plugins/VirtualProduction/Takes/Source/TakeRecorder/Private/TakeRecorderModule.cpp
+++ b/Engine/Plugins/VirtualProduction/Takes/Source/TakeRecorder/Private/TakeRecorderModule.cpp
@@ -452,9 +452,12 @@ void FTakeRecorderModule::RegisterDetailCustomizations()
 
 void FTakeRecorderModule::UnregisterDetailCustomizations()
 {
-	if (FPropertyEditorModule* PropertyEditorModule = FModuleManager::Get().GetModulePtr<FPropertyEditorModule>("PropertyEditor"))
+	if (UObjectInitialized() && !IsEngineExitRequested())
 	{
-		PropertyEditorModule->UnregisterCustomClassLayout(UTakeRecorderProjectSettings::StaticClass()->GetFName());
+		if (FPropertyEditorModule* PropertyEditorModule = FModuleManager::Get().GetModulePtr<FPropertyEditorModule>("PropertyEditor"))
+		{
+			PropertyEditorModule->UnregisterCustomClassLayout(UTakeRecorderProjectSettings::StaticClass()->GetFName());
+		}
 	}
 }
 
