Add functionality to Open in New Tab via the context menu for collapsed nodes and function nodes in UAF
- Fixed Collapsed Nodes and Functions being opened without context of their outer, leading to loss of access to variables
- Fixed debug only assert when opening items via double click the workspace outliner

diff --git a/Engine/Plugins/Experimental/UAF/UAF/Source/UAFEditor/Private/AnimNextEditorModule.cpp b/Engine/Plugins/Experimental/UAF/UAF/Source/UAFEditor/Private/AnimNextEditorModule.cpp
--- a/Engine/Plugins/Experimental/UAF/UAF/Source/UAFEditor/Private/AnimNextEditorModule.cpp
+++ b/Engine/Plugins/Experimental/UAF/UAF/Source/UAFEditor/Private/AnimNextEditorModule.cpp
@@ -1208,9 +1208,15 @@ void FAnimNextEditorModule::RegisterWorkspaceDocumentTypes(Workspace::IWorkspace
 					{
 						if (IRigVMClientHost* RigVMClientHost = ContainedGraph->GetImplementingOuter<IRigVMClientHost>())
 						{
-							if (UObject* EditorObject = RigVMClientHost->GetEditorObjectForRigVMGraph(ContainedGraph))
+							if (const UObject* EditorObject = RigVMClientHost->GetEditorObjectForRigVMGraph(ContainedGraph))
 							{
-								WorkspaceEditor->OpenObjects({ EditorObject });
+								const FSoftObjectPath GraphSoftPath = EditorObject;
+								if (GraphSoftPath.IsValid() && GraphSoftPath.IsSubobject())
+								{
+									const FWorkspaceOutlinerItemExport ParentExport = InContext.Document.Export;
+									FWorkspaceOutlinerItemExport Export = FWorkspaceOutlinerItemExport(*GraphSoftPath.GetSubPathUtf8String(), ParentExport);
+									WorkspaceEditor->OpenExports({ Export });
+								}
 							}
 						}
 					}
@@ -1259,6 +1265,86 @@ void FAnimNextEditorModule::RegisterWorkspaceDocumentTypes(Workspace::IWorkspace
 			}
 		}
 	});
+	GraphArgs.OnCanOpenInNewTab = Workspace::FOnCanPerformActionOnSelectedNodes::CreateLambda([this](const Workspace::FWorkspaceEditorContext& InContext, const FGraphPanelSelectionSet& InSelectedNodes)
+	{
+		for (FGraphPanelSelectionSet::TConstIterator NodeIt(InSelectedNodes); NodeIt; ++NodeIt)
+		{
+			URigVMGraph* ContainedGraph = nullptr;
+			if (const UAnimNextEdGraphNode* RigVMEdGraphNode = Cast<UAnimNextEdGraphNode>(*NodeIt))
+			{
+				const URigVMNode* ModelNode = RigVMEdGraphNode->GetModelNode();
+				if (const URigVMLibraryNode* LibraryNode = Cast<URigVMLibraryNode>(ModelNode))
+				{
+					ContainedGraph = LibraryNode->GetContainedGraph();
+					if (const URigVMFunctionReferenceNode* FunctionReferenceNode = Cast<URigVMFunctionReferenceNode>(LibraryNode))
+					{
+						if (const URigVMLibraryNode* ReferencedNode = FunctionReferenceNode->LoadReferencedNode())
+						{
+							ContainedGraph = ReferencedNode->GetContainedGraph();
+						}
+					}
+				}
+
+				if (ContainedGraph)
+				{
+					return true;
+				}
+			}
+		}
+
+		return false;
+	});
+	GraphArgs.OnOpenInNewTab = Workspace::FOnPerformActionOnSelectedNodes::CreateLambda([this](const Workspace::FWorkspaceEditorContext& InContext, const FGraphPanelSelectionSet& InSelectedNodes)
+	{		
+		TArray<FWorkspaceOutlinerItemExport> ExportsToOpen;
+		for (FGraphPanelSelectionSet::TConstIterator NodeIt(InSelectedNodes); NodeIt; ++NodeIt)
+		{
+			if (const UAnimNextEdGraphNode* RigVMEdGraphNode = Cast<UAnimNextEdGraphNode>(*NodeIt))
+			{
+				const URigVMNode* ModelNode = RigVMEdGraphNode->GetModelNode();
+				if (const URigVMLibraryNode* LibraryNode = Cast<URigVMLibraryNode>(ModelNode))
+				{
+					const URigVMGraph* ContainedGraph = nullptr;
+					if (const URigVMCollapseNode* CollapsedNode = Cast<URigVMCollapseNode>(ModelNode))
+					{
+						ContainedGraph = CollapsedNode->GetContainedGraph();
+					}
+					else if(const URigVMFunctionReferenceNode* FunctionReferenceNode = Cast<URigVMFunctionReferenceNode>(ModelNode))
+					{
+						if (const URigVMLibraryNode* ReferencedNode = FunctionReferenceNode->LoadReferencedNode())
+						{
+							ContainedGraph = ReferencedNode->GetContainedGraph();
+						}
+					}
+
+					if (ContainedGraph)
+					{
+						if (const IRigVMClientHost* RigVMClientHost = ContainedGraph->GetImplementingOuter<IRigVMClientHost>())
+						{
+							if (const UObject* EditorObject = RigVMClientHost->GetEditorObjectForRigVMGraph(ContainedGraph))
+							{
+								const FSoftObjectPath GraphSoftPath = EditorObject;
+								if (GraphSoftPath.IsValid() && GraphSoftPath.IsSubobject())
+								{
+									const FWorkspaceOutlinerItemExport ParentExport = InContext.Document.Export;
+									FWorkspaceOutlinerItemExport& Export = ExportsToOpen.Add_GetRef(FWorkspaceOutlinerItemExport(*GraphSoftPath.GetSubPathUtf8String(), ParentExport));
+								}
+							}
+						}
+					}
+				}
+			}
+		}
+
+		if (ExportsToOpen.Num() > 0)
+		{
+			if (TSharedPtr<Workspace::IWorkspaceEditor> WorkspaceEditor = InContext.WorkspaceEditor)
+			{
+				WorkspaceEditor->OpenExports(ExportsToOpen, FDocumentTracker::EOpenDocumentCause::ForceOpenNewDocument);
+			}
+		}
+	});
+
 
 	Workspace::FObjectDocumentArgs GraphDocumentArgs = WorkspaceEditorModule.CreateGraphDocumentArgs(GraphArgs);
 	Workspace::FOnMakeDocumentWidget WorkspaceMakeDocumentWidgetDelegate = GraphDocumentArgs.OnMakeDocumentWidget;
diff --git a/Engine/Plugins/Experimental/UAF/UAF/Source/UAFEditor/Private/Common/AnimNextCollapseNodeItemDetails.cpp b/Engine/Plugins/Experimental/UAF/UAF/Source/UAFEditor/Private/Common/AnimNextCollapseNodeItemDetails.cpp
--- a/Engine/Plugins/Experimental/UAF/UAF/Source/UAFEditor/Private/Common/AnimNextCollapseNodeItemDetails.cpp
+++ b/Engine/Plugins/Experimental/UAF/UAF/Source/UAFEditor/Private/Common/AnimNextCollapseNodeItemDetails.cpp
@@ -28,16 +28,10 @@ bool FAnimNextCollapseNodeItemDetails::HandleDoubleClick(const FToolMenuContext&
 	{
 		if (const TSharedPtr<UE::Workspace::IWorkspaceEditor> WorkspaceEditor = StaticCastSharedPtr<UE::Workspace::IWorkspaceEditor>(AssetEditorContext->Toolkit.Pin()))
 		{
-			const TInstancedStruct<FWorkspaceOutlinerItemData>& Data = WorkspaceItemContext->SelectedExports[0].GetResolvedExport().GetData();
-			if (Data.IsValid() && Data.GetScriptStruct() == FAnimNextCollapseGraphOutlinerData::StaticStruct())
+			if (WorkspaceItemContext->SelectedExports.Num() > 0)
 			{
-				const FAnimNextCollapseGraphOutlinerData& CollapseGraphData = Data.Get<FAnimNextCollapseGraphOutlinerData>();
-				if (CollapseGraphData.SoftEditorObject.IsValid())
-				{
-					UObject* EditorObject = CollapseGraphData.SoftEditorObject.Get();
-					WorkspaceEditor->OpenObjects({ EditorObject });
-					return true;
-				}
+				WorkspaceEditor->OpenExports({ WorkspaceItemContext->SelectedExports[0].GetResolvedExport() });
+				return true;
 			}
 		}
 	}
@@ -189,6 +183,26 @@ void FAnimNextCollapseNodeItemDetails::UnregisterToolMenuExtensions()
 {
 }
 
+FString FAnimNextCollapseNodeItemDetails::GetDisplayString(const FWorkspaceOutlinerItemExport& Export) const
+{
+	const TInstancedStruct<FWorkspaceOutlinerItemData>& Data = Export.GetData();
+	if (Data.IsValid() && Data.GetScriptStruct() == FAnimNextCollapseGraphOutlinerData::StaticStruct())
+	{
+		const FAnimNextCollapseGraphOutlinerData& CollapseGraphData = Data.Get<FAnimNextCollapseGraphOutlinerData>();
+		if (CollapseGraphData.SoftEditorObject.IsValid())
+		{
+			const URigVMEdGraph* EdGraph = CollapseGraphData.SoftEditorObject.Get();
+			if (const UEdGraphSchema* GraphSchema = EdGraph->GetSchema())
+			{
+				FGraphDisplayInfo DisplayInfo;
+				GraphSchema->GetGraphDisplayInformation(*EdGraph, DisplayInfo);
+
+				return DisplayInfo.PlainName.ToString();
+			}
+		}
+	}
+	return Export.GetIdentifier().ToString();
+}
 
 } // UE::UAF::Editor
 
diff --git a/Engine/Plugins/Experimental/UAF/UAF/Source/UAFEditor/Private/Common/AnimNextCollapseNodeItemDetails.h b/Engine/Plugins/Experimental/UAF/UAF/Source/UAFEditor/Private/Common/AnimNextCollapseNodeItemDetails.h
--- a/Engine/Plugins/Experimental/UAF/UAF/Source/UAFEditor/Private/Common/AnimNextCollapseNodeItemDetails.h
+++ b/Engine/Plugins/Experimental/UAF/UAF/Source/UAFEditor/Private/Common/AnimNextCollapseNodeItemDetails.h
@@ -21,6 +21,7 @@ public:
 	virtual bool ValidateName(const FWorkspaceOutlinerItemExport& Export, const FText& InName, FText& OutErrorMessage) const override;
 	virtual UPackage* GetPackage(const FWorkspaceOutlinerItemExport& Export) const override;
 	virtual const FSlateBrush* GetItemIcon(const FWorkspaceOutlinerItemExport& Export) const override;
+	virtual FString GetDisplayString(const FWorkspaceOutlinerItemExport& Export) const override;
 
 	static void RegisterToolMenuExtensions();
 	static void UnregisterToolMenuExtensions();
diff --git a/Engine/Plugins/Experimental/UAF/UAF/Source/UAFEditor/Private/Common/AnimNextFunctionItemDetails.cpp b/Engine/Plugins/Experimental/UAF/UAF/Source/UAFEditor/Private/Common/AnimNextFunctionItemDetails.cpp
--- a/Engine/Plugins/Experimental/UAF/UAF/Source/UAFEditor/Private/Common/AnimNextFunctionItemDetails.cpp
+++ b/Engine/Plugins/Experimental/UAF/UAF/Source/UAFEditor/Private/Common/AnimNextFunctionItemDetails.cpp
@@ -28,15 +28,10 @@ bool FAnimNextFunctionItemDetails::HandleDoubleClick(const FToolMenuContext& Too
 	{
 		if (const TSharedPtr<UE::Workspace::IWorkspaceEditor> WorkspaceEditor = StaticCastSharedPtr<UE::Workspace::IWorkspaceEditor>(AssetEditorContext->Toolkit.Pin()))
 		{
-			const TInstancedStruct<FWorkspaceOutlinerItemData>& Data = WorkspaceItemContext->SelectedExports[0].GetResolvedExport().GetData();
-			if (Data.IsValid() && Data.GetScriptStruct() == FAnimNextGraphFunctionOutlinerData::StaticStruct())
+			if (WorkspaceItemContext->SelectedExports.Num() > 0)
 			{
-				const FAnimNextGraphFunctionOutlinerData& FunctionData = Data.Get<FAnimNextGraphFunctionOutlinerData>();
-				if (FunctionData.SoftEditorObject.IsValid())
-				{
-					WorkspaceEditor->OpenObjects({ FunctionData.SoftEditorObject.Get() });
-					return true;
-				}
+				WorkspaceEditor->OpenExports({ WorkspaceItemContext->SelectedExports[0].GetResolvedExport() });
+				return true;
 			}
 		}
 	}
@@ -224,6 +219,27 @@ void FAnimNextFunctionItemDetails::UnregisterToolMenuExtensions()
 {
 }
 
+FString FAnimNextFunctionItemDetails::GetDisplayString(const FWorkspaceOutlinerItemExport& Export) const
+{
+	const TInstancedStruct<FWorkspaceOutlinerItemData>& Data = Export.GetData();
+	if (Data.IsValid() && Data.GetScriptStruct() == FAnimNextGraphFunctionOutlinerData::StaticStruct())
+	{
+		const FAnimNextGraphFunctionOutlinerData& FunctionData = Data.Get<FAnimNextGraphFunctionOutlinerData>();
+		if (FunctionData.SoftEditorObject.IsValid())
+		{
+			const URigVMEdGraph* EdGraph = FunctionData.SoftEditorObject.Get();
+			if (const UEdGraphSchema* GraphSchema = EdGraph->GetSchema())
+			{
+				FGraphDisplayInfo DisplayInfo;
+				GraphSchema->GetGraphDisplayInformation(*EdGraph, DisplayInfo);
+
+				return DisplayInfo.PlainName.ToString();
+			}
+		}
+	}
+	return Export.GetIdentifier().ToString();
+}
+
 } // UE::UAF::Editor
 
 #undef LOCTEXT_NAMESPACE // "FAnimNextFunctionItemDetails"
diff --git a/Engine/Plugins/Experimental/UAF/UAF/Source/UAFEditor/Private/Common/AnimNextFunctionItemDetails.h b/Engine/Plugins/Experimental/UAF/UAF/Source/UAFEditor/Private/Common/AnimNextFunctionItemDetails.h
--- a/Engine/Plugins/Experimental/UAF/UAF/Source/UAFEditor/Private/Common/AnimNextFunctionItemDetails.h
+++ b/Engine/Plugins/Experimental/UAF/UAF/Source/UAFEditor/Private/Common/AnimNextFunctionItemDetails.h
@@ -22,6 +22,7 @@ public:
 	virtual UPackage* GetPackage(const FWorkspaceOutlinerItemExport& Export) const override;
 	virtual const FSlateBrush* GetItemIcon(const FWorkspaceOutlinerItemExport& Export) const override;
 	virtual bool HandleSelected(const FToolMenuContext& ToolMenuContext) const override;
+	virtual FString GetDisplayString(const FWorkspaceOutlinerItemExport& Export) const override;
 
 	static void RegisterToolMenuExtensions();
 	static void UnregisterToolMenuExtensions();
diff --git a/Engine/Plugins/Experimental/UAF/UAF/Source/UAFEditor/Private/Variables/SVariablesView.cpp b/Engine/Plugins/Experimental/UAF/UAF/Source/UAFEditor/Private/Variables/SVariablesView.cpp
--- a/Engine/Plugins/Experimental/UAF/UAF/Source/UAFEditor/Private/Variables/SVariablesView.cpp
+++ b/Engine/Plugins/Experimental/UAF/UAF/Source/UAFEditor/Private/Variables/SVariablesView.cpp
@@ -1,4 +1,4 @@
-﻿// Copyright Epic Games, Inc. All Rights Reserved.
+// Copyright Epic Games, Inc. All Rights Reserved.
 
 #include "SVariablesView.h"
 
@@ -107,7 +107,9 @@ void SVariablesOutliner::UpdateAssets()
 		}
 	};
 	
-	if (FWorkspaceOutlinerAssetReferenceItemData::IsAssetReference(Export))
+	if (FWorkspaceOutlinerAssetReferenceItemData::IsAssetReference(Export)
+		|| FWorkspaceOutlinerGroupItemData::IsGroupItem(Export)
+		|| FAnimNextCollapseGraphsOutlinerDataBase::IsCollapsedGraphBase(Export))
 	{
 		// Reference item
 		TArray<FSoftObjectPath> AssetPaths;
@@ -136,20 +138,6 @@ void SVariablesOutliner::UpdateAssets()
 				}
 			}
 		}
-		else if(Export.HasData() && Export.GetData().GetScriptStruct()->IsChildOf(FAnimNextGraphFunctionOutlinerData::StaticStruct()))
-		{
-			const FAnimNextGraphFunctionOutlinerData& EntryData = Export.GetData().Get<FAnimNextGraphFunctionOutlinerData>();
-			if(EntryData.SoftEditorObject.IsValid())
-			{
-				if(URigVMEdGraph* EdGraph = EntryData.SoftEditorObject.Get())
-				{
-					if(UAnimNextRigVMAsset* Asset = EdGraph->GetTypedOuter<UAnimNextRigVMAsset>())
-					{
-						HandleAssetPath(Asset);
-					}
-				}
-			}
-		}
 	}
 
 	SetAssets(ExportAssets);
diff --git a/Engine/Plugins/Experimental/UAF/UAF/Source/UAFUncookedOnly/Internal/AnimNextAssetWorkspaceAssetUserData.h b/Engine/Plugins/Experimental/UAF/UAF/Source/UAFUncookedOnly/Internal/AnimNextAssetWorkspaceAssetUserData.h
--- a/Engine/Plugins/Experimental/UAF/UAF/Source/UAFUncookedOnly/Internal/AnimNextAssetWorkspaceAssetUserData.h
+++ b/Engine/Plugins/Experimental/UAF/UAF/Source/UAFUncookedOnly/Internal/AnimNextAssetWorkspaceAssetUserData.h
@@ -95,6 +95,11 @@ struct FAnimNextCollapseGraphsOutlinerDataBase : public FWorkspaceOutlinerItemDa
 
 	UPROPERTY(VisibleAnywhere, Category=AnimNext)
 	TSoftObjectPtr<URigVMEdGraph> SoftEditorObject;
+
+	static bool IsCollapsedGraphBase(const FWorkspaceOutlinerItemExport& InExport)
+	{
+		return InExport.HasData() && InExport.GetData().GetScriptStruct()->IsChildOf(FAnimNextCollapseGraphsOutlinerDataBase::StaticStruct());
+	}
 };
 
 USTRUCT()
diff --git a/Engine/Plugins/Experimental/UAF/UAF/Source/UAFUncookedOnly/Private/UncookedOnlyUtils.cpp b/Engine/Plugins/Experimental/UAF/UAF/Source/UAFUncookedOnly/Private/UncookedOnlyUtils.cpp
--- a/Engine/Plugins/Experimental/UAF/UAF/Source/UAFUncookedOnly/Private/UncookedOnlyUtils.cpp
+++ b/Engine/Plugins/Experimental/UAF/UAF/Source/UAFUncookedOnly/Private/UncookedOnlyUtils.cpp
@@ -1267,10 +1267,13 @@ void FUtils::GetSubGraphWorkspaceExportsRecursive(const UAnimNextRigVMAssetEdito
 		{
 			if(ensure(EditorObject->GetModel()))
 			{
-				URigVMCollapseNode* CollapseNode = CastChecked<URigVMCollapseNode>(EditorObject->GetModel()->GetOuter());
+				const URigVMCollapseNode* CollapseNode = CastChecked<URigVMCollapseNode>(EditorObject->GetModel()->GetOuter());
+				const FSoftObjectPath EditorObjectSoftObjPath = EditorObject;
+				ensureMsgf(EditorObjectSoftObjPath.IsSubobject(), TEXT("EditorObject for RigVMCollapseNode Graph was not a subobject as expected."));
 
 				const FWorkspaceOutlinerItemExport& ParentExport = ParentExportIndex == INDEX_NONE ? RootExport : OutExports.Exports[ParentExportIndex];
-				FWorkspaceOutlinerItemExport& Export = OutExports.Exports.Add_GetRef(FWorkspaceOutlinerItemExport(CollapseNode->GetFName(), ParentExport));
+				const FName Identifier = EditorObjectSoftObjPath.IsSubobject() ? *EditorObjectSoftObjPath.GetSubPathUtf8String() : CollapseNode->GetFName();
+				FWorkspaceOutlinerItemExport& Export = OutExports.Exports.Add_GetRef(FWorkspaceOutlinerItemExport(Identifier, ParentExport));
 				Export.GetData().InitializeAsScriptStruct(FAnimNextCollapseGraphOutlinerData::StaticStruct());
 			
 				FAnimNextCollapseGraphOutlinerData& FnGraphData = Export.GetData().GetMutable<FAnimNextCollapseGraphOutlinerData>();
@@ -1295,18 +1298,23 @@ void FUtils::GetSubGraphWorkspaceExportsRecursive(const UAnimNextRigVMAssetEdito
 			if (FunctionIdentifier.HostObject == FSoftObjectPath(EditorData))
 			{
 				const URigVMLibraryNode* FunctionNode = EditorData->RigVMClient.GetFunctionLibrary()->FindFunction(FunctionIdentifier.GetFunctionFName());
-				if (URigVMGraph* ContainedGraph = FunctionNode->GetContainedGraph())
+				if (const URigVMGraph* ContainedGraph = FunctionNode->GetContainedGraph())
 				{
+					URigVMEdGraph* ContainedEdGraph = Cast<URigVMEdGraph>(EditorData->GetEditorObjectForRigVMGraph(ContainedGraph));
+					const FSoftObjectPath EditorObjectSoftObjPath = ContainedEdGraph;
+					ensureMsgf(EditorObjectSoftObjPath.IsSubobject(), TEXT("EditorObject for RigVMFunctionReferenceNode Graph was not a subobject as expected."));
+
 					const FWorkspaceOutlinerItemExport& ParentExport = ParentExportIndex == INDEX_NONE ? RootExport : OutExports.Exports[ParentExportIndex];
-					FWorkspaceOutlinerItemExport& Export = OutExports.Exports.Add_GetRef(FWorkspaceOutlinerItemExport(FunctionNode->GetFName(), ParentExport));
+					const FName Identifier = EditorObjectSoftObjPath.IsSubobject() ? *EditorObjectSoftObjPath.GetSubPathUtf8String() : FunctionNode->GetFName();
+					FWorkspaceOutlinerItemExport& Export = OutExports.Exports.Add_GetRef(FWorkspaceOutlinerItemExport(Identifier, ParentExport));
 					Export.GetData().InitializeAsScriptStruct(FAnimNextGraphFunctionOutlinerData::StaticStruct());
 					FAnimNextGraphFunctionOutlinerData& FnGraphData = Export.GetData().GetMutable<FAnimNextGraphFunctionOutlinerData>();
 
-					if (URigVMEdGraph* ContainedEdGraph = Cast<URigVMEdGraph>(EditorData->GetEditorObjectForRigVMGraph(ContainedGraph)))
+					if (ContainedEdGraph)
 					{
 						FnGraphData.SoftEditorObject = ContainedEdGraph;
 						FnGraphData.SoftEdGraphNode = EdNode;
-						
+
 						int32 ExportIndex = OutExports.Exports.Num() - 1;
 						GetSubGraphWorkspaceExportsRecursive(EditorData, OutExports, RootExport, ExportIndex, ContainedEdGraph, Context);
 					}
@@ -1338,12 +1346,16 @@ void FUtils::GetFunctionsWorkspaceExportsRecursive(const UAnimNextRigVMAssetEdit
 	{
 		if (const URigVMLibraryNode* FunctionNode = EditorData->RigVMClient.GetFunctionLibrary()->FindFunction(FunctionData.Header.LibraryPointer.GetFunctionFName()))
 		{
-			if (URigVMGraph* ContainedModelGraph = FunctionNode->GetContainedGraph())
+			if (const URigVMGraph* ContainedModelGraph = FunctionNode->GetContainedGraph())
 			{
 				if (URigVMEdGraph* EditorObject = Cast<URigVMEdGraph>(EditorData->GetEditorObjectForRigVMGraph(ContainedModelGraph)))
 				{
+					const FSoftObjectPath EditorObjectSoftObjPath = EditorObject;
+					ensureMsgf(EditorObjectSoftObjPath.IsSubobject(), TEXT("EditorObject for RigVMFunctionReferenceNode Graph was not a subobject as expected."));
+
 					const FWorkspaceOutlinerItemExport& ParentExport = ParentExportIndex == INDEX_NONE ? RootExport : OutExports.Exports[ParentExportIndex];
-					FWorkspaceOutlinerItemExport& Export = OutExports.Exports.Add_GetRef(FWorkspaceOutlinerItemExport(FunctionData.Header.Name, ParentExport));
+					const FName Identifier = EditorObjectSoftObjPath.IsSubobject() ? *EditorObjectSoftObjPath.GetSubPathUtf8String() : FunctionNode->GetFName();
+					FWorkspaceOutlinerItemExport& Export = OutExports.Exports.Add_GetRef(FWorkspaceOutlinerItemExport(Identifier, ParentExport));
 
 					Export.GetData().InitializeAsScriptStruct(FAnimNextGraphFunctionOutlinerData::StaticStruct());
 					FAnimNextGraphFunctionOutlinerData& FnGraphData = Export.GetData().GetMutable<FAnimNextGraphFunctionOutlinerData>();
diff --git a/Engine/Plugins/Experimental/Workspace/Source/WorkspaceEditor/Private/SGraphDocument.cpp b/Engine/Plugins/Experimental/Workspace/Source/WorkspaceEditor/Private/SGraphDocument.cpp
--- a/Engine/Plugins/Experimental/Workspace/Source/WorkspaceEditor/Private/SGraphDocument.cpp
+++ b/Engine/Plugins/Experimental/Workspace/Source/WorkspaceEditor/Private/SGraphDocument.cpp
@@ -6,6 +6,7 @@
 #include "WorkspaceEditor.h"
 #include "EdGraph/EdGraph.h"
 #include "HAL/PlatformApplicationMisc.h"
+#include "GraphEditorActions.h"
 
 namespace UE::Workspace
 {
@@ -28,6 +29,8 @@ void SGraphDocument::Construct(const FArguments& InArgs, TSharedRef<FWorkspaceEd
 	OnPasteNodes = InArgs._OnPasteNodes;
 	OnCanDuplicateSelectedNodes = InArgs._OnCanDuplicateSelectedNodes;
 	OnDuplicateSelectedNodes = InArgs._OnDuplicateSelectedNodes;
+	OnCanOpenInNewTab = InArgs._OnCanOpenInNewTab;
+	OnOpenInNewTab = InArgs._OnOpenInNewTab;
 
 	SGraphEditor::FGraphEditorEvents Events;
 	Events.OnCreateActionMenuAtLocation = SGraphEditor::FOnCreateActionMenuAtLocation::CreateLambda([this, OnCreateActionMenu = InArgs._OnCreateActionMenu](UEdGraph* InGraph, const FVector2f& InNodePosition, const TArray<UEdGraphPin*>& InDraggedPins, bool bInAutoExpand, SGraphEditor::FActionMenuClosed InOnMenuClosed)
@@ -96,6 +99,11 @@ void SGraphDocument::BindCommands()
 	CommandList->MapAction(FGenericCommands::Get().SelectAll,
 		FExecuteAction::CreateSP(this, &SGraphDocument::SelectAllNodes),
 		FCanExecuteAction::CreateSP(this, &SGraphDocument::CanSelectAllNodes));
+
+	CommandList->MapAction(FGraphEditorCommands::Get().OpenInNewTab,
+		FExecuteAction::CreateSP(this, &SGraphDocument::OpenInNewTab),
+		FCanExecuteAction::CreateSP(this, &SGraphDocument::CanOpenInNewTab));
+
 }
 
 bool SGraphDocument::CanDeleteSelectedNodes() const
@@ -184,6 +192,20 @@ void SGraphDocument::SelectAllNodes()
 	GraphEditor->SelectAllNodes();
 }
 
+bool SGraphDocument::CanOpenInNewTab() const
+{
+	const bool bCan = OnCanOpenInNewTab.IsBound()
+		? OnCanOpenInNewTab.Execute(FWorkspaceEditorContext(HostingAppPtr.Pin().ToSharedRef(), Document), GraphEditor->GetSelectedNodes())
+		: false;
+	return bCan;
+}
+
+void SGraphDocument::OpenInNewTab()
+{
+	const FGraphPanelSelectionSet SelectedNodes = GraphEditor->GetSelectedNodes();
+	OnOpenInNewTab.ExecuteIfBound(FWorkspaceEditorContext(HostingAppPtr.Pin().ToSharedRef(), Document), SelectedNodes);
+}
+
 bool SGraphDocument::IsEditable(UEdGraph* InGraph) const
 {
 	return InGraph && HostingAppPtr.Pin()->InEditingMode() && InGraph->bEditable;
diff --git a/Engine/Plugins/Experimental/Workspace/Source/WorkspaceEditor/Private/SGraphDocument.h b/Engine/Plugins/Experimental/Workspace/Source/WorkspaceEditor/Private/SGraphDocument.h
--- a/Engine/Plugins/Experimental/Workspace/Source/WorkspaceEditor/Private/SGraphDocument.h
+++ b/Engine/Plugins/Experimental/Workspace/Source/WorkspaceEditor/Private/SGraphDocument.h
@@ -46,6 +46,9 @@ class SGraphDocument : public SCompoundWidget
 
 	SLATE_EVENT(FOnNodeDoubleClicked, OnNodeDoubleClicked)
 
+	SLATE_EVENT(FOnCanPerformActionOnSelectedNodes, OnCanOpenInNewTab)
+	SLATE_EVENT(FOnPerformActionOnSelectedNodes, OnOpenInNewTab)
+
 	SLATE_END_ARGS()
 
 	void Construct(const FArguments& InArgs, TSharedRef<FWorkspaceEditor> InHostingApp, const FWorkspaceDocument& InDocument);
@@ -71,6 +74,8 @@ class SGraphDocument : public SCompoundWidget
 	bool CanSelectAllNodes() const;
 	void SelectAllNodes();
 
+	bool CanOpenInNewTab() const;
+	void OpenInNewTab();
 
 	bool IsEditable(UEdGraph* InGraph) const;
 	void OnResetSelection();
@@ -108,6 +113,10 @@ class SGraphDocument : public SCompoundWidget
 	FOnCanPerformActionOnSelectedNodes OnCanDuplicateSelectedNodes;
 	FOnDuplicateSelectedNodes OnDuplicateSelectedNodes;
 
+	// Open in new Tab delegates
+	FOnCanPerformActionOnSelectedNodes OnCanOpenInNewTab;
+	FOnPerformActionOnSelectedNodes OnOpenInNewTab;
+
 	friend class FWorkspaceEditorModule;
 };
 
diff --git a/Engine/Plugins/Experimental/Workspace/Source/WorkspaceEditor/Private/SWorkspaceTabWrapper.cpp b/Engine/Plugins/Experimental/Workspace/Source/WorkspaceEditor/Private/SWorkspaceTabWrapper.cpp
--- a/Engine/Plugins/Experimental/Workspace/Source/WorkspaceEditor/Private/SWorkspaceTabWrapper.cpp
+++ b/Engine/Plugins/Experimental/Workspace/Source/WorkspaceEditor/Private/SWorkspaceTabWrapper.cpp
@@ -247,17 +247,7 @@ void SWorkspaceTabWrapper::RebuildBreadcrumbTrail() const
 	if (Export.GetFirstAssetPath().IsValid())
 	{
 		TArray<FWorkspaceOutlinerItemExport> Exports;
-		if (FWorkspaceOutlinerAssetReferenceItemData::IsAssetReference(Export))
-		{
-			// Prepend outer asset/export chain
-			TArray<FSoftObjectPath> AssetPaths;
-			Export.GetAssetPaths(AssetPaths);
-			Export.GetExports(Exports);
-		}
-		else
-		{
-			Exports.Add(Export);
-		}
+		Export.GetExports(Exports);
 
 		for (const FWorkspaceOutlinerItemExport& DocumentExport : Exports)
 		{
diff --git a/Engine/Plugins/Experimental/Workspace/Source/WorkspaceEditor/Private/WorkspaceEditor.cpp b/Engine/Plugins/Experimental/Workspace/Source/WorkspaceEditor/Private/WorkspaceEditor.cpp
--- a/Engine/Plugins/Experimental/Workspace/Source/WorkspaceEditor/Private/WorkspaceEditor.cpp
+++ b/Engine/Plugins/Experimental/Workspace/Source/WorkspaceEditor/Private/WorkspaceEditor.cpp
@@ -360,32 +360,33 @@ void FWorkspaceEditor::OpenAssets(TConstArrayView<FAssetData> InAssets)
 	}
 }
 
-void FWorkspaceEditor::OpenExports(TConstArrayView<FWorkspaceOutlinerItemExport> InExports)
+void FWorkspaceEditor::OpenExports(TConstArrayView<FWorkspaceOutlinerItemExport> InExports, TOptional<FDocumentTracker::EOpenDocumentCause> OpenCauseOverride)
 {
-	TArray<const FWorkspaceOutlinerItemExport*> OpenedExports;
+	const FDocumentTracker::EOpenDocumentCause OpenCause = OpenCauseOverride.Get(FDocumentTracker::EOpenDocumentCause::NavigatingCurrentDocument); 
+
+	TArray<const FWorkspaceOutlinerItemExport> OpenedExports;
 	for(const FWorkspaceOutlinerItemExport& Export : InExports)
 	{
-		FWorkspaceOutlinerItemExport QualifiedExport = Export;
+		FWorkspaceOutlinerItemExport QualifiedExport = Export; 
 		WorkspaceView->GetWorkspaceExportData(QualifiedExport);
 		
-		if(UObject* LoadedAsset = QualifiedExport.GetFirstAssetPath().TryLoad())
+		if(UObject* LoadedObject = QualifiedExport.GetFirstValidObjectPath().TryLoad())
 		{
-			OpenDocument(LoadedAsset, QualifiedExport, FDocumentTracker::EOpenDocumentCause::NavigatingCurrentDocument);
-			OpenedExports.AddUnique(&QualifiedExport);
+			OpenDocument(LoadedObject, QualifiedExport, OpenCause);
+			OpenedExports.AddUnique(QualifiedExport); 
 		}
 	}
 
 	if(OpenedExports.Num() > 0)
 	{
-		if(const FWorkspaceOutlinerItemExport* LoadedExport = OpenedExports.Last())
-		{
-			WorkspaceView->SelectExport(*LoadedExport);
-		}
+		WorkspaceView->SelectExport(OpenedExports.Last());
 	}
 }
 
-void FWorkspaceEditor::OpenObjects(TConstArrayView<UObject*> InObjects)
+void FWorkspaceEditor::OpenObjects(TConstArrayView<UObject*> InObjects, TOptional<FDocumentTracker::EOpenDocumentCause> OpenCauseOverride)
 {
+	const FDocumentTracker::EOpenDocumentCause OpenCause = OpenCauseOverride.Get(InObjects.Num() == 1 ? FDocumentTracker::EOpenDocumentCause::NavigatingCurrentDocument : FDocumentTracker::EOpenDocumentCause::OpenNewDocument);
+
 	for(UObject* Object : InObjects)
 	{
 		const FWorkspaceEditorModule& WorkspaceEditorModule = FModuleManager::LoadModuleChecked<FWorkspaceEditorModule>("WorkspaceEditor");
@@ -398,7 +399,7 @@ void FWorkspaceEditor::OpenObjects(TConstArrayView<UObject*> InObjects)
 
 		if(Object)
 		{
-			TSharedPtr<SDockTab> DocumentTab = OpenDocument(Object, InObjects.Num() == 1 ? FDocumentTracker::EOpenDocumentCause::NavigatingCurrentDocument : FDocumentTracker::EOpenDocumentCause::OpenNewDocument);
+			TSharedPtr<SDockTab> DocumentTab = OpenDocument(Object, OpenCause);
 			if(DocumentSubObjectArgs && DocumentTab.IsValid())
 			{
 				TSharedRef<SWorkspaceTabWrapper> WorkspaceTabWrapper = StaticCastSharedRef<SWorkspaceTabWrapper>(DocumentTab->GetContent());
diff --git a/Engine/Plugins/Experimental/Workspace/Source/WorkspaceEditor/Private/WorkspaceEditor.h b/Engine/Plugins/Experimental/Workspace/Source/WorkspaceEditor/Private/WorkspaceEditor.h
--- a/Engine/Plugins/Experimental/Workspace/Source/WorkspaceEditor/Private/WorkspaceEditor.h
+++ b/Engine/Plugins/Experimental/Workspace/Source/WorkspaceEditor/Private/WorkspaceEditor.h
@@ -95,8 +95,8 @@ private:
 
 	// IWorkspaceEditor interface
 	virtual void OpenAssets(TConstArrayView<FAssetData> InAssets) override;
-	virtual void OpenExports(TConstArrayView<FWorkspaceOutlinerItemExport> InExports) override;
-	virtual void OpenObjects(TConstArrayView<UObject*> InObjects) override;
+	virtual void OpenExports(TConstArrayView<FWorkspaceOutlinerItemExport> InExports, TOptional<FDocumentTracker::EOpenDocumentCause> OpenCauseOverride) override;
+	virtual void OpenObjects(TConstArrayView<UObject*> InObjects, TOptional<FDocumentTracker::EOpenDocumentCause> OpenCauseOverride) override;
 	virtual void GetOpenedAssetsOfClass(TSubclassOf<UObject> InClass, TArray<UObject*>& OutObjects) const override;
 	virtual void GetAssets(TArray<FAssetData>& OutAssets, bool bIncludeAssetReferences = false) const override;
 	virtual void CloseObjects(TConstArrayView<UObject*> InObjects) override;
diff --git a/Engine/Plugins/Experimental/Workspace/Source/WorkspaceEditor/Private/WorkspaceEditorModule.cpp b/Engine/Plugins/Experimental/Workspace/Source/WorkspaceEditor/Private/WorkspaceEditorModule.cpp
--- a/Engine/Plugins/Experimental/Workspace/Source/WorkspaceEditor/Private/WorkspaceEditorModule.cpp
+++ b/Engine/Plugins/Experimental/Workspace/Source/WorkspaceEditor/Private/WorkspaceEditorModule.cpp
@@ -117,7 +117,9 @@ FObjectDocumentArgs FWorkspaceEditorModule::CreateGraphDocumentArgs(const FGraph
 			.OnDuplicateSelectedNodes(InArgs.OnDuplicateSelectedNodes)
 			.OnNavigateHistoryForward_Lambda([WeakWorkspaceEditor](){ if(const TSharedPtr<FWorkspaceEditor> SharedWorkspaceEditor = WeakWorkspaceEditor.Pin()) SharedWorkspaceEditor->NavigateForward(); })
 			.OnNavigateHistoryBack_Lambda([WeakWorkspaceEditor](){ if(const TSharedPtr<FWorkspaceEditor> SharedWorkspaceEditor = WeakWorkspaceEditor.Pin()) SharedWorkspaceEditor->NavigateBack(); })
-			.OnNodeDoubleClicked(InArgs.OnNodeDoubleClicked);
+			.OnNodeDoubleClicked(InArgs.OnNodeDoubleClicked)
+			.OnCanOpenInNewTab(InArgs.OnCanOpenInNewTab)
+			.OnOpenInNewTab(InArgs.OnOpenInNewTab);
 
 		// Let listeners know that the document widget has been created 
 		InArgs.OnGraphDocumentCreated.ExecuteIfBound(InContext, Widget);
diff --git a/Engine/Plugins/Experimental/Workspace/Source/WorkspaceEditor/Public/IWorkspaceEditor.h b/Engine/Plugins/Experimental/Workspace/Source/WorkspaceEditor/Public/IWorkspaceEditor.h
--- a/Engine/Plugins/Experimental/Workspace/Source/WorkspaceEditor/Public/IWorkspaceEditor.h
+++ b/Engine/Plugins/Experimental/Workspace/Source/WorkspaceEditor/Public/IWorkspaceEditor.h
@@ -3,6 +3,7 @@
 #pragma once
 
 #include "Tools/BaseAssetToolkit.h"
+#include "WorkflowOrientedApp/WorkflowTabManager.h"
 
 #define UE_API WORKSPACEEDITOR_API
 
@@ -40,10 +41,10 @@ public:
 	virtual void OpenAssets(TConstArrayView<FAssetData> InAssets) = 0;
 
 	// Open the supplied exports for editing within the workspace editor
-	virtual void OpenExports(TConstArrayView<FWorkspaceOutlinerItemExport> InExports) = 0;
+	virtual void OpenExports(TConstArrayView<FWorkspaceOutlinerItemExport> InExports, TOptional<FDocumentTracker::EOpenDocumentCause> OpenCauseOverride = TOptional<FDocumentTracker::EOpenDocumentCause>()) = 0;
 
 	// Open the supplied objects for editing within the workspace editor
-	virtual void OpenObjects(TConstArrayView<UObject*> InObjects) = 0;
+	virtual void OpenObjects(TConstArrayView<UObject*> InObjects, TOptional<FDocumentTracker::EOpenDocumentCause> OpenCauseOverride = TOptional<FDocumentTracker::EOpenDocumentCause>()) = 0;
 
 	// Get the current set of opened (loaded) assets of the specified class
 	virtual void GetOpenedAssetsOfClass(TSubclassOf<UObject> InClass, TArray<UObject*>& OutObjects) const = 0;
diff --git a/Engine/Plugins/Experimental/Workspace/Source/WorkspaceEditor/Public/IWorkspaceEditorModule.h b/Engine/Plugins/Experimental/Workspace/Source/WorkspaceEditor/Public/IWorkspaceEditorModule.h
--- a/Engine/Plugins/Experimental/Workspace/Source/WorkspaceEditor/Public/IWorkspaceEditorModule.h
+++ b/Engine/Plugins/Experimental/Workspace/Source/WorkspaceEditor/Public/IWorkspaceEditorModule.h
@@ -222,6 +222,9 @@ struct FGraphDocumentWidgetArgs
 	FOnNodeDoubleClicked OnNodeDoubleClicked;
 
 	FOnGraphDocumentCreated OnGraphDocumentCreated;
+
+	FOnCanPerformActionOnSelectedNodes OnCanOpenInNewTab;
+	FOnPerformActionOnSelectedNodes OnOpenInNewTab;
 };
 
 // Enum describing how to open a workspace
diff --git a/Engine/Plugins/Experimental/Workspace/Source/WorkspaceEditor/Public/WorkspaceAssetRegistryInfo.h b/Engine/Plugins/Experimental/Workspace/Source/WorkspaceEditor/Public/WorkspaceAssetRegistryInfo.h
--- a/Engine/Plugins/Experimental/Workspace/Source/WorkspaceEditor/Public/WorkspaceAssetRegistryInfo.h
+++ b/Engine/Plugins/Experimental/Workspace/Source/WorkspaceEditor/Public/WorkspaceAssetRegistryInfo.h
@@ -83,13 +83,21 @@ struct FWorkspaceOutlinerItemExport
 		Path.PathSegments.Add(*InObjectPath.ToString());
 		Path.PathSegments.Add(InIdentifier);
 	}
-	
-	FWorkspaceOutlinerItemExport(const FName InIdentifier, const FWorkspaceOutlinerItemExport& InParent) : Path(InParent.Path.AppendSegment(InIdentifier))
+
+	FWorkspaceOutlinerItemExport(const FName InIdentifier, const FWorkspaceOutlinerItemExport& InParent) 
+		: Path(InParent.Path.AppendSegment(InIdentifier))
+	{
+	}
+
+	FWorkspaceOutlinerItemExport(const FName InIdentifier, const FWorkspaceOutlinerItemExport& InParent, const TInstancedStruct<FWorkspaceOutlinerItemData>& InData) 
+		: Path(InParent.Path.AppendSegment(InIdentifier))
+		, Data(InData)
 	{
 	}
 
-	FWorkspaceOutlinerItemExport(const FName InIdentifier, const FWorkspaceOutlinerItemExport& InParent, const TInstancedStruct<FWorkspaceOutlinerItemData>& InData) : Path(InParent.Path.AppendSegment(InIdentifier)),  Data(InData)
+	bool operator==(const FWorkspaceOutlinerItemExport& Other) const
 	{
+		return this->GetFullPath().Equals(Other.GetFullPath());
 	}
 
 protected:
@@ -157,6 +165,20 @@ public:
 				FSoftObjectPath ObjectPath = Path.PathSegments[SegmentIndex].ToString();
 				if (ObjectPath.IsValid())
 				{
+					if (!ObjectPath.IsSubobject() && !GetIdentifier().IsNone() && ObjectPath.GetAssetName() != GetIdentifier())
+					{
+						// Check if Identifier points to a valid subobject
+						ObjectPath.SetSubPathString(GetIdentifier().ToUtf8String());
+						if (UObject* Object = ObjectPath.TryLoad())
+						{
+							if (AssetClass* TypedObject = Cast<AssetClass>(Object))
+							{
+								return TypedObject;
+							}
+						}
+						ObjectPath.SetSubPathString(FUtf8String());
+					}
+
 					if (UObject* Object = ObjectPath.TryLoad())
 					{
 						if (AssetClass* TypedObject = Cast<AssetClass>(Object))
@@ -192,6 +214,44 @@ public:
 		return FSoftObjectPath();
 	}
 
+	// Tries to get the first valid path to an object - taking into account the Identifier might point to a subobject
+	// To verify validity, it may try to load the object if bForceLoad == true (default)
+	FSoftObjectPath GetFirstValidObjectPath(bool bForceLoad = true) const
+	{
+		if (Path.PathSegments.Num() > 0)
+		{
+			for (int32 SegmentIndex = Path.PathSegments.Num() - 1; SegmentIndex >= 0; SegmentIndex--)
+			{
+				FSoftObjectPath ObjectPath = Path.PathSegments[SegmentIndex].ToString();
+				if (ObjectPath.IsValid() && ObjectPath.IsAsset())
+				{
+					if (ObjectPath.IsSubobject())
+					{
+						// If the path is valid and already includes a subobject we can return directly
+						return ObjectPath;
+					}
+					else if (ObjectPath.GetAssetName() != GetIdentifier())
+					{
+						// Otherwise check if adding the identifier as a subobject results in a valid object.
+						ObjectPath.SetSubPathString(GetIdentifier().ToUtf8String());
+
+						UObject* TestObject = ObjectPath.ResolveObject();
+						TestObject = TestObject ? TestObject : (bForceLoad ? ObjectPath.TryLoad() : nullptr);
+						if (TestObject)
+						{
+							return ObjectPath;
+						}
+
+						ObjectPath.SetSubPathString(FUtf8String());
+					}
+					return ObjectPath;
+				}
+			}
+		}
+
+		return FSoftObjectPath();
+	}
+
 	// Returns the first path segment as a FSoftObjectPath. e.g:
 	//  - "SoftObjectPath" - "Foo" - "SoftObjectPathTwo" - "Bar" will return SoftObjectPath	
 	//  - "Foo" - "SoftObjectPath" - "SoftObjectPathTwo" - "Bar" will return FSoftObjectPath()
@@ -322,15 +382,7 @@ struct FWorkspaceOutlinerAssetReferenceItemData : public FWorkspaceOutlinerItemD
 
 	static bool IsAssetReference(const FWorkspaceOutlinerItemExport& InExport)
 	{
-		if (InExport.HasData())
-		{
-			if (InExport.GetData().GetScriptStruct() == FWorkspaceOutlinerAssetReferenceItemData::StaticStruct())
-			{
-				return true;
-			}			
-		}
-		
-		return false;
+		return InExport.HasData() && InExport.GetData().GetScriptStruct() == FWorkspaceOutlinerAssetReferenceItemData::StaticStruct();
 	}
 };
 
@@ -346,6 +398,11 @@ struct FWorkspaceOutlinerGroupItemData : public FWorkspaceOutlinerItemData
 
 	UPROPERTY()
 	FSlateBrush GroupIcon;
+
+	static bool IsGroupItem(const FWorkspaceOutlinerItemExport& InExport)
+	{
+		return InExport.HasData() && InExport.GetData().GetScriptStruct() == FWorkspaceOutlinerGroupItemData::StaticStruct();
+	}
 };
 
 namespace UE::Workspace
