UAF Wizard - Only set module on SMC if the skeletal mesh has been set to avoid crash

#rnx

diff --git a/Engine/Plugins/Experimental/UAF/UAF/Source/UAFEditor/Private/Wizard/AssetWizard.cpp b/Engine/Plugins/Experimental/UAF/UAF/Source/UAFEditor/Private/Wizard/AssetWizard.cpp
--- a/Engine/Plugins/Experimental/UAF/UAF/Source/UAFEditor/Private/Wizard/AssetWizard.cpp
+++ b/Engine/Plugins/Experimental/UAF/UAF/Source/UAFEditor/Private/Wizard/AssetWizard.cpp
@@ -211,7 +211,7 @@ namespace UE::UAF::Editor
 					SkelMeshComp->SetRelativeTransform(RelativeTransform);
 				}
 
-				if (UAFComp && ModuleAssetPair)
+				if (UAFComp && ModuleAssetPair && SkelMeshComp && SkelMeshComp->GetSkeletalMeshAsset())
 				{
 					UAFComp->SetModule(CastChecked<UAnimNextModule>(ModuleAssetPair->UserAsset.Get()));
 				}
