Fix UAF asset wizard template assets being visible in the content browser and selectable in the workspace picker etc.
* This CL hides the UAF content directory. Visibility can be toggled with the `UAF.ToggleHiddenContent` if needed.

diff --git a/Engine/Plugins/Experimental/UAF/UAF/Source/UAFEditor/Private/AnimNextEditorModule.cpp b/Engine/Plugins/Experimental/UAF/UAF/Source/UAFEditor/Private/AnimNextEditorModule.cpp
--- a/Engine/Plugins/Experimental/UAF/UAF/Source/UAFEditor/Private/AnimNextEditorModule.cpp
+++ b/Engine/Plugins/Experimental/UAF/UAF/Source/UAFEditor/Private/AnimNextEditorModule.cpp
@@ -77,6 +77,8 @@
 #include "UAF/Viewport/SystemViewportController.h"
 #include "Variables/Outliner/VariablesOutlinerCommands.h"
 #include "Wizard/AssetWizard.h"
+#include "Interfaces/IPluginManager.h"
+#include "AssetToolsModule.h"
 
 #define LOCTEXT_NAMESPACE "AnimNextEditorModule"
 
@@ -93,6 +95,31 @@ namespace
 
 namespace UE::UAF::Editor
 {
+	const auto ToggleUAFContentHidden = FConsoleCommandDelegate::CreateStatic([]()
+    {
+    	static const TCHAR* DenyItemKey = TEXT("UAFContent");
+    
+    	if (TSharedPtr<IPlugin> UAFPlugin = IPluginManager::Get().FindPlugin("UAF"))
+    	{
+    		TSharedRef<FPathPermissionList>& FolderPermissions = FAssetToolsModule::GetModule().Get().GetFolderPermissionList();
+
+			const FString DenyItemValue = UAFPlugin->GetMountedAssetPath();
+
+    		if (FolderPermissions->ContainsDenyListItem(DenyItemValue))
+    		{
+    			FolderPermissions->RemoveDenyListItem(DenyItemKey, DenyItemValue);
+    		}
+		    else
+		    {
+		    	FolderPermissions->AddDenyListItem(DenyItemKey, DenyItemValue);
+		    }
+    	}
+    });
+
+static FAutoConsoleCommand ToggleHiddenUAFContent(
+    TEXT("UAF.ToggleHiddenContent"),
+    TEXT("Toggles whether internal UAF content is shown in the editor"),
+    ToggleUAFContentHidden);
 	
 void FAnimNextEditorModule::StartupModule()
 {
@@ -505,6 +532,9 @@ void FAnimNextEditorModule::StartupModule()
 	{
 		return MakeUnique<FSystemViewportController>();
 	});
+	
+	// Hide the UAF content directory and its contents
+	ToggleUAFContentHidden.Execute();
 }
 
 void FAnimNextEditorModule::ShutdownModule()
