Crash fixes for LODPose

diff --git a/Engine/Plugins/Experimental/UAF/UAF/Source/UAF/Public/LODPose.h b/Engine/Plugins/Experimental/UAF/UAF/Source/UAF/Public/LODPose.h
--- a/Engine/Plugins/Experimental/UAF/UAF/Source/UAF/Public/LODPose.h
+++ b/Engine/Plugins/Experimental/UAF/UAF/Source/UAF/Public/LODPose.h
@@ -41,6 +41,11 @@ struct FLODPose
 
 	void CopyFrom(const FLODPose& SourcePose)
 	{
+		if (!SourcePose.IsValid())
+		{
+			return;
+		}
+		
 		check(RefPose && SourcePose.RefPose);
 		check(RefPose->ReferenceLocalTransforms.Num() == SourcePose.RefPose->ReferenceLocalTransforms.Num());
 
@@ -269,7 +274,7 @@ struct FLODPose
 
 	bool IsValid() const
 	{
-		return RefPose != nullptr;
+		return (RefPose != nullptr && RefPose->IsValid());
 	}
 
 	bool IsAdditive() const
