Fix crash from selecting non-Actor derived blueprint to modify in UAF asset wizard

#rnx

diff --git a/Engine/Plugins/Experimental/UAF/UAF/Source/UAFEditor/Private/Wizard/STemplatePicker.cpp b/Engine/Plugins/Experimental/UAF/UAF/Source/UAFEditor/Private/Wizard/STemplatePicker.cpp
--- a/Engine/Plugins/Experimental/UAF/UAF/Source/UAFEditor/Private/Wizard/STemplatePicker.cpp
+++ b/Engine/Plugins/Experimental/UAF/UAF/Source/UAFEditor/Private/Wizard/STemplatePicker.cpp
@@ -6,6 +6,8 @@
 #include "AssetRegistry/AssetRegistryModule.h"
 #include "AssetThumbnail.h" 
 #include "DetailsViewArgs.h"
+#include "Engine/Blueprint.h"
+#include "Misc/MessageDialog.h"
 #include "PropertyEditorModule.h"
 #include "SPrimaryButton.h"
 #include "Styling/StyleColors.h"
@@ -67,6 +69,7 @@ namespace UE::UAF::Editor
 		TemplateConfig = GetMutableDefault<UUAFTemplateConfig>();
 		TemplateConfig->Reset();
 		TemplateConfig->OnOutputPathChanged.AddSP(this, &STemplatePicker::HandleOutputPathChanged);
+		TemplateConfig->OnBlueprintToModifyChanged.AddSP(this, &STemplatePicker::HandleBlueprintToModifyChanged);
 		
 		FDetailsViewArgs DetailsViewArgs;
 		DetailsViewArgs.bAllowSearch = false;
@@ -393,6 +396,7 @@ namespace UE::UAF::Editor
 							return SelectedTemplate.IsValid();
 						})
 						.Text(LOCTEXT("CreateAssets", "Create Assets"))
+						.IsEnabled(this, &STemplatePicker::CanCreateAssets)
 						.OnClicked_Lambda([this]()
 						{
 							HandleUseTemplate();
@@ -488,6 +492,21 @@ namespace UE::UAF::Editor
 		RebuildIncludedAssetsListView();
 	}
 
+	void STemplatePicker::HandleBlueprintToModifyChanged()
+	{
+		if (!TemplateConfig->BlueprintToModify.IsNull())
+		{
+			UBlueprint* BP = TemplateConfig->BlueprintToModify.LoadSynchronous();
+			if (!BP->GeneratedClass->IsChildOf<AActor>())
+			{
+				// Restricting the blueprint picker to only allow Actor-derived BPs requires loading the asset to verify the class inheritance which
+				// is too slow so the check is performed at selection time instead.
+				FMessageDialog::Open(EAppMsgCategory::Error, EAppMsgType::Ok, LOCTEXT("InvalidBlueprintSelected", "The selected blueprint to modify must derive from Actor to be able to contain the components necessary for UAF to function."));
+				TemplateConfig->BlueprintToModify.Reset();
+			}
+		}
+	}
+
 	void STemplatePicker::RebuildIncludedAssetsListView()
 	{
 		UEditorAssetSubsystem* EditorAssetSubsystem = GEditor->GetEditorSubsystem<UEditorAssetSubsystem>();
@@ -583,6 +602,23 @@ namespace UE::UAF::Editor
 		TemplatesTileView->RequestListRefresh();
 	}
 
+	bool STemplatePicker::CanCreateAssets() const
+	{
+		switch (TemplateConfig->BlueprintMode)
+		{
+		case ETemplateBlueprintMode::CreateNewBlueprint:
+			return !TemplateConfig->BlueprintAssetName.IsEmpty()
+				&& TemplateConfig->BlueprintClass.Get()
+				&& !TemplateConfig->SkeletalMesh.IsNull();
+		case ETemplateBlueprintMode::ModifyExistingBlueprint:
+			return !TemplateConfig->BlueprintToModify.IsNull();
+		case ETemplateBlueprintMode::DoNothing:
+			return true;
+		}
+
+		return false;
+	}
+	
 	void SModuleTemplate::Construct(const FArguments& InArgs, const TSharedRef<STableViewBase>& OwnerTable)
 	{
 		Item = InArgs._Item;
diff --git a/Engine/Plugins/Experimental/UAF/UAF/Source/UAFEditor/Private/Wizard/STemplatePicker.h b/Engine/Plugins/Experimental/UAF/UAF/Source/UAFEditor/Private/Wizard/STemplatePicker.h
--- a/Engine/Plugins/Experimental/UAF/UAF/Source/UAFEditor/Private/Wizard/STemplatePicker.h
+++ b/Engine/Plugins/Experimental/UAF/UAF/Source/UAFEditor/Private/Wizard/STemplatePicker.h
@@ -43,7 +43,11 @@ namespace UE::UAF::Editor
 		
 		void HandleOutputPathChanged();
 
+		void HandleBlueprintToModifyChanged();
+
 		void RebuildIncludedAssetsListView();
+
+		bool CanCreateAssets() const;
 		
 		enum class EPanel
 		{
diff --git a/Engine/Plugins/Experimental/UAF/UAF/Source/UAFEditor/Private/Wizard/TemplateConfig.cpp b/Engine/Plugins/Experimental/UAF/UAF/Source/UAFEditor/Private/Wizard/TemplateConfig.cpp
--- a/Engine/Plugins/Experimental/UAF/UAF/Source/UAFEditor/Private/Wizard/TemplateConfig.cpp
+++ b/Engine/Plugins/Experimental/UAF/UAF/Source/UAFEditor/Private/Wizard/TemplateConfig.cpp
@@ -12,10 +12,15 @@ void UUAFTemplateConfig::PostEditChangeProperty(FPropertyChangedEvent& PropertyC
 	{
 		OnOutputPathChanged.Broadcast();
 	}
+	else if (PropertyChangedEvent.GetMemberPropertyName() == GET_MEMBER_NAME_CHECKED(UUAFTemplateConfig, BlueprintToModify))
+	{
+		OnBlueprintToModifyChanged.Broadcast();
+	}
 }
 
 void UUAFTemplateConfig::Reset()
 {
 	OnOutputPathChanged.Clear();
+	OnBlueprintToModifyChanged.Clear();
 	AssetNaming.Empty();
 }
\ No newline at end of file
diff --git a/Engine/Plugins/Experimental/UAF/UAF/Source/UAFEditor/Private/Wizard/TemplateConfig.h b/Engine/Plugins/Experimental/UAF/UAF/Source/UAFEditor/Private/Wizard/TemplateConfig.h
--- a/Engine/Plugins/Experimental/UAF/UAF/Source/UAFEditor/Private/Wizard/TemplateConfig.h
+++ b/Engine/Plugins/Experimental/UAF/UAF/Source/UAFEditor/Private/Wizard/TemplateConfig.h
@@ -48,6 +48,7 @@ public:
 	virtual void PostEditChangeProperty(FPropertyChangedEvent& PropertyChangedEvent) override;
 
 	FSimpleMulticastDelegate OnOutputPathChanged;
+	FSimpleMulticastDelegate OnBlueprintToModifyChanged;
 
 	void Reset();
 
