Fix double binding to IsEnabled

#rnx

diff --git a/Engine/Plugins/Experimental/UAF/UAF/Source/UAFEditor/Private/Wizard/STemplatePicker.cpp b/Engine/Plugins/Experimental/UAF/UAF/Source/UAFEditor/Private/Wizard/STemplatePicker.cpp
--- a/Engine/Plugins/Experimental/UAF/UAF/Source/UAFEditor/Private/Wizard/STemplatePicker.cpp
+++ b/Engine/Plugins/Experimental/UAF/UAF/Source/UAFEditor/Private/Wizard/STemplatePicker.cpp
@@ -391,10 +391,6 @@ namespace UE::UAF::Editor
 					.Padding(FMargin(16.0f, 16.0f, 8.0f, 16.0f))
 					[
 						SNew(SPrimaryButton)
-						.IsEnabled_Lambda([this]()
-						{
-							return SelectedTemplate.IsValid();
-						})
 						.Text(LOCTEXT("CreateAssets", "Create Assets"))
 						.IsEnabled(this, &STemplatePicker::CanCreateAssets)
 						.OnClicked_Lambda([this]()
@@ -604,6 +600,11 @@ namespace UE::UAF::Editor
 
 	bool STemplatePicker::CanCreateAssets() const
 	{
+		if (!SelectedTemplate.IsValid())
+		{
+			return false;
+		}
+		
 		switch (TemplateConfig->BlueprintMode)
 		{
 		case ETemplateBlueprintMode::CreateNewBlueprint:
