Allow export access specifiers to be set on multiple selected variables in the variables outliner

diff --git a/Engine/Plugins/Experimental/UAF/UAF/Source/UAFEditor/Private/Variables/Outliner/VariablesOutlinerColumns.cpp b/Engine/Plugins/Experimental/UAF/UAF/Source/UAFEditor/Private/Variables/Outliner/VariablesOutlinerColumns.cpp
--- a/Engine/Plugins/Experimental/UAF/UAF/Source/UAFEditor/Private/Variables/Outliner/VariablesOutlinerColumns.cpp
+++ b/Engine/Plugins/Experimental/UAF/UAF/Source/UAFEditor/Private/Variables/Outliner/VariablesOutlinerColumns.cpp
@@ -24,6 +24,7 @@
 #include "Variables/AnimNextVariableSettings.h"
 #include "Widgets/Input/SButton.h"
 #include "Widgets/Input/SCheckBox.h"
+#include "Variables/Outliner/VariablesOutlinerMode.h"
 
 #define LOCTEXT_NAMESPACE "VariablesOutlinerColumns"
 
@@ -557,21 +558,17 @@ class SVariablesOutlinerAccessSpecifier : public SCompoundWidget, public FNotify
 
 	FReply OnClicked() const
 	{
-		TSharedPtr<FVariablesOutlinerEntryItem> TreeItem = WeakTreeItem.Pin();
-		if(!TreeItem.IsValid())
-		{
-			return FReply::Unhandled();
-		}
-		
-		IAnimNextRigVMExportInterface* Export = Cast<IAnimNextRigVMExportInterface>(TreeItem->WeakEntry.Get());
-		if(Export == nullptr)
+		if (TSharedPtr<FVariablesOutlinerEntryItem> TreeItem = WeakTreeItem.Pin())
 		{
-			return FReply::Unhandled();
+			if (TSharedPtr<ISceneOutliner> SceneOutliner = TreeItem->WeakSceneOutliner.Pin())
+			{
+				if (const FVariablesOutlinerMode* Mode = static_cast<const FVariablesOutlinerMode*>(SceneOutliner->GetMode()))
+				{
+					Mode->ToggleSelectedVariablesExportAccessSpecifiers(TreeItem);
+				}
+			}
 		}
 
-		FScopedTransaction Transaction(LOCTEXT("SetAccessSpecifierTransaction", "Set Access Specifier"));
-		Export->SetExportAccessSpecifier(Export->GetExportAccessSpecifier() == EAnimNextExportAccessSpecifier::Public ? EAnimNextExportAccessSpecifier::Private : EAnimNextExportAccessSpecifier::Public);  
-
 		return FReply::Unhandled();	// Fall through so we dont deselect our item
 	}
 
diff --git a/Engine/Plugins/Experimental/UAF/UAF/Source/UAFEditor/Private/Variables/Outliner/VariablesOutlinerCommands.cpp b/Engine/Plugins/Experimental/UAF/UAF/Source/UAFEditor/Private/Variables/Outliner/VariablesOutlinerCommands.cpp
--- a/Engine/Plugins/Experimental/UAF/UAF/Source/UAFEditor/Private/Variables/Outliner/VariablesOutlinerCommands.cpp
+++ b/Engine/Plugins/Experimental/UAF/UAF/Source/UAFEditor/Private/Variables/Outliner/VariablesOutlinerCommands.cpp
@@ -16,7 +16,8 @@ void FVariablesOutlinerCommands::RegisterCommands()
 
 	UI_COMMAND(SaveAsset, "Save Asset", "Save the selected asset.", EUserInterfaceActionType::Button, FInputChord());
 
-	UI_COMMAND(ToggleVariableExport, "Toggle Variable Export", "Sets Variable to be either Public or Private.", EUserInterfaceActionType::Button, FInputChord());
+	UI_COMMAND(SetVariableExportPublic, "Sets Variable Export to Public", "Sets Variable to be Public.", EUserInterfaceActionType::Button, FInputChord());
+	UI_COMMAND(SetVariableExportPrivate, "Sets Variable Export to Private", "Sets Variable to be Private.", EUserInterfaceActionType::Button, FInputChord());
 
 	UI_COMMAND(CreateSharedVariablesAssets, "Create Shared Variables from Selection", "Creates a new Shared Variables asset from the selected set of variables.", EUserInterfaceActionType::Button, FInputChord());
 }
diff --git a/Engine/Plugins/Experimental/UAF/UAF/Source/UAFEditor/Private/Variables/Outliner/VariablesOutlinerCommands.h b/Engine/Plugins/Experimental/UAF/UAF/Source/UAFEditor/Private/Variables/Outliner/VariablesOutlinerCommands.h
--- a/Engine/Plugins/Experimental/UAF/UAF/Source/UAFEditor/Private/Variables/Outliner/VariablesOutlinerCommands.h
+++ b/Engine/Plugins/Experimental/UAF/UAF/Source/UAFEditor/Private/Variables/Outliner/VariablesOutlinerCommands.h
@@ -23,7 +23,8 @@ public:
 	TSharedPtr<FUICommandInfo> FindReferencesInWorkspace;
 	TSharedPtr<FUICommandInfo> FindReferencesInAsset;
 	TSharedPtr<FUICommandInfo> SaveAsset;	
-	TSharedPtr<FUICommandInfo> ToggleVariableExport;
+	TSharedPtr<FUICommandInfo> SetVariableExportPublic;
+	TSharedPtr<FUICommandInfo> SetVariableExportPrivate;
 	TSharedPtr<FUICommandInfo> CreateSharedVariablesAssets;
 };
 
diff --git a/Engine/Plugins/Experimental/UAF/UAF/Source/UAFEditor/Private/Variables/Outliner/VariablesOutlinerMode.cpp b/Engine/Plugins/Experimental/UAF/UAF/Source/UAFEditor/Private/Variables/Outliner/VariablesOutlinerMode.cpp
--- a/Engine/Plugins/Experimental/UAF/UAF/Source/UAFEditor/Private/Variables/Outliner/VariablesOutlinerMode.cpp
+++ b/Engine/Plugins/Experimental/UAF/UAF/Source/UAFEditor/Private/Variables/Outliner/VariablesOutlinerMode.cpp
@@ -38,6 +38,7 @@
 #include "Variables/AnimNextSharedVariablesFactory.h"
 #include "Variables/AnimNextVariableSettings.h"
 #include "Widgets/Docking/SDockTab.h"
+#include "AnimNextScopedCompileJob.h"
 
 #define LOCTEXT_NAMESPACE "VariablesOutlinerMode"
 
@@ -139,14 +140,42 @@ TSharedPtr<SWidget> FVariablesOutlinerMode::CreateContextMenu()
 					
 					VariablesSection.AddMenuEntryWithCommandList(FGenericCommands::Get().Copy, MenuContext->WeakCommandList.Pin());
 
-					if (NumSelectedVariables == 1)
+					// Check which export access related context menu entries we need to populate
+					if (NumSelectedVariables > 0)
 					{
-						if (const UAnimNextVariableEntry* VariableEntry = Cast<UAnimNextVariableEntry>(MenuContext->WeakEntries[0].Get()))
+						bool bHasPublicVariables = false;
+						bool bHasPrivateVariables = false;
+						for (const TWeakObjectPtr<UAnimNextRigVMAssetEntry> WeakEntry : MenuContext->WeakEntries)
 						{
-							VariablesSection.AddMenuEntryWithCommandList(FVariablesOutlinerCommands::Get().ToggleVariableExport, MenuContext->WeakCommandList.Pin(), VariableEntry->GetExportAccessSpecifier() == EAnimNextExportAccessSpecifier::Private ? LOCTEXT("MakePublicMenuLabel", "Expose as Public") : LOCTEXT("MakePrivateMenuLabel", "Make Private"), TAttribute<FText>(), FSlateIcon(FAppStyle::Get().GetStyleSetName(), VariableEntry->GetExportAccessSpecifier() == EAnimNextExportAccessSpecifier::Private ? "Level.VisibleIcon16x" : "Level.NotVisibleHighlightIcon16x"));
+							if (const UAnimNextVariableEntry* VariableEntry = Cast<UAnimNextVariableEntry>(WeakEntry))
+							{
+								if (VariableEntry->GetExportAccessSpecifier() == EAnimNextExportAccessSpecifier::Public)
+								{
+									bHasPublicVariables = true;
+								}
+								else if (VariableEntry->GetExportAccessSpecifier() == EAnimNextExportAccessSpecifier::Private)
+								{
+									bHasPrivateVariables = true;
+								}
+							}
+
+							if (bHasPrivateVariables && bHasPublicVariables)
+							{
+								break;
+							}
+						}
+
+						if (bHasPrivateVariables)
+						{
+							VariablesSection.AddMenuEntryWithCommandList(FVariablesOutlinerCommands::Get().SetVariableExportPublic, MenuContext->WeakCommandList.Pin(), LOCTEXT("MakePublicMenuLabel", "Expose as Public"), TAttribute<FText>(), FSlateIcon(FAppStyle::Get().GetStyleSetName(), "Level.VisibleIcon16x"));
 						}
-					}
 
+						if (bHasPublicVariables)
+						{
+							VariablesSection.AddMenuEntryWithCommandList(FVariablesOutlinerCommands::Get().SetVariableExportPrivate, MenuContext->WeakCommandList.Pin(), LOCTEXT("MakePrivateMenuLabel", "Make Private"), TAttribute<FText>(), FSlateIcon(FAppStyle::Get().GetStyleSetName(), "Level.NotVisibleHighlightIcon16x"));
+						}
+					}
+					
 					AssetSection.AddMenuEntryWithCommandList(FVariablesOutlinerCommands::Get().CreateSharedVariablesAssets, MenuContext->WeakCommandList.Pin(), TAttribute<FText>(), TAttribute<FText>(), FSlateIcon(FUAFStyle::Get().GetStyleSetName(), "ClassIcon.AnimNextSharedVariables"));
 				}
 
@@ -351,12 +380,20 @@ void FVariablesOutlinerMode::BindCommands(const TSharedRef<FUICommandList>& OutC
 		);
 	
 	CommandList->MapAction(
-		FVariablesOutlinerCommands::Get().ToggleVariableExport,
-		FExecuteAction::CreateRaw(this, &FVariablesOutlinerMode::ToggleVariableExport),
-		FCanExecuteAction::CreateRaw(this, &FVariablesOutlinerMode::CanToggleVariableExport),
+		FVariablesOutlinerCommands::Get().SetVariableExportPublic,
+		FExecuteAction::CreateRaw(this, &FVariablesOutlinerMode::SetVariableExportAccessSpecifiersOnSelection, EAnimNextExportAccessSpecifier::Public),
+		FCanExecuteAction::CreateRaw(this, &FVariablesOutlinerMode::CanModifyVariableExport),
 		FIsActionChecked(),
-		FIsActionButtonVisible::CreateRaw(this, &FVariablesOutlinerMode::CanToggleVariableExport)
-		);
+		FIsActionButtonVisible::CreateRaw(this, &FVariablesOutlinerMode::CanModifyVariableExport)
+	);
+
+	CommandList->MapAction(
+		FVariablesOutlinerCommands::Get().SetVariableExportPrivate,
+		FExecuteAction::CreateRaw(this, &FVariablesOutlinerMode::SetVariableExportAccessSpecifiersOnSelection, EAnimNextExportAccessSpecifier::Private),
+		FCanExecuteAction::CreateRaw(this, &FVariablesOutlinerMode::CanModifyVariableExport),
+		FIsActionChecked(),
+		FIsActionButtonVisible::CreateRaw(this, &FVariablesOutlinerMode::CanModifyVariableExport)
+	);
 
 	CommandList->MapAction(		
 		FVariablesOutlinerCommands::Get().CreateSharedVariablesAssets,
@@ -515,52 +552,71 @@ bool FVariablesOutlinerMode::CanPaste() const
 	return false;
 }
 
-void FVariablesOutlinerMode::ToggleVariableExport() const
+void FVariablesOutlinerMode::SetVariableExportAccessSpecifiersOnSelection(EAnimNextExportAccessSpecifier VariableAccessSpecifier) const
 {
+	if (!SceneOutliner)
+	{
+		return;
+	}
+	
 	const FSceneOutlinerItemSelection& Selection = SceneOutliner->GetSelection();
-	if (Selection.Num())
-	{		
-		FScopedTransaction Transaction(LOCTEXT("SetAccessSpecifierTransaction", "Set Access Specifier"));
-		Selection.ForEachItem<FVariablesOutlinerEntryItem>([](const FVariablesOutlinerEntryItem& EntryItem)
-		{			
-			if (IAnimNextRigVMExportInterface* Export = Cast<IAnimNextRigVMExportInterface>(EntryItem.WeakEntry.Get()))
+	if (Selection.Num() > 0)
+	{
+		TArray<UObject*> ObjectsToCompile;
+		TArray<IAnimNextRigVMExportInterface*> ExportsToChange;
+		Selection.ForEachItem<FVariablesOutlinerEntryItem>([&ObjectsToCompile, &ExportsToChange](const FVariablesOutlinerEntryItem& SelectedItem)
+		{
+			if (IAnimNextRigVMExportInterface* Export = Cast<IAnimNextRigVMExportInterface>(SelectedItem.WeakEntry.Get()))
 			{
-				if (UAnimNextRigVMAsset* OuterAsset = EntryItem.WeakEntry.Get()->GetTypedOuter<UAnimNextRigVMAsset>())
+				if (UAnimNextRigVMAsset* OuterAsset = SelectedItem.WeakEntry.Get()->GetTypedOuter<UAnimNextRigVMAsset>())
 				{
 					if (OuterAsset->GetClass() && OuterAsset->GetClass() != UAnimNextSharedVariables::StaticClass())
 					{
-						Export->SetExportAccessSpecifier(Export->GetExportAccessSpecifier() == EAnimNextExportAccessSpecifier::Public ? EAnimNextExportAccessSpecifier::Private : EAnimNextExportAccessSpecifier::Public);
+						ObjectsToCompile.AddUnique(OuterAsset);
+						ExportsToChange.AddUnique(Export);
 					}
 				}
 			}
 		});
+
+		UncookedOnly::FScopedCompileJob CompilerResults(LOCTEXT("ModifiedAssets_SetAccessSpecifierOnSelection", "Modified Assets Set Access Specifier on Selection"), ObjectsToCompile);
+
+		FScopedTransaction Transaction(LOCTEXT("SetAccessSpecifierSelectionTransaction", "Set Access Specifier on Selection"));
+
+		for (IAnimNextRigVMExportInterface* Export : ExportsToChange)
+		{
+			if (Export)
+			{
+				Export->SetExportAccessSpecifier(VariableAccessSpecifier);
+			}
+		}
 	}
 }
 
-bool FVariablesOutlinerMode::CanToggleVariableExport() const
+bool FVariablesOutlinerMode::CanModifyVariableExport() const
 {
 	const FSceneOutlinerItemSelection& Selection = SceneOutliner->GetSelection();
 	if (Selection.Num())
 	{
-		bool bCanToggle = false;
+		bool bCanModify = false;
 		
 		FScopedTransaction Transaction(LOCTEXT("SetAccessSpecifierTransaction", "Set Access Specifier"));
-		Selection.ForEachItem<FVariablesOutlinerEntryItem>([&bCanToggle](const FVariablesOutlinerEntryItem& EntryItem)
+		Selection.ForEachItem<FVariablesOutlinerEntryItem>([&bCanModify](const FVariablesOutlinerEntryItem& EntryItem)
 		{
 			if (UAnimNextVariableEntry* VariableEntry = Cast<UAnimNextVariableEntry>(EntryItem.WeakEntry.Get()))
 			{
 				if (UAnimNextRigVMAsset* OuterAsset = VariableEntry->GetTypedOuter<UAnimNextRigVMAsset>())
 				{
-					// Disable toggling public/private for pure UAnimNextSharedVariables objects				
+					// Disable modifying of variable export for pure UAnimNextSharedVariables objects				
 					if (OuterAsset->GetClass() && OuterAsset->GetClass() != UAnimNextSharedVariables::StaticClass())
 					{
-						bCanToggle = true;
+						bCanModify = true;
 					}
 				}
 			}
 		});
 
-		return bCanToggle;
+		return bCanModify;
 	}
 
 	return false;
@@ -941,6 +997,45 @@ void FVariablesOutlinerMode::ClearHighlightedItem(FSceneOutlinerTreeItemID InID)
 	}
 }
 
+void FVariablesOutlinerMode::SetSelectedVariablesExportAccessSpecifiers(FSceneOutlinerTreeItemPtr ItemToChange, EAnimNextExportAccessSpecifier VariableAccessSpecifier) const
+{
+	if (!SceneOutliner || !ItemToChange.IsValid())
+	{
+		return;
+	}
+
+	const FSceneOutlinerItemSelection& Selection = SceneOutliner->GetSelection();
+	const bool bEditChangedItemOnly = !Selection.SelectedItems.Contains(ItemToChange);
+
+	FScopedTransaction Transaction(LOCTEXT("SetAccessSpecifierTransaction", "Set Access Specifier"));
+
+	if (bEditChangedItemOnly)
+	{
+		if (TSharedPtr<FVariablesOutlinerEntryItem> VariablesOutlinerItem = StaticCastSharedPtr<FVariablesOutlinerEntryItem>(ItemToChange))
+		{
+			if (IAnimNextRigVMExportInterface* Export = Cast<IAnimNextRigVMExportInterface>(VariablesOutlinerItem->WeakEntry.Get()))
+			{
+				Export->SetExportAccessSpecifier(VariableAccessSpecifier);
+			}
+		}
+	}
+	else
+	{
+		SetVariableExportAccessSpecifiersOnSelection(VariableAccessSpecifier);
+	}
+}
+
+void FVariablesOutlinerMode::ToggleSelectedVariablesExportAccessSpecifiers(FSceneOutlinerTreeItemPtr ItemToChange) const
+{
+	if (TSharedPtr<FVariablesOutlinerEntryItem> VariablesOutlinerItem = StaticCastSharedPtr<FVariablesOutlinerEntryItem>(ItemToChange))
+	{
+		if (IAnimNextRigVMExportInterface* Export = Cast<IAnimNextRigVMExportInterface>(VariablesOutlinerItem->WeakEntry.Get()))
+		{
+			SetSelectedVariablesExportAccessSpecifiers(ItemToChange, Export->GetExportAccessSpecifier() == EAnimNextExportAccessSpecifier::Public ? EAnimNextExportAccessSpecifier::Private : EAnimNextExportAccessSpecifier::Public);
+		}
+	}	
+}
+
 void FVariablesOutlinerMode::SaveAsset() const
 {
 	const FSceneOutlinerItemSelection& Selection = SceneOutliner->GetSelection();
diff --git a/Engine/Plugins/Experimental/UAF/UAF/Source/UAFEditor/Private/Variables/Outliner/VariablesOutlinerMode.h b/Engine/Plugins/Experimental/UAF/UAF/Source/UAFEditor/Private/Variables/Outliner/VariablesOutlinerMode.h
--- a/Engine/Plugins/Experimental/UAF/UAF/Source/UAFEditor/Private/Variables/Outliner/VariablesOutlinerMode.h
+++ b/Engine/Plugins/Experimental/UAF/UAF/Source/UAFEditor/Private/Variables/Outliner/VariablesOutlinerMode.h
@@ -48,6 +48,12 @@ public:
 	void SetHighlightedItem(FSceneOutlinerTreeItemID InID) const;
 	void ClearHighlightedItem(FSceneOutlinerTreeItemID InID) const;
 
+	// Sets the export access specifier on ItemToChange and the current selection if ItemToChange is part of it
+	void SetSelectedVariablesExportAccessSpecifiers(FSceneOutlinerTreeItemPtr ItemToChange, EAnimNextExportAccessSpecifier VariableAccessSpecifier) const;
+
+	// Toggles the export access specifier on ItemToChange and sets the current selection to match, if ItemToChange is part of it
+	void ToggleSelectedVariablesExportAccessSpecifiers(FSceneOutlinerTreeItemPtr ItemToChange) const;
+
 	static void PopulateNewVariableToolMenuEntries(UToolMenu* InMenu, bool bAddSeparator);
 protected:
 	virtual TUniquePtr<ISceneOutlinerHierarchy> CreateHierarchy() override;
@@ -69,8 +75,8 @@ protected:
 	void Paste() const;
 	virtual bool CanPaste() const override;
 
-	void ToggleVariableExport() const;
-	bool CanToggleVariableExport() const;
+	void SetVariableExportAccessSpecifiersOnSelection(EAnimNextExportAccessSpecifier VariableAccessSpecifier) const;
+	bool CanModifyVariableExport() const;
 
 	void Duplicate() const;
 	bool CanDuplicate() const;
