Expose headers to facilitate unit testing

diff --git a/Engine/Plugins/Experimental/UAF/UAF/Source/UAFEditor/Internal/UAF/AbstractSkeleton/SetBindingFactory.h b/Engine/Plugins/Experimental/UAF/UAF/Source/UAFEditor/Internal/UAF/AbstractSkeleton/SetBindingFactory.h
new file mode 100644
--- /dev/null
+++ b/Engine/Plugins/Experimental/UAF/UAF/Source/UAFEditor/Internal/UAF/AbstractSkeleton/SetBindingFactory.h
@@ -0,0 +1,23 @@
+// Copyright Epic Games, Inc. All Rights Reserved.
+
+#pragma once
+
+#include "CoreMinimal.h"
+#include "Factories/Factory.h"
+
+#include "SetBindingFactory.generated.h"
+
+UCLASS()
+class UAbstractSkeletonSetBindingFactory : public UFactory
+{
+	GENERATED_BODY()
+
+	UAbstractSkeletonSetBindingFactory();
+
+	virtual bool ConfigureProperties() override;
+	virtual UObject* FactoryCreateNew(UClass* Class, UObject* InParent, FName Name, EObjectFlags Flags, UObject* Context, FFeedbackContext* Warn, FName CallingContext) override;
+	virtual UObject* FactoryCreateNew(UClass* Class, UObject* InParent, FName Name, EObjectFlags Flags, UObject* Context, FFeedbackContext* Warn) override
+	{
+		return FactoryCreateNew(Class, InParent, Name, Flags, Context, Warn, NAME_None);
+	}
+};
diff --git a/Engine/Plugins/Experimental/UAF/UAF/Source/UAFEditor/Internal/UAF/AbstractSkeleton/SetCollectionFactory.h b/Engine/Plugins/Experimental/UAF/UAF/Source/UAFEditor/Internal/UAF/AbstractSkeleton/SetCollectionFactory.h
new file mode 100644
--- /dev/null
+++ b/Engine/Plugins/Experimental/UAF/UAF/Source/UAFEditor/Internal/UAF/AbstractSkeleton/SetCollectionFactory.h
@@ -0,0 +1,23 @@
+// Copyright Epic Games, Inc. All Rights Reserved.
+
+#pragma once
+
+#include "CoreMinimal.h"
+#include "Factories/Factory.h"
+
+#include "SetCollectionFactory.generated.h"
+
+UCLASS()
+class UAbstractSkeletonSetCollectionFactory : public UFactory
+{
+	GENERATED_BODY()
+
+	UAbstractSkeletonSetCollectionFactory();
+
+	virtual bool ConfigureProperties() override;
+	virtual UObject* FactoryCreateNew(UClass* Class, UObject* InParent, FName Name, EObjectFlags Flags, UObject* Context, FFeedbackContext* Warn, FName CallingContext) override;
+	virtual UObject* FactoryCreateNew(UClass* Class, UObject* InParent, FName Name, EObjectFlags Flags, UObject* Context, FFeedbackContext* Warn) override
+	{
+		return FactoryCreateNew(Class, InParent, Name, Flags, Context, Warn, NAME_None);
+	}
+};
diff --git a/Engine/Plugins/Experimental/UAF/UAF/Source/UAFEditor/Private/UAF/AbstractSkeleton/Sets/SetBindingFactory.cpp b/Engine/Plugins/Experimental/UAF/UAF/Source/UAFEditor/Private/UAF/AbstractSkeleton/Sets/SetBindingFactory.cpp
--- a/Engine/Plugins/Experimental/UAF/UAF/Source/UAFEditor/Private/UAF/AbstractSkeleton/Sets/SetBindingFactory.cpp
+++ b/Engine/Plugins/Experimental/UAF/UAF/Source/UAFEditor/Private/UAF/AbstractSkeleton/Sets/SetBindingFactory.cpp
@@ -1,6 +1,6 @@
 // Copyright Epic Games, Inc. All Rights Reserved.
 
-#include "UAF/AbstractSkeleton/Sets/SetBindingFactory.h"
+#include "UAF/AbstractSkeleton/SetBindingFactory.h"
 #include "UAF/AbstractSkeleton/AbstractSkeletonSetBinding.h"
 #include "UObject/Package.h"
 
diff --git a/Engine/Plugins/Experimental/UAF/UAF/Source/UAFEditor/Private/UAF/AbstractSkeleton/Sets/SetBindingFactory.h b/Engine/Plugins/Experimental/UAF/UAF/Source/UAFEditor/Private/UAF/AbstractSkeleton/Sets/SetBindingFactory.h
deleted file mode 100644
diff --git a/Engine/Plugins/Experimental/UAF/UAF/Source/UAFEditor/Private/UAF/AbstractSkeleton/Sets/SetCollectionFactory.cpp b/Engine/Plugins/Experimental/UAF/UAF/Source/UAFEditor/Private/UAF/AbstractSkeleton/Sets/SetCollectionFactory.cpp
--- a/Engine/Plugins/Experimental/UAF/UAF/Source/UAFEditor/Private/UAF/AbstractSkeleton/Sets/SetCollectionFactory.cpp
+++ b/Engine/Plugins/Experimental/UAF/UAF/Source/UAFEditor/Private/UAF/AbstractSkeleton/Sets/SetCollectionFactory.cpp
@@ -1,6 +1,6 @@
 // Copyright Epic Games, Inc. All Rights Reserved.
 
-#include "UAF/AbstractSkeleton/Sets/SetCollectionFactory.h"
+#include "UAF/AbstractSkeleton/SetCollectionFactory.h"
 
 #include "UAF/AbstractSkeleton/AbstractSkeletonSetCollection.h"
 #include "UObject/Package.h"
diff --git a/Engine/Plugins/Experimental/UAF/UAF/Source/UAFEditor/Private/UAF/AbstractSkeleton/Sets/SetCollectionFactory.h b/Engine/Plugins/Experimental/UAF/UAF/Source/UAFEditor/Private/UAF/AbstractSkeleton/Sets/SetCollectionFactory.h
deleted file mode 100644
