Add GetTypeHash support for FConstExprUID

diff --git a/Engine/Plugins/Experimental/UAF/UAF/Source/UAF/Public/ConstExprUID.h b/Engine/Plugins/Experimental/UAF/UAF/Source/UAF/Public/ConstExprUID.h
--- a/Engine/Plugins/Experimental/UAF/UAF/Source/UAF/Public/ConstExprUID.h
+++ b/Engine/Plugins/Experimental/UAF/UAF/Source/UAF/Public/ConstExprUID.h
@@ -31,7 +31,7 @@ namespace UE::UAF
 		}
 
 		// Constructs a UID
-		explicit constexpr FConstExprUID(FFConstExprUIDRaw InUID, const TCHAR* InDebugName = TEXT("<Unknown ConstExprUID Name>")) noexcept
+		[[nodiscard]] explicit constexpr FConstExprUID(FFConstExprUIDRaw InUID, const TCHAR* InDebugName = TEXT("<Unknown ConstExprUID Name>")) noexcept
 			: UID(InUID)
 #if !UE_BUILD_SHIPPING
 			, DebugName(InDebugName)
@@ -40,21 +40,28 @@ namespace UE::UAF
 		}
 
 		// Constructs a UID
-		static constexpr FConstExprUID MakeFromString(const TCHAR* InName) noexcept
+		[[nodiscard]] static constexpr FConstExprUID MakeFromString(const TCHAR* InName) noexcept
 		{
 			return FConstExprUID(UE::HashStringFNV1a32(MakeStringView(InName)), InName);
 		}
 
 #if !UE_BUILD_SHIPPING
 		// Returns a literal string to the debug name
-		constexpr const TCHAR* GetDebugName() const noexcept { return DebugName; }
+		[[nodiscard]] constexpr const TCHAR* GetDebugName() const noexcept { return DebugName; }
 #endif
 
 		// Returns the raw UID
-		constexpr FFConstExprUIDRaw GetUID() const noexcept { return UID; }
+		[[nodiscard]] constexpr FFConstExprUIDRaw GetUID() const noexcept { return UID; }
 
 		// Returns whether this UID is valid or not
-		constexpr bool IsValid() const noexcept { return UID != INVALID_UID; }
+		[[nodiscard]] constexpr bool IsValid() const noexcept { return UID != INVALID_UID; }
+
+		// Returns the hash of this UID
+		[[nodiscard]] constexpr friend uint32 GetTypeHash(FConstExprUID InUID)
+		{
+			// Return the UID as-is since it is already a hashed value
+			return InUID.UID;
+		}
 
 	private:
 		static constexpr FFConstExprUIDRaw INVALID_UID = 0;
