Update UAF variables tracing to include shared variables and overrides
* Also fix UAF trace provider not enumerating all variable data

diff --git a/Engine/Plugins/Experimental/UAF/UAF/Source/UAF/Private/RewindDebugger/AnimNextTrace.cpp b/Engine/Plugins/Experimental/UAF/UAF/Source/UAF/Private/RewindDebugger/AnimNextTrace.cpp
--- a/Engine/Plugins/Experimental/UAF/UAF/Source/UAF/Private/RewindDebugger/AnimNextTrace.cpp
+++ b/Engine/Plugins/Experimental/UAF/UAF/Source/UAF/Private/RewindDebugger/AnimNextTrace.cpp
@@ -10,6 +10,8 @@
 #include "Serialization/MemoryWriter.h"
 #include "ObjectAsTraceIdProxyArchive.h"
 #include "Serialization/ObjectWriter.h"
+#include "Variables/VariableOverrides.h"
+#include "Variables/StructDataCache.h"
 
 #if ANIMNEXT_TRACE_ENABLED
 
@@ -161,11 +163,11 @@ void FAnimNextTrace::OutputAnimNextVariables(const FUAFAssetInstance* Instance,
 	OutputAnimNextInstance(Instance, OuterObject);
 
 	bool bVariablesOutput = false;
-	if (Instance->Variables.OwnedVariableContainers.Num() > 0)
+	if (Instance->Variables.AllSharedVariableContainers.Num() > 0)
 	{
-		for (const TSharedRef<FUAFInstanceVariableContainer>& OwnedVariableSet : Instance->Variables.OwnedVariableContainers)
+		for (const TWeakPtr<FUAFInstanceVariableContainer>& OwnedVariableSet : Instance->Variables.AllSharedVariableContainers)
 		{
-			if (OutputAnimNextVariableSet(OwnedVariableSet, InstanceId, OuterObject))
+			if (OutputAnimNextVariableSet(OwnedVariableSet.Pin().ToSharedRef(), InstanceId, OuterObject))
 			{
 				bVariablesOutput = true;
 			}
@@ -189,9 +191,8 @@ bool FAnimNextTrace::OutputAnimNextVariableSet(const TSharedRef<FUAFInstanceVari
 	{
 	case FUAFInstanceVariableContainer::FVariablesContainerType::IndexOfType<FInstancedPropertyBag>():
 		{
-			// TODO: Accomodate overrides
-
-			FInstancedPropertyBag& InstancedPropertyBag = InVariableSet->VariablesContainer.Get<FInstancedPropertyBag>();
+			// Copy property bag data as it will be modified
+			FInstancedPropertyBag InstancedPropertyBag = InVariableSet->VariablesContainer.Get<FInstancedPropertyBag>();
 			if (InstancedPropertyBag.GetNumPropertiesInBag() == 0)
 			{
 				break;
@@ -237,6 +238,23 @@ bool FAnimNextTrace::OutputAnimNextVariableSet(const TSharedRef<FUAFInstanceVari
 				UPropertyBag* PropertyBag = const_cast<UPropertyBag*>(InstancedPropertyBag.GetPropertyBagStruct());
 				if (PropertyBag != nullptr)
 				{
+					// Apply overrides
+					if (!InVariableSet->ResolvedOverrides.IsEmpty())
+					{
+						const int32 VariablesNum = PropertyDescs.Num();
+						for (int32 VariableIndex = 0; VariableIndex < VariablesNum; ++VariableIndex)
+						{
+							if (InVariableSet->ResolvedOverrides[VariableIndex] != nullptr)
+							{
+								const FPropertyBagPropertyDesc& Desc = PropertyDescs[VariableIndex];
+								uint8* Target = Desc.CachedProperty->ContainerPtrToValuePtr<uint8>(InstancedPropertyBag.GetMutableValue().GetMemory());
+								uint8* Source = InVariableSet->ResolvedOverrides[VariableIndex];
+								
+								Desc.CachedProperty->CopyCompleteValue(Target, Source);
+							}
+						}
+					}
+				
 					PropertyBag->SerializeItem(Archive, InstancedPropertyBag.GetMutableValue().GetMemory(), nullptr);
 				}
 
@@ -258,7 +276,23 @@ bool FAnimNextTrace::OutputAnimNextVariableSet(const TSharedRef<FUAFInstanceVari
 			FMemoryWriter WriterArchive(ArchiveData);
 			FObjectAsTraceIdProxyArchive Archive(WriterArchive);
 
-			FInstancedStruct& InstancedStruct = InVariableSet->VariablesContainer.Get<FInstancedStruct>();
+			// Copy property bag data as it will be modified
+			FInstancedStruct InstancedStruct = InVariableSet->VariablesContainer.Get<FInstancedStruct>();
+
+			// Apply overrides
+			if (!InVariableSet->ResolvedOverrides.IsEmpty())
+			{
+				for (int32 VariableIndex = 0; VariableIndex < InVariableSet->NumVariables; ++VariableIndex)
+				{
+					const FStructDataCache::FPropertyInfo& PropertyInfo = InVariableSet->StructDataCache->GetProperties()[VariableIndex];
+					const FProperty* Property = PropertyInfo.Property;
+					
+					uint8* Target = Property->ContainerPtrToValuePtr<uint8>(InstancedStruct.GetMutableMemory());
+                    uint8* Source = InVariableSet->ResolvedOverrides[VariableIndex];
+					
+					Property->CopyCompleteValue(Target, Source);
+				}
+			}
 			
 			if (InstancedStruct.IsValid())
 			{
diff --git a/Engine/Plugins/Experimental/UAF/UAF/Source/UAFEditor/Private/RewindDebugger/AnimNextModuleTrack.cpp b/Engine/Plugins/Experimental/UAF/UAF/Source/UAFEditor/Private/RewindDebugger/AnimNextModuleTrack.cpp
--- a/Engine/Plugins/Experimental/UAF/UAF/Source/UAFEditor/Private/RewindDebugger/AnimNextModuleTrack.cpp
+++ b/Engine/Plugins/Experimental/UAF/UAF/Source/UAFEditor/Private/RewindDebugger/AnimNextModuleTrack.cpp
@@ -197,7 +197,7 @@ bool FAnimNextModuleTrack::UpdateInternal()
 								FInstancedStruct::StaticStruct()->SerializeItem(ReaderProxy, &InstancedStruct, nullptr);
 							}
 
-							return TraceServices::EEventEnumerate::Stop;
+							return TraceServices::EEventEnumerate::Continue;
 						});
 					}
 				}
