Update UAF Instance tracks in rewind debugger to be named by the FactoryObject

diff --git a/Engine/Plugins/Experimental/UAF/UAF/Source/UAF/Private/RewindDebugger/AnimNextTrace.cpp b/Engine/Plugins/Experimental/UAF/UAF/Source/UAF/Private/RewindDebugger/AnimNextTrace.cpp
--- a/Engine/Plugins/Experimental/UAF/UAF/Source/UAF/Private/RewindDebugger/AnimNextTrace.cpp
+++ b/Engine/Plugins/Experimental/UAF/UAF/Source/UAF/Private/RewindDebugger/AnimNextTrace.cpp
@@ -119,7 +119,7 @@ void FAnimNextTrace::OutputAnimNextInstance(const FUAFAssetInstance* AnimNextIns
 
 		TRACE_OBJECT(OuterObject);
 		TRACE_OBJECT(AnimNextInstance->GetAsset<UAnimNextRigVMAsset>());
-		TRACE_INSTANCE(OuterObject, AnimNextInstance->GetUniqueId(), HostInstanceId, FUAFAssetInstance::StaticStruct(), AnimNextInstance->GetAssetName().ToString());
+		TRACE_INSTANCE(OuterObject, AnimNextInstance->GetUniqueId(), HostInstanceId, FUAFAssetInstance::StaticStruct(), AnimNextInstance->GetDebugName().ToString());
 
 		UE_TRACE_LOG(UAF, Instance, UAFChannel)
 			<< Instance.InstanceId(InstanceId)
diff --git a/Engine/Plugins/Experimental/UAF/UAF/Source/UAF/Public/UAFAssetInstance.h b/Engine/Plugins/Experimental/UAF/UAF/Source/UAF/Public/UAFAssetInstance.h
--- a/Engine/Plugins/Experimental/UAF/UAF/Source/UAF/Public/UAFAssetInstance.h
+++ b/Engine/Plugins/Experimental/UAF/UAF/Source/UAF/Public/UAFAssetInstance.h
@@ -292,7 +292,26 @@ protected:
 	friend UE::UAF::FInstanceTaskContext;
 	friend FUAFRigVMComponent;
 
+	void SetDebugName(FName Name)
+	{
+	#if ANIMNEXT_TRACE_ENABLED
+		DebugName = Name;
+	#endif
+	}
+	
+	FName GetDebugName() const
+	{
+#if ANIMNEXT_TRACE_ENABLED
+		if (!DebugName.IsNone())
+		{
+			return DebugName;
+		}
+#endif
+		return GetAssetName();
+	}
+
 #if ANIMNEXT_TRACE_ENABLED
+	FName DebugName;
 	uint64 UniqueId;
 	UE_API volatile static int64 NextUniqueId;
 #endif 
diff --git a/Engine/Plugins/Experimental/UAF/UAF/Source/UAFEditor/Private/RewindDebugger/AnimNextModuleTrack.cpp b/Engine/Plugins/Experimental/UAF/UAF/Source/UAFEditor/Private/RewindDebugger/AnimNextModuleTrack.cpp
--- a/Engine/Plugins/Experimental/UAF/UAF/Source/UAFEditor/Private/RewindDebugger/AnimNextModuleTrack.cpp
+++ b/Engine/Plugins/Experimental/UAF/UAF/Source/UAFEditor/Private/RewindDebugger/AnimNextModuleTrack.cpp
@@ -41,12 +41,9 @@ FText FAnimNextModuleTrack::GetDisplayNameInternal() const
 
 		if (InstanceId != 0)
 		{
-			if (const FDataInterfaceData* Data = AnimNextProvider->GetDataInterfaceData(InstanceId))
-			{
-				const FObjectInfo& ModuleInfo = GameplayProvider->GetObjectInfo(Data->AssetId);
-				DisplayNameCache = FText::FromString(ModuleInfo.Name);
-				return DisplayNameCache;
-			}
+			const FObjectInfo& ModuleInfo = GameplayProvider->GetObjectInfo(InstanceId);
+			DisplayNameCache = FText::FromString(ModuleInfo.Name);
+			return DisplayNameCache;
 		}
 
 		// Don't cache this since it is a placeholder name, will be replaced once the real one is sent
diff --git a/Engine/Plugins/Experimental/UAF/UAFAnimGraph/Source/UAFAnimGraph/Private/Traits/BlendStackTrait.cpp b/Engine/Plugins/Experimental/UAF/UAFAnimGraph/Source/UAFAnimGraph/Private/Traits/BlendStackTrait.cpp
--- a/Engine/Plugins/Experimental/UAF/UAFAnimGraph/Source/UAFAnimGraph/Private/Traits/BlendStackTrait.cpp
+++ b/Engine/Plugins/Experimental/UAF/UAFAnimGraph/Source/UAFAnimGraph/Private/Traits/BlendStackTrait.cpp
@@ -216,7 +216,12 @@ namespace UE::UAF
 							.ParentContext = &Context,
 							.ParentGraphInstance = &Owner,
 						});
+					if (Graph.Request.FactoryObject)
+					{
+						Graph.Instance->SetDebugName(Graph.Request.FactoryObject->GetFName());
+					}
 					Graph.Request.FactoryParams.InitializeInstance(*Graph.Instance.Get());
+
 					break;
 				}
 			case EGraphRequestType::Child:
