Add all public functions to export by default for a RigVM Asset
- Fixes public functions not showing up in any UAF search context (Graph Context menu, call function, UAF ST Task / etc).

#tests editor

diff --git a/Engine/Plugins/Experimental/UAF/UAF/Source/UAFUncookedOnly/Internal/AnimNextRigVMAssetEditorData.h b/Engine/Plugins/Experimental/UAF/UAF/Source/UAFUncookedOnly/Internal/AnimNextRigVMAssetEditorData.h
--- a/Engine/Plugins/Experimental/UAF/UAF/Source/UAFUncookedOnly/Internal/AnimNextRigVMAssetEditorData.h
+++ b/Engine/Plugins/Experimental/UAF/UAF/Source/UAFUncookedOnly/Internal/AnimNextRigVMAssetEditorData.h
@@ -609,11 +609,11 @@ protected:
 	// Check whether this asset should be recompiled
 	UE_API bool IsDirtyForRecompilation() const;
 
-	// Called pre programmatic function header population to allow population of compile context in a const manner. May be used by external assets querying our VM.
-	virtual void BuildFunctionHeadersContext(const FRigVMCompileSettings& InSettings, FAnimNextGetFunctionHeaderCompileContext& OutCompileContext) const {}
+	// Called pre programmatic function header population to allow population of compile context in a const manner. May be used by external assets querying our VM. 
+	UE_API virtual void BuildFunctionHeadersContext(const FRigVMCompileSettings& InSettings, FAnimNextGetFunctionHeaderCompileContext& OutCompileContext) const;
 
 	// Called pre programmatic variable population to allow population of compile context in a const manner. May be used by external assets querying our VM.
-	virtual void BuildProgrammaticVariablesContext(const FRigVMCompileSettings& InSettings, FAnimNextGetVariableCompileContext& OutCompileContext) const {}
+	UE_API virtual void BuildProgrammaticVariablesContext(const FRigVMCompileSettings& InSettings, FAnimNextGetVariableCompileContext& OutCompileContext) const {}
 
 	// Called during compilation to generate internal variables used for natively-callable RigVM functions
 	UE_API void BuildFunctionWrapperEventVariables(FAnimNextRigVMAssetCompileContext& InContext) const;
diff --git a/Engine/Plugins/Experimental/UAF/UAF/Source/UAFUncookedOnly/Private/AnimNextRigVMAssetEditorData.cpp b/Engine/Plugins/Experimental/UAF/UAF/Source/UAFUncookedOnly/Private/AnimNextRigVMAssetEditorData.cpp
--- a/Engine/Plugins/Experimental/UAF/UAF/Source/UAFUncookedOnly/Private/AnimNextRigVMAssetEditorData.cpp
+++ b/Engine/Plugins/Experimental/UAF/UAF/Source/UAFUncookedOnly/Private/AnimNextRigVMAssetEditorData.cpp
@@ -807,6 +807,20 @@ TObjectPtr<URigVMGraph> UAnimNextRigVMAssetEditorData::CreateContainedGraphModel
 	return Model;
 }
 
+void UAnimNextRigVMAssetEditorData::BuildFunctionHeadersContext(const FRigVMCompileSettings& InSettings, FAnimNextGetFunctionHeaderCompileContext& OutCompileContext) const
+{
+	// By default gather & generate event stubs for all public functions so they can be called externally.
+	for (const FRigVMGraphFunctionData& FunctionData : GraphFunctionStore.PublicFunctions)
+	{
+		if (FunctionData.Header.IsValid())
+		{
+			FAnimNextProgrammaticFunctionHeader AnimNextFunctionHeader = {};
+			AnimNextFunctionHeader.FunctionHeader = FunctionData.Header;
+			OutCompileContext.AddUniqueFunctionHeader(AnimNextFunctionHeader);
+		}
+	}
+}
+
 void UAnimNextRigVMAssetEditorData::BuildFunctionWrapperEventVariables(FAnimNextRigVMAssetCompileContext& InContext) const
 {
 	using namespace UE::UAF::UncookedOnly;
diff --git a/Engine/Plugins/Experimental/UAF/UAF/Source/UAFUncookedOnly/Private/UncookedOnlyUtils.cpp b/Engine/Plugins/Experimental/UAF/UAF/Source/UAFUncookedOnly/Private/UncookedOnlyUtils.cpp
--- a/Engine/Plugins/Experimental/UAF/UAF/Source/UAFUncookedOnly/Private/UncookedOnlyUtils.cpp
+++ b/Engine/Plugins/Experimental/UAF/UAF/Source/UAFUncookedOnly/Private/UncookedOnlyUtils.cpp
@@ -848,6 +848,7 @@ void FUtils::GetAssetFunctions(const UAnimNextRigVMAssetEditorData* InEditorData
 {
 	for (const FRigVMGraphFunctionData& FunctionData : InEditorData->GraphFunctionStore.PublicFunctions)
 	{
+		// Note: By default all public functions should get added for compilation in `UAnimNextRigVMAssetEditorData::BuildFunctionHeadersContext`
 		if (FunctionData.CompilationData.IsValid())
 		{
 			OutExports.Headers.Add(FunctionData.Header);
