Added common outliner implementation to be shared between variables/functions sharing:
- Hierarchy generation from assets
- Category/asset/entry items
- Common functionality (copy/paste/rename etc.)
- Drag-drop operations (registered with OutlinerMode rather than through item widgets)

Variables Outliner:
- Copy/cut/paste behavior
- Categories are now allowed to be empty, so now also able to them manually

Functions Outliner:
- Display and allows for editing RigVM functions
- Hierarchy matches Variables Outliner
- Currently limited to rename, duplicate, delete operations (rest to follow)
- Categories for functions are fixed, as they are tied to the Node objects themselves

Misc:
- Only show user generated Collapse node in Workspace outliner

diff --git a/Engine/Plugins/Experimental/UAF/UAF/Source/UAFEditor/Internal/Common/GraphEditorSchemaActions.h b/Engine/Plugins/Experimental/UAF/UAF/Source/UAFEditor/Internal/Common/GraphEditorSchemaActions.h
--- a/Engine/Plugins/Experimental/UAF/UAF/Source/UAFEditor/Internal/Common/GraphEditorSchemaActions.h
+++ b/Engine/Plugins/Experimental/UAF/UAF/Source/UAFEditor/Internal/Common/GraphEditorSchemaActions.h
@@ -136,6 +136,8 @@ struct FAnimNextSchemaAction_Function : public FAnimNextSchemaAction
 
 	FAnimNextSchemaAction_Function(const FRigVMGraphFunctionHeader& InReferencedPublicFunctionHeader, const FText& InNodeCategory, const FText& InMenuDesc, const FText& InToolTip, const FText& InKeywords = FText::GetEmpty());
 	FAnimNextSchemaAction_Function(const URigVMLibraryNode* InFunctionLibraryNode, const FText& InNodeCategory, const FText& InMenuDesc, const FText& InToolTip, const FText& InKeywords = FText::GetEmpty());
+	
+	FAnimNextSchemaAction_Function(const URigVMLibraryNode* InFunctionLibraryNode);
 
 	virtual const FSlateBrush* GetIconBrush() const override;
 	
diff --git a/Engine/Plugins/Experimental/UAF/UAF/Source/UAFEditor/Private/AnimNextEditorModule.cpp b/Engine/Plugins/Experimental/UAF/UAF/Source/UAFEditor/Private/AnimNextEditorModule.cpp
--- a/Engine/Plugins/Experimental/UAF/UAF/Source/UAFEditor/Private/AnimNextEditorModule.cpp
+++ b/Engine/Plugins/Experimental/UAF/UAF/Source/UAFEditor/Private/AnimNextEditorModule.cpp
@@ -73,6 +73,8 @@
 #include "AnimNextSharedVariableNode.h"
 #include "UAFStyle.h"
 #include "Common/AnimNextAssetFindReplaceVariables.h"
+#include "Common/Outliner/OutlinerCommands.h"
+#include "Functions/Outliner/FunctionsOutlinerView.h"
 #include "Toolkits/AssetEditorToolkitMenuContext.h"
 #include "UAF/Viewport/SystemViewportController.h"
 #include "Variables/Outliner/VariablesOutlinerCommands.h"
@@ -127,6 +129,7 @@ void FAnimNextEditorModule::StartupModule()
 	FRigVMCommands::Register();
 	FAnimNextRigVMAssetCommands::Register();
 	FVariablesOutlinerCommands::Register();
+	FOutlinerCommands::Register();
 	FUAFStyle::Get();
 
 	// Register settings for user editing
@@ -206,6 +209,10 @@ void FAnimNextEditorModule::StartupModule()
 			TabFactories.RegisterFactory(VariablesTabSummoner);
 			VariablesTabSummoner->RegisterTabSpawner(InTabManager, nullptr);
 
+			TSharedRef<FAnimNextFunctionsTabSummoner> FunctionsTabSummoner = MakeShared<FAnimNextFunctionsTabSummoner>(InEditorPtr);
+			TabFactories.RegisterFactory(FunctionsTabSummoner);
+			FunctionsTabSummoner->RegisterTabSpawner(InTabManager, nullptr);
+
 			TSharedRef<FAnimNextFindAndReplaceTabSummoner> FindAndReplaceTabSummoner = MakeShared<FAnimNextFindAndReplaceTabSummoner>(InEditorPtr);
 			TabFactories.RegisterFactory(FindAndReplaceTabSummoner);
 			FindAndReplaceTabSummoner->RegisterTabSpawner(InTabManager, nullptr);
@@ -223,7 +230,10 @@ void FAnimNextEditorModule::StartupModule()
 		InLayoutExtender.ExtendLayout(FTabId(Workspace::WorkspaceTabs::BottomMiddleDocumentArea), ELayoutExtensionPosition::After, FindAndReplaceTab);
 
 		FTabManager::FTab VariablesTab(FTabId(VariablesTabName), ETabState::OpenedTab);
-		InLayoutExtender.ExtendLayout(FTabId(Workspace::WorkspaceTabs::BottomLeftDocumentArea), ELayoutExtensionPosition::After, VariablesTab);
+		InLayoutExtender.ExtendStack(Workspace::WorkspaceTabs::BottomLeftDocumentArea, ELayoutExtensionPosition::After, VariablesTab, true);
+
+		FTabManager::FTab FunctionsTab(FTabId(FunctionsTabName), ETabState::OpenedTab);
+		InLayoutExtender.ExtendStack(Workspace::WorkspaceTabs::BottomLeftDocumentArea, ELayoutExtensionPosition::After, FunctionsTab);
 	});
 
 	UToolMenu* AnimNextEdGraphContextMenu = UToolMenus::Get()->ExtendMenu( "GraphEditor.GraphNodeContextMenu.AnimNextEdGraphNode");
@@ -302,11 +312,11 @@ void FAnimNextEditorModule::StartupModule()
 											}
 										};
 
-										InSubmenu->AddMenuEntry(NAME_None, FToolMenuEntry::InitMenuEntry("FindGlobalReferences", FVariablesOutlinerCommands::Get().FindReferences->GetLabel(), FVariablesOutlinerCommands::Get().FindReferences->GetDescription(), UAnimNextAssetFindReplaceVariables::GetIconFromSearchScope(ESearchScope::Global), FToolUIAction(FToolMenuExecuteAction::CreateLambda(InvokeFindVariables, ESearchScope::Global))));
+										InSubmenu->AddMenuEntry(NAME_None, FToolMenuEntry::InitMenuEntry("FindGlobalReferences", FOutlinerCommands::Get().FindReferences->GetLabel(), FOutlinerCommands::Get().FindReferences->GetDescription(), UAnimNextAssetFindReplaceVariables::GetIconFromSearchScope(ESearchScope::Global), FToolUIAction(FToolMenuExecuteAction::CreateLambda(InvokeFindVariables, ESearchScope::Global))));
 
-										InSubmenu->AddMenuEntry(NAME_None, FToolMenuEntry::InitMenuEntry("FindWorkspaceReferences", FVariablesOutlinerCommands::Get().FindReferencesInWorkspace->GetLabel(), FVariablesOutlinerCommands::Get().FindReferencesInWorkspace->GetDescription(), UAnimNextAssetFindReplaceVariables::GetIconFromSearchScope(ESearchScope::Workspace), FToolUIAction(FToolMenuExecuteAction::CreateLambda(InvokeFindVariables, ESearchScope::Workspace))));
+										InSubmenu->AddMenuEntry(NAME_None, FToolMenuEntry::InitMenuEntry("FindWorkspaceReferences", FOutlinerCommands::Get().FindReferencesInWorkspace->GetLabel(), FOutlinerCommands::Get().FindReferencesInWorkspace->GetDescription(), UAnimNextAssetFindReplaceVariables::GetIconFromSearchScope(ESearchScope::Workspace), FToolUIAction(FToolMenuExecuteAction::CreateLambda(InvokeFindVariables, ESearchScope::Workspace))));
 
-										InSubmenu->AddMenuEntry(NAME_None, FToolMenuEntry::InitMenuEntry("FindAssetReferences", FVariablesOutlinerCommands::Get().FindReferencesInAsset->GetLabel(), FVariablesOutlinerCommands::Get().FindReferencesInAsset->GetDescription(), UAnimNextAssetFindReplaceVariables::GetIconFromSearchScope(ESearchScope::Asset), FToolUIAction(FToolMenuExecuteAction::CreateLambda(InvokeFindVariables, ESearchScope::Asset))));
+										InSubmenu->AddMenuEntry(NAME_None, FToolMenuEntry::InitMenuEntry("FindAssetReferences", FOutlinerCommands::Get().FindReferencesInAsset->GetLabel(), FOutlinerCommands::Get().FindReferencesInAsset->GetDescription(), UAnimNextAssetFindReplaceVariables::GetIconFromSearchScope(ESearchScope::Asset), FToolUIAction(FToolMenuExecuteAction::CreateLambda(InvokeFindVariables, ESearchScope::Asset))));
 									}),
 									FToolUIActionChoice(),
 									EUserInterfaceActionType::Button,
diff --git a/Engine/Plugins/Experimental/UAF/UAF/Source/UAFEditor/Private/Common/AnimNextFunctionItemDetails.cpp b/Engine/Plugins/Experimental/UAF/UAF/Source/UAFEditor/Private/Common/AnimNextFunctionItemDetails.cpp
--- a/Engine/Plugins/Experimental/UAF/UAF/Source/UAFEditor/Private/Common/AnimNextFunctionItemDetails.cpp
+++ b/Engine/Plugins/Experimental/UAF/UAF/Source/UAFEditor/Private/Common/AnimNextFunctionItemDetails.cpp
@@ -82,7 +82,6 @@ bool FAnimNextFunctionItemDetails::CanDelete(const FWorkspaceOutlinerItemExport&
 
 void FAnimNextFunctionItemDetails::Delete(TConstArrayView<FWorkspaceOutlinerItemExport> Exports) const
 {
-	TMap<UAnimNextRigVMAssetEditorData*, TArray<UAnimNextRigVMAssetEntry*>> EntriesToDelete;
 	for (const FWorkspaceOutlinerItemExport& Export : Exports)
 	{
 		const TInstancedStruct<FWorkspaceOutlinerItemData>& Data = Export.GetData();
@@ -170,10 +169,7 @@ void FAnimNextFunctionItemDetails::Rename(const FWorkspaceOutlinerItemExport& Ex
 
 					FScopedTransaction Transaction(LOCTEXT("RenameFunctionInOutliner", "Rename Function"));
 					FRigVMControllerCompileBracketScope CompileScope(EdGraph->GetController());
-					if (GraphSchema->TryRenameGraph(EdGraph, *InName.ToString()))
-					{
-						return;
-					}
+					GraphSchema->TryRenameGraph(EdGraph, *InName.ToString());
 				}
 			}
 		}
diff --git a/Engine/Plugins/Experimental/UAF/UAF/Source/UAFEditor/Private/Common/GraphEditorSchemaActions.cpp b/Engine/Plugins/Experimental/UAF/UAF/Source/UAFEditor/Private/Common/GraphEditorSchemaActions.cpp
--- a/Engine/Plugins/Experimental/UAF/UAF/Source/UAFEditor/Private/Common/GraphEditorSchemaActions.cpp
+++ b/Engine/Plugins/Experimental/UAF/UAF/Source/UAFEditor/Private/Common/GraphEditorSchemaActions.cpp
@@ -422,6 +422,17 @@ FAnimNextSchemaAction_Function::FAnimNextSchemaAction_Function(const URigVMLibra
 	bIsLocalFunction = true;
 }
 
+FAnimNextSchemaAction_Function::FAnimNextSchemaAction_Function(const URigVMLibraryNode* InFunctionLibraryNode)
+	: FAnimNextSchemaAction(InFunctionLibraryNode ? FText::FromString(InFunctionLibraryNode->GetNodeCategory()) : FText::GetEmpty(), InFunctionLibraryNode ? FText::FromString(InFunctionLibraryNode->GetNodeDescription()) : FText::GetEmpty(), InFunctionLibraryNode ? InFunctionLibraryNode->GetToolTipText() : FText::GetEmpty(), InFunctionLibraryNode ? FText::FromString(InFunctionLibraryNode->GetNodeKeywords()) : FText::GetEmpty())
+{
+	if (InFunctionLibraryNode)
+	{
+		ReferencedPublicFunctionHeader = InFunctionLibraryNode->GetFunctionHeader();
+	}
+	NodeClass = UAnimNextEdGraphNode::StaticClass();
+	bIsLocalFunction = true;
+}
+
 const FSlateBrush* FAnimNextSchemaAction_Function::GetIconBrush() const
 {
 	return FAppStyle::GetBrush("GraphEditor.Function_16x");
diff --git a/Engine/Plugins/Experimental/UAF/UAF/Source/UAFEditor/Private/Common/Outliner/OutlinerHierarchy.cpp b/Engine/Plugins/Experimental/UAF/UAF/Source/UAFEditor/Private/Common/Outliner/OutlinerHierarchy.cpp
new file mode 100644
--- /dev/null
+++ b/Engine/Plugins/Experimental/UAF/UAF/Source/UAFEditor/Private/Common/Outliner/OutlinerHierarchy.cpp
@@ -0,0 +1,125 @@
+// Copyright Epic Games, Inc. All Rights Reserved.
+
+#include "OutlinerHierarchy.h"
+
+#include "AnimNextRigVMAsset.h"
+#include "ISceneOutlinerMode.h"
+#include "AnimNextRigVMAssetEditorData.h"
+
+#include "OutlinerAssetItem.h"
+#include "CommonOutlinerMode.h"
+
+#include "Entries/AnimNextSharedVariablesEntry.h"
+#include "Variables/SVariablesView.h"
+#include "Variables/AnimNextSharedVariables.h"
+
+#include "AnimNextRigVMAsset.h"
+
+namespace UE::UAF::Editor
+{
+
+TAutoConsoleVariable<bool> GOutputHierarchyDebugInformationCvar(TEXT("Outliner.OutputHierarchyDebugInformation"), false, TEXT("If enabled will output debug information about the hierarchy generation to the output log")); 
+
+FOutlinerHierarchy::FOutlinerHierarchy(ISceneOutlinerMode* Mode) : ISceneOutlinerHierarchy(Mode) {}
+
+void FOutlinerHierarchy::CreateItems(TArray<FSceneOutlinerTreeItemPtr>& OutItems) const
+{
+	TArray<TSoftObjectPtr<UAnimNextRigVMAsset>> Assets;
+	static_cast<FCommonOutlinerMode* const>(Mode)->GetAssets(Assets);
+
+	WeakProcessedAssets.Reset();
+	WeakProcessedStructs.Reset();
+
+	if(Assets.Num())
+	{
+		const int32 NumAssets = Assets.Num();
+		for(int32 Index = 0; Index < NumAssets; Index++)
+		{
+			const TSoftObjectPtr<UAnimNextRigVMAsset>& SoftAsset = Assets[Index];
+
+			TArray<const FSoftObjectPath> SharedVariableSourcePaths;
+			ProcessAsset(SoftAsset, (NumAssets - Index), OutItems, SharedVariableSourcePaths);
+		}
+	}
+}
+
+void FOutlinerHierarchy::ProcessAsset(const TSoftObjectPtr<UAnimNextRigVMAsset>& InSoftAsset, int32 Depth, TArray<FSceneOutlinerTreeItemPtr>& OutItems, TArray<const FSoftObjectPath>& OutSharedVariableSourcePaths) const
+{
+	UAnimNextRigVMAsset* Asset = InSoftAsset.Get();
+	if(Asset == nullptr)
+	{
+		return;
+	}
+
+	const UAnimNextRigVMAssetEditorData* EditorData = UncookedOnly::FUtils::GetEditorData<UAnimNextRigVMAssetEditorData>(Asset);
+	if (EditorData == nullptr)
+	{
+		return;
+	}
+
+	if (WeakProcessedAssets.Contains(Asset))
+	{
+		return;
+	}
+
+	WeakProcessedAssets.Add(Asset);
+
+	UE_CLOG(GOutputHierarchyDebugInformationCvar.GetValueOnAnyThread(),LogAnimation, Warning, TEXT("ProcessAsset: %s %i"), *InSoftAsset.GetAssetName(), Depth);
+
+	// Adjust asset depth to ensure SharedVariables are always pinned to 0 depth
+	const int32 AssetDepth = Asset->GetClass() && Asset->GetClass() == UAnimNextSharedVariables::StaticClass() ? 0 : Depth;
+
+	TArray<const FSoftObjectPath> SharedAssetPaths;
+	RetrieveSharedAssetPaths(Asset, SharedAssetPaths);
+	if (const FSceneOutlinerTreeItemPtr AssetItem = Mode->CreateItemFor<FOutlinerAssetItem>(FOutlinerAssetItem::FItemData{InSoftAsset, AssetDepth, SharedAssetPaths}))
+	{
+		OutItems.Add(AssetItem);
+
+		UE_CLOG(GOutputHierarchyDebugInformationCvar.GetValueOnAnyThread(), LogAnimation, Warning, TEXT("* FOutlinerAssetItem: %s %i"), *InSoftAsset.GetAssetName(), AssetDepth);
+	}
+
+	CreateItemsForAsset(EditorData, OutItems);
+}
+
+void FOutlinerHierarchy::RetrieveSharedAssetPaths(const TSoftObjectPtr<UAnimNextRigVMAsset>& InSoftAsset, TArray<const FSoftObjectPath>& OutSharedVariableSourcePaths) const
+{
+	UAnimNextRigVMAsset* Asset = InSoftAsset.Get();
+	if(Asset == nullptr)
+	{
+		return;
+	}
+
+	const UAnimNextRigVMAssetEditorData* EditorData = UncookedOnly::FUtils::GetEditorData<UAnimNextRigVMAssetEditorData>(Asset);
+	if (EditorData == nullptr)
+	{
+		return;
+	}
+
+	TArray<FSceneOutlinerTreeItemPtr> ParentItems;
+	EditorData->ForEachEntryOfType<UAnimNextSharedVariablesEntry>([this, &OutSharedVariableSourcePaths](const UAnimNextSharedVariablesEntry* InSharedVariablesEntry)
+	{
+		switch (InSharedVariablesEntry->Type)
+		{
+		case EAnimNextSharedVariablesType::Asset:
+			{
+				if (const UAnimNextSharedVariables* SharedVariables = InSharedVariablesEntry->GetAsset())
+				{
+					OutSharedVariableSourcePaths.Add(SharedVariables);
+				}
+			}
+			break;
+
+		case EAnimNextSharedVariablesType::Struct:
+			{
+				if (const UScriptStruct* Struct = InSharedVariablesEntry->GetStruct())
+				{
+					OutSharedVariableSourcePaths.Add(Struct);
+				}
+			}
+			break;
+		}
+
+		return true;
+	});
+}
+}
diff --git a/Engine/Plugins/Experimental/UAF/UAF/Source/UAFEditor/Private/Common/Outliner/OutlinerHierarchy.h b/Engine/Plugins/Experimental/UAF/UAF/Source/UAFEditor/Private/Common/Outliner/OutlinerHierarchy.h
new file mode 100644
--- /dev/null
+++ b/Engine/Plugins/Experimental/UAF/UAF/Source/UAFEditor/Private/Common/Outliner/OutlinerHierarchy.h
@@ -0,0 +1,35 @@
+// Copyright Epic Games, Inc. All Rights Reserved.
+
+#pragma once
+
+#include "CoreMinimal.h"
+#include "ISceneOutlinerHierarchy.h"
+
+class UAnimNextRigVMAssetEditorData;
+class UAnimNextRigVMAsset;
+
+namespace UE::UAF::Editor
+{
+
+extern TAutoConsoleVariable<bool> GOutputHierarchyDebugInformationCvar;
+class FOutlinerHierarchy : public ISceneOutlinerHierarchy
+{
+public:
+	FOutlinerHierarchy(ISceneOutlinerMode* Mode);
+	FOutlinerHierarchy(const FOutlinerHierarchy&) = delete;
+	FOutlinerHierarchy& operator=(const FOutlinerHierarchy&) = delete;
+
+	// Begin ISceneOutlinerHierarchy overrides
+	virtual void CreateItems(TArray<FSceneOutlinerTreeItemPtr>& OutItems) const override;
+	virtual void CreateChildren(const FSceneOutlinerTreeItemPtr& Item, TArray<FSceneOutlinerTreeItemPtr>& OutChildren) const override {}
+	virtual void CreateItemsForAsset(const UAnimNextRigVMAssetEditorData* EditorData, TArray<FSceneOutlinerTreeItemPtr>& OutItems) const = 0;	
+	// End ISceneOutlinerHierarchy overrides
+protected:
+	void ProcessAsset(const TSoftObjectPtr<UAnimNextRigVMAsset>& InSoftAsset, int32 Depth, TArray<FSceneOutlinerTreeItemPtr>& OutItems, TArray<const FSoftObjectPath>& OutSharedVariableSourcePaths) const;
+	void RetrieveSharedAssetPaths(const TSoftObjectPtr<UAnimNextRigVMAsset>& InSoftAsset, TArray<const FSoftObjectPath>& OutSharedVariableSourcePaths) const;
+
+	mutable TSet<TWeakObjectPtr<const UAnimNextRigVMAsset>> WeakProcessedAssets;
+	mutable TSet<TWeakObjectPtr<const UScriptStruct>> WeakProcessedStructs;
+};
+
+}
diff --git a/Engine/Plugins/Experimental/UAF/UAF/Source/UAFEditor/Private/Common/Outliner/SCommonOutliner.cpp b/Engine/Plugins/Experimental/UAF/UAF/Source/UAFEditor/Private/Common/Outliner/SCommonOutliner.cpp
new file mode 100644
--- /dev/null
+++ b/Engine/Plugins/Experimental/UAF/UAF/Source/UAFEditor/Private/Common/Outliner/SCommonOutliner.cpp
@@ -0,0 +1,205 @@
+// Copyright Epic Games, Inc. All Rights Reserved.
+
+#include "SCommonOutliner.h"
+
+#include "AnimNextAssetWorkspaceAssetUserData.h"
+#include "AnimNextRigVMAsset.h"
+#include "AnimNextRigVMAssetEditorData.h"
+#include "ContentBrowserModule.h"
+#include "IContentBrowserSingleton.h"
+#include "SAssetDropTarget.h"
+#include "SSceneOutliner.h"
+#include "IWorkspaceEditorModule.h"
+#include "UncookedOnlyUtils.h"
+#include "WorkspaceAssetRegistryInfo.h"
+#include "Entries/AnimNextSharedVariablesEntry.h"
+
+#define LOCTEXT_NAMESPACE "SCommonOutliner"
+
+namespace UE::UAF::Editor
+{
+
+void SCommonOutliner::SetAssets(TConstArrayView<TSoftObjectPtr<UAnimNextRigVMAsset>> InAssets)
+{
+	for(const TSoftObjectPtr<UAnimNextRigVMAsset>& CurrentSoftAsset : Assets)
+	{
+		if(UAnimNextRigVMAsset* CurrentAsset = CurrentSoftAsset.Get())
+		{
+			UnregisterAssetDelegates(CurrentAsset);
+		}
+	}
+
+	Assets = InAssets;
+
+	for(const TSoftObjectPtr<UAnimNextRigVMAsset>& NewSoftAsset : InAssets)
+	{
+		if(UAnimNextRigVMAsset* NewAsset = NewSoftAsset.Get())
+		{
+			RegisterAssetDelegates(NewAsset);
+		}
+	}
+
+	FullRefresh();
+}
+
+void SCommonOutliner::UpdateAssets()
+{
+	TArray<FSoftObjectPath> AssetsToAsyncLoad;
+	TArray<TSoftObjectPtr<UAnimNextRigVMAsset>> ExportAssets;
+
+	auto HandleAssetPath = [this, &AssetsToAsyncLoad, &ExportAssets](const FSoftObjectPath& Path)
+	{
+		ExportAssets.Add(TSoftObjectPtr<UAnimNextRigVMAsset>(Path));
+
+		if (Path.ResolveObject() == nullptr)
+		{
+			AssetsToAsyncLoad.Add(Path);
+		}
+	};
+
+	if (FWorkspaceOutlinerAssetReferenceItemData::IsAssetReference(Export)
+		|| FWorkspaceOutlinerGroupItemData::IsGroupItem(Export)
+		|| FAnimNextCollapseGraphsOutlinerDataBase::IsCollapsedGraphBase(Export))
+	{
+		// Reference item
+		TArray<FSoftObjectPath> AssetPaths;
+		Export.GetAssetPaths(AssetPaths);
+
+		for (const FSoftObjectPath& Path : AssetPaths)
+		{
+			HandleAssetPath(Path);
+		}
+	}
+	else
+	{
+		const FSoftObjectPath FirstAssetPath = Export.GetFirstAssetPath();
+		if(FirstAssetPath.IsAsset())
+		{
+			HandleAssetPath(FirstAssetPath);
+		}
+		else if(Export.HasData() && Export.GetData().GetScriptStruct()->IsChildOf(FAnimNextAssetEntryOutlinerData::StaticStruct()))
+		{
+			const FAnimNextAssetEntryOutlinerData& EntryData = Export.GetData().Get<FAnimNextAssetEntryOutlinerData>();
+			if(EntryData.SoftEntryPtr.IsValid())
+			{
+				if(UAnimNextRigVMAsset* Asset = EntryData.GetEntry()->GetTypedOuter<UAnimNextRigVMAsset>())
+				{
+					HandleAssetPath(Asset);
+				}
+			}
+		}
+	}
+
+	SetAssets(ExportAssets);
+
+	// Try async load any missing assets
+	for(const FSoftObjectPath& AssetPath : AssetsToAsyncLoad)
+	{
+		TWeakPtr<SCommonOutliner> WeakOutliner = StaticCastSharedRef<SCommonOutliner>(AsShared());
+
+		AssetPath.LoadAsync(FLoadSoftObjectPathAsyncDelegate::CreateLambda([WeakOutliner](const FSoftObjectPath& InSoftObjectPath, UObject* InObject)
+		{
+			UAnimNextRigVMAsset* Asset = Cast<UAnimNextRigVMAsset>(InObject);
+			if(Asset == nullptr)
+			{
+				return;
+			}
+
+			TSharedPtr<SCommonOutliner> PinnedVariablesOutliner = WeakOutliner.Pin();
+			if(!PinnedVariablesOutliner.IsValid())
+			{
+				return;
+			}
+
+			PinnedVariablesOutliner->HandleAssetLoaded(InSoftObjectPath, Asset);
+		}));
+	}
+}
+
+void SCommonOutliner::Construct(const FArguments& InArgs, const FSceneOutlinerInitializationOptions& InitOptions)
+{
+	BindAssetDelegate = InArgs._BindAssetDelegate;
+	UnbindAssetDelegate = InArgs._UnbindAssetDelegate;
+	
+	// Either both bounds, or none at all
+	ensure((BindAssetDelegate.IsBound() && UnbindAssetDelegate.IsBound()) || (!BindAssetDelegate.IsBound() && !UnbindAssetDelegate.IsBound()));
+
+	SSceneOutliner::Construct(SSceneOutliner::FArguments(), InitOptions);
+}
+
+void SCommonOutliner::SetExport(const FWorkspaceOutlinerItemExport& InExport)
+{
+	Export = InExport;
+
+	UpdateAssets();
+}
+
+void SCommonOutliner::HandleAssetLoaded(const FSoftObjectPath& InSoftObjectPath, UAnimNextRigVMAsset* InAsset)
+{
+	if(InAsset && Assets.Contains(TSoftObjectPtr<UAnimNextRigVMAsset>(InAsset)))
+	{
+		RegisterAssetDelegates(InAsset);
+
+		FullRefresh();
+	}
+}
+
+void SCommonOutliner::RegisterAssetDelegates(const UAnimNextRigVMAsset* InAsset)
+{
+	// Bind for any modification callbacks
+	UAnimNextRigVMAssetEditorData* EditorData = UncookedOnly::FUtils::GetEditorData<UAnimNextRigVMAssetEditorData>(InAsset);
+	if(EditorData != nullptr)
+	{
+		BindAssetDelegate.ExecuteIfBound(EditorData);
+
+		for (const TObjectPtr<UAnimNextRigVMAssetEntry>& Entry : EditorData->Entries)
+		{
+			if (UAnimNextSharedVariablesEntry* SharedVariablesEntry = Cast<UAnimNextSharedVariablesEntry>(Entry.Get()))
+			{
+				if (TObjectPtr<const UAnimNextSharedVariables> SharedVariablesAsset = SharedVariablesEntry->GetAsset())
+				{
+					RegisterAssetDelegates(SharedVariablesAsset);
+					Assets.AddUnique(Cast<UAnimNextRigVMAsset>(const_cast<UAnimNextSharedVariables*>(SharedVariablesAsset.Get())));
+				}
+			}
+		}
+	}
+}
+
+void SCommonOutliner::UnregisterAssetDelegates(const UAnimNextRigVMAsset* InAsset)
+{
+	UAnimNextRigVMAssetEditorData* EditorData = UncookedOnly::FUtils::GetEditorData<UAnimNextRigVMAssetEditorData>(InAsset);
+	if(EditorData != nullptr)
+	{
+		UnbindAssetDelegate.ExecuteIfBound(EditorData);
+
+		for (const TObjectPtr<UAnimNextRigVMAssetEntry>& Entry : EditorData->Entries)
+		{
+			if (UAnimNextSharedVariablesEntry* SharedVariablesEntry = Cast<UAnimNextSharedVariablesEntry>(Entry.Get()))
+			{
+				if (TObjectPtr<const UAnimNextSharedVariables> SharedVariablesAsset = SharedVariablesEntry->GetAsset())
+				{
+					UnregisterAssetDelegates(SharedVariablesAsset);
+				}
+			}
+		}
+	}
+}
+
+void SCommonOutliner::SetHighlightedItem(FSceneOutlinerTreeItemPtr Item) const
+{
+	GetTreeView()->Private_SetItemHighlighted(Item, true);
+}
+
+void SCommonOutliner::ClearHighlightedItem(FSceneOutlinerTreeItemPtr Item) const
+{
+	GetTreeView()->Private_SetItemHighlighted(Item, false);
+}
+
+bool SCommonOutliner::HasAssets() const
+{
+	return Assets.Num() > 0;
+}
+}
+
+#undef LOCTEXT_NAMESPACE
diff --git a/Engine/Plugins/Experimental/UAF/UAF/Source/UAFEditor/Private/Common/Outliner/SCommonOutliner.h b/Engine/Plugins/Experimental/UAF/UAF/Source/UAFEditor/Private/Common/Outliner/SCommonOutliner.h
new file mode 100644
--- /dev/null
+++ b/Engine/Plugins/Experimental/UAF/UAF/Source/UAFEditor/Private/Common/Outliner/SCommonOutliner.h
@@ -0,0 +1,57 @@
+// Copyright Epic Games, Inc. All Rights Reserved.
+
+#pragma once
+
+#include "ContentBrowserDelegates.h"
+#include "SSceneOutliner.h"
+#include "WorkspaceAssetRegistryInfo.h"
+
+enum class EAnimNextEditorDataNotifType : uint8;
+class UAnimNextRigVMAsset;
+struct FWorkspaceOutlinerItemExport;
+class SPositiveActionButton;
+class UAnimNextRigVMAssetEditorData;
+
+namespace UE::Workspace
+{
+class IWorkspaceEditor;
+struct FWorkspaceDocument;
+}
+
+namespace UE::UAF::Editor
+{
+
+DECLARE_DELEGATE_OneParam(FAssetEditorDataDelegate, UAnimNextRigVMAssetEditorData*)
+
+class SCommonOutliner : public SSceneOutliner
+{
+	friend class FCommonOutlinerMode;
+public:
+
+	SLATE_BEGIN_ARGS(SCommonOutliner) {}
+		SLATE_EVENT(FAssetEditorDataDelegate, BindAssetDelegate)
+		SLATE_EVENT(FAssetEditorDataDelegate, UnbindAssetDelegate)
+	SLATE_END_ARGS()
+
+	void Construct(const FArguments& InArgs, const FSceneOutlinerInitializationOptions& InitOptions);
+
+	void SetExport(const FWorkspaceOutlinerItemExport& InExport);
+	TConstArrayView<TSoftObjectPtr<UAnimNextRigVMAsset>> GetAssetsView() const { return Assets; }	
+	void UpdateAssets();
+	bool HasAssets() const;
+
+protected:
+	void SetAssets(TConstArrayView<TSoftObjectPtr<UAnimNextRigVMAsset>> InAssets);
+	void HandleAssetLoaded(const FSoftObjectPath& InSoftObjectPath, UAnimNextRigVMAsset* InAsset);
+	void SetHighlightedItem(FSceneOutlinerTreeItemPtr Item) const;
+	void ClearHighlightedItem(FSceneOutlinerTreeItemPtr Item) const;
+	void RegisterAssetDelegates(const UAnimNextRigVMAsset* InAsset);
+	void UnregisterAssetDelegates(const UAnimNextRigVMAsset* InAsset);
+private:
+	TArray<TSoftObjectPtr<UAnimNextRigVMAsset>> Assets;
+	FWorkspaceOutlinerItemExport Export;
+
+	FAssetEditorDataDelegate BindAssetDelegate;
+	FAssetEditorDataDelegate UnbindAssetDelegate;
+};
+}
diff --git a/Engine/Plugins/Experimental/UAF/UAF/Source/UAFEditor/Private/Functions/Outliner/FunctionsOutlinerHierarchy.cpp b/Engine/Plugins/Experimental/UAF/UAF/Source/UAFEditor/Private/Functions/Outliner/FunctionsOutlinerHierarchy.cpp
new file mode 100644
--- /dev/null
+++ b/Engine/Plugins/Experimental/UAF/UAF/Source/UAFEditor/Private/Functions/Outliner/FunctionsOutlinerHierarchy.cpp
@@ -0,0 +1,169 @@
+// Copyright Epic Games, Inc. All Rights Reserved.
+
+#include "FunctionsOutlinerHierarchy.h"
+
+#include "AnimNextRigVMAssetEditorData.h"
+#include "FunctionOutlinerEntry.h"
+#include "FunctionsOutlinerCategoryItem.h"
+#include "ISceneOutlinerMode.h"
+#include "Common/Outliner/OutlinerAssetItem.h"
+
+namespace UE::UAF::Editor
+{
+
+FSceneOutlinerTreeItemPtr FFunctionsOutlinerHierarchy::FindOrCreateParentItem(const ISceneOutlinerTreeItem& Item, const TMap<FSceneOutlinerTreeItemID, FSceneOutlinerTreeItemPtr>& Items, bool bCreate)
+{
+	FSceneOutlinerTreeItemPtr ParentItem;
+	
+	uint32 ParentHash = INDEX_NONE;
+	FSoftObjectPath OwnerPath;
+	FString CategoryPath;
+	
+	if (const FFunctionsOutlinerEntryItem* EntryItem = Item.CastTo<FFunctionsOutlinerEntryItem>())
+	{
+		if (URigVMLibraryNode* LibraryNode = EntryItem->WeakLibraryNode.Get())
+		{
+			if (UAnimNextRigVMAsset* Asset = LibraryNode->GetTypedOuter<UAnimNextRigVMAsset>())
+			{
+				OwnerPath = Asset;
+			}
+
+			uint32 CategoryHash = INDEX_NONE;
+			if (!LibraryNode->GetNodeCategory().IsEmpty())
+			{
+				CategoryPath = LibraryNode->GetNodeCategory();
+				CategoryHash = GetTypeHash(LibraryNode->GetNodeCategory());
+			}
+			
+			check(OwnerPath.IsValid());
+			ParentHash = CategoryHash != INDEX_NONE ? HashCombine(GetTypeHash(OwnerPath), CategoryHash) : GetTypeHash(OwnerPath);				
+		}
+	}
+
+	if (const FFunctionsOutlinerCategoryItem* CategoryItem = Item.CastTo<FFunctionsOutlinerCategoryItem>())
+	{
+		if (UAnimNextRigVMAsset* Asset = CategoryItem->WeakOwner.Get())
+		{
+			OwnerPath = Asset;
+			const uint32 AssetHash = GetTypeHash(OwnerPath);
+			const FStringView ParentCategory = CategoryItem->ParentCategoryName;
+			CategoryPath = CategoryItem->ParentCategoryName;
+			ParentHash = ParentCategory.Len() ? HashCombine(AssetHash, GetTypeHash(ParentCategory)) : AssetHash;
+		}
+	}
+
+	if (const FSceneOutlinerTreeItemPtr* ParentItemPtr = Items.Find(ParentHash))
+	{
+		return *ParentItemPtr;
+	}
+	
+	TSoftObjectPtr<UAnimNextRigVMAsset> SoftAsset = Cast<UAnimNextRigVMAsset>(OwnerPath.ResolveObject());
+
+	if (!CategoryPath.IsEmpty())
+	{
+		UAnimNextRigVMAssetEditorData* EditorData = UncookedOnly::FUtils::GetEditorData(SoftAsset.Get());
+		const int32 CategoryIndex = 0;
+		ensure(CategoryIndex != INDEX_NONE);
+	
+		FStringView ParentCategoryName;
+		FStringView SubCategoryName;
+	
+		int32 Index = INDEX_NONE;
+		if (CategoryPath.FindLastChar('|', Index))
+		{
+			ParentCategoryName = CategoryPath.Left(Index);
+			SubCategoryName = CategoryPath.RightChop(Index+1);
+		}
+		else
+		{
+			SubCategoryName = CategoryPath;
+		}
+	
+		if (FSceneOutlinerTreeItemPtr CategoryItem = Mode->CreateItemFor<FFunctionsOutlinerCategoryItem>(FOutlinerCategoryItem::FItemData{UncookedOnly::FUtils::GetAsset<UAnimNextRigVMAsset>(EditorData), CategoryIndex, SubCategoryName, ParentCategoryName, CategoryPath}, bCreate))
+		{			
+			UE_CLOG(GOutputHierarchyDebugInformationCvar.GetValueOnAnyThread(),LogAnimation, Warning, TEXT("+ FFunctionsOutlinerCategoryItem: %s  %i"), *CategoryPath, 0);
+			
+			return CategoryItem;
+		}
+	}
+	else if (OwnerPath.IsValid())
+	{
+		TArray<const FSoftObjectPath> SharedVariableSourcePaths;
+		if (SoftAsset.IsValid())
+		{
+			RetrieveSharedAssetPaths(SoftAsset, SharedVariableSourcePaths);
+			if (FSceneOutlinerTreeItemPtr AssetItem = Mode->CreateItemFor<FOutlinerAssetItem>(FOutlinerAssetItem::FItemData{SoftAsset, 0, SharedVariableSourcePaths}, bCreate))
+			{
+				UE_CLOG(GOutputHierarchyDebugInformationCvar.GetValueOnAnyThread(),LogAnimation, Warning, TEXT("+ FOutlinerAssetItem: %s  %i"), *SoftAsset.GetAssetName(), 0);				
+				return AssetItem;
+			}
+		}
+	}
+
+	return ParentItem;
+}
+
+FFunctionsOutlinerHierarchy::FFunctionsOutlinerHierarchy(ISceneOutlinerMode* Mode)
+	: FOutlinerHierarchy(Mode)
+{
+}
+
+void FFunctionsOutlinerHierarchy::CreateItemsForAsset(const UAnimNextRigVMAssetEditorData* EditorData, TArray<FSceneOutlinerTreeItemPtr>& OutItems) const
+{
+	TSet<FString> ParsedCategories;
+	auto RecursivelyAddCategories = [this, &EditorData, &OutItems, &ParsedCategories](const FString InCategoryPath)
+	{	
+		const int32 InsertionIndex = OutItems.Num(); 
+		FString CategoryPath = InCategoryPath;
+		while (!CategoryPath.IsEmpty())
+		{
+			FString CategoryParent;
+			FString Category;
+			
+			if (!CategoryPath.Split(TEXT("|"), &CategoryParent, &Category, ESearchCase::IgnoreCase, ESearchDir::FromEnd))
+			{
+				Category = CategoryPath;
+			}
+			
+			if (!ParsedCategories.Contains(CategoryPath))
+			{
+				if (auto Item = Mode->CreateItemFor<FFunctionsOutlinerCategoryItem>(FOutlinerCategoryItem::FItemData{UncookedOnly::FUtils::GetAsset<UAnimNextRigVMAsset>(EditorData), 0, Category, CategoryParent, CategoryPath}))
+				{
+					OutItems.Insert(Item, InsertionIndex);
+					ParsedCategories.Add(CategoryPath);
+				
+					UE_CLOG(GOutputHierarchyDebugInformationCvar.GetValueOnAnyThread(),LogAnimation, Warning, TEXT("+ FFunctionsOutlinerCategoryItem: %s  %i"), *CategoryPath, 0);
+				}
+			}			
+			
+			CategoryPath = CategoryParent;
+		}
+	};
+	
+	
+	if (URigVMFunctionLibrary* FunctionLibrary = EditorData->GetLocalFunctionLibrary())
+	{
+		TArray<URigVMLibraryNode*> LibraryNodes = FunctionLibrary->GetFunctions();
+
+		int32 Index = 0;
+		for (URigVMLibraryNode* LibraryNode : LibraryNodes)
+		{
+			if (FSceneOutlinerTreeItemPtr Item = Mode->CreateItemFor<FFunctionsOutlinerEntryItem>(FFunctionsOutlinerEntryItem::FItemData{UncookedOnly::FUtils::GetAsset<UAnimNextRigVMAsset>(EditorData), Index++, LibraryNode}))
+			{
+				if (FFunctionsOutlinerEntryItem* EntryItem = Item->CastTo<FFunctionsOutlinerEntryItem>())
+				{
+					EntryItem->SortValue = Index++;	
+					
+					UE_CLOG(GOutputHierarchyDebugInformationCvar.GetValueOnAnyThread(), LogAnimation, Warning, TEXT("\t- FFunctionsOutlinerEntryItem: %s %i %s"), *EntryItem->WeakLibraryNode.Get()->GetName(), Index, *LibraryNode->GetNodeCategory());
+				}
+
+				const FString CategoryPath = LibraryNode->GetNodeCategory();
+				RecursivelyAddCategories(CategoryPath);
+			
+				OutItems.Add(Item);
+			}
+		}
+	}
+}
+
+}
diff --git a/Engine/Plugins/Experimental/UAF/UAF/Source/UAFEditor/Private/Variables/Outliner/VariablesOutlinerCategoryItem.cpp b/Engine/Plugins/Experimental/UAF/UAF/Source/UAFEditor/Private/Variables/Outliner/VariablesOutlinerCategoryItem.cpp
--- a/Engine/Plugins/Experimental/UAF/UAF/Source/UAFEditor/Private/Variables/Outliner/VariablesOutlinerCategoryItem.cpp
+++ b/Engine/Plugins/Experimental/UAF/UAF/Source/UAFEditor/Private/Variables/Outliner/VariablesOutlinerCategoryItem.cpp
@@ -7,7 +7,7 @@
 #include "ISceneOutliner.h"
 #include "ScopedTransaction.h"
 #include "UncookedOnlyUtils.h"
-#include "VariablesOutlinerDragDrop.h"
+#include "Common/Outliner/OutlinerDragDropOps.h"
 #include "Entries/AnimNextRigVMAssetEntry.h"
 #include "Styling/StarshipCoreStyle.h"
 #include "Widgets/Text/SInlineEditableTextBlock.h"
@@ -17,277 +17,11 @@
 
 namespace UE::UAF::Editor
 {
-class FVariableDragDropOp;
 
-const FSceneOutlinerTreeItemType FVariablesOutlinerCategoryItem::Type;
+const FSceneOutlinerTreeItemType FVariablesOutlinerCategoryItem::Type(&FOutlinerCategoryItem::Type);
 
-class SVariablesOutlinerCategoryLabel : FSceneOutlinerCommonLabelData, public SCompoundWidget
+FVariablesOutlinerCategoryItem::FVariablesOutlinerCategoryItem(const FOutlinerCategoryItem::FItemData& InItemData) : FOutlinerCategoryItem(FVariablesOutlinerCategoryItem::Type, InItemData)
 {
-	SLATE_BEGIN_ARGS(SVariablesOutlinerCategoryLabel) {}
-	SLATE_END_ARGS()
-
-	void Construct(const FArguments& InArgs, FVariablesOutlinerCategoryItem& InTreeItem, ISceneOutliner& SceneOutliner)
-	{
-		WeakSceneOutliner = StaticCastSharedRef<ISceneOutliner>(SceneOutliner.AsShared());
-		TreeItem = StaticCastSharedRef<FVariablesOutlinerCategoryItem>(InTreeItem.AsShared());
-		
-		ChildSlot
-		[
-			SNew(SHorizontalBox)
-			+SHorizontalBox::Slot()
-			.AutoWidth()
-			.VAlign(VAlign_Center)
-			.Padding(2.0f, 2.0f)
-			[
-				SAssignNew(TextBlock, SInlineEditableTextBlock)
-				.Font(FStyleFonts::Get().NormalBold)
-				.Text(this, &SVariablesOutlinerCategoryLabel::GetDisplayText)
-				.HighlightText(SceneOutliner.GetFilterHighlightText())
-				.ColorAndOpacity(this, &SVariablesOutlinerCategoryLabel::GetForegroundColor)
-				.OnTextCommitted(this, &SVariablesOutlinerCategoryLabel::OnTextCommited)
-				.OnVerifyTextChanged(this, &SVariablesOutlinerCategoryLabel::OnVerifyTextChanged)
-			]
-		];
-	}
-	
-	FText GetDisplayText() const
-	{
-		if (const TSharedPtr<FVariablesOutlinerCategoryItem> Item = TreeItem.Pin())
-		{
-			return FText::FromString(Item->GetDisplayString());
-		}
-		return FText();
-	}
-
-	virtual FSlateColor GetForegroundColor() const override
-	{
-		const TOptional<FLinearColor> BaseColor = FSceneOutlinerCommonLabelData::GetForegroundColor(*TreeItem.Pin());
-		return BaseColor.IsSet() ? BaseColor.GetValue() : FSlateColor::UseForeground();
-	}
-
-	
-	void OnTextCommited(const FText& InLabel, ETextCommit::Type InCommitInfo) const	
-	{
-		if(InCommitInfo == ETextCommit::OnEnter)
-		{
-			if (const TSharedPtr<FVariablesOutlinerCategoryItem> Item = TreeItem.Pin())
-			{
-				Item->Rename(InLabel); 
-			}
-		}
-	}
-	
-	bool OnVerifyTextChanged(const FText& InLabel, FText& OutErrorMessage) const
-	{
-		if (const TSharedPtr<FVariablesOutlinerCategoryItem> Item = TreeItem.Pin())
-		{
-			return Item->ValidateName(InLabel, OutErrorMessage); 
-		}
-		return false;
-	}
-
-	virtual FReply OnDrop( const FGeometry& MyGeometry, const FDragDropEvent& DragDropEvent ) override
-	{
-		TSharedPtr<FVariableDragDropOp> GraphDropOp = DragDropEvent.GetOperationAs<FVariableDragDropOp>();
-		if (GraphDropOp.IsValid())
-		{
-			if (const TSharedPtr<FVariablesOutlinerCategoryItem> CategoryItem = TreeItem.Pin())
-			{
-				TSharedPtr<FVariablesOutlinerEntryItem> EntryItem = GraphDropOp->WeakItem.Pin();
-				if (EntryItem.IsValid())
-				{
-					if (UAnimNextVariableEntry* VariableEntry = EntryItem->WeakEntry.Get())
-					{
-						UAnimNextRigVMAsset* VariableAsset = VariableEntry->GetTypedOuter<UAnimNextRigVMAsset>();
-						UAnimNextRigVMAsset* CategoryAsset = CategoryItem->WeakOwner.Get();
-
-						const FText OtherVariableNameText = FText::FromName(VariableEntry->GetVariableName());
-						const FText ThisVariableCategoryText = FText::FromString(CategoryItem->GetDisplayString());
-						const FText FormattedMessage = FText::Format(LOCTEXT("MoveVariableToCategoryFormat", "Moving {0} to Category {1}"), OtherVariableNameText, ThisVariableCategoryText);	
-						FScopedTransaction Transaction(FormattedMessage);
-
-						UAnimNextVariableEntry* NewVariableEntry = VariableEntry;
-						
-						if (VariableAsset != CategoryAsset)
-						{
-							const FName VariableEntryName = VariableEntry->GetVariableName();
-							UncookedOnly::FUtils::MoveVariableToAsset(VariableEntry, CategoryAsset, true, true);
-							NewVariableEntry = Cast<UAnimNextVariableEntry>(UncookedOnly::FUtils::GetEditorData(CategoryAsset)->FindEntry(VariableEntryName));
-						}
-						
-						check(NewVariableEntry);
-						
-						// Only add if not already part of category
-						if (NewVariableEntry->GetVariableCategory() != CategoryItem->CategoryPath)
-						{
-							NewVariableEntry->SetVariableCategory(CategoryItem->CategoryPath);
-						}
-					}
-				}
-			}
-			
-			return FReply::Handled();
-		}
-		
-		TSharedPtr<FCategoryDragDropOp> CategoryDropOp = DragDropEvent.GetOperationAs<FCategoryDragDropOp>();
-		if (CategoryDropOp.IsValid())
-		{
-			if (const TSharedPtr<FVariablesOutlinerCategoryItem> CategoryItem = TreeItem.Pin())
-			{
-				TSharedPtr<FVariablesOutlinerCategoryItem> DropCategoryItem = CategoryDropOp->WeakItem.Pin();
-				if (DropCategoryItem.IsValid())
-				{
-					const FString ThisCategoryName = CategoryItem->CategoryName;
-					const FString OtherCategoryName = DropCategoryItem->CategoryName;
-
-					UAnimNextRigVMAsset* ThisOwner = CategoryItem->WeakOwner.Get();
-					const UAnimNextRigVMAsset* OtherOwner = DropCategoryItem->WeakOwner.Get();
-
-					if (ThisOwner == OtherOwner)
-					{
-						const FText OtherCategoryText = FText::FromString(DropCategoryItem->GetDisplayString());
-						const FText ThisCategoryText = FText::FromString(CategoryItem->GetDisplayString());
-						
-						const FText FormattedMessage = FText::Format(LOCTEXT("ReorderCategory", "Reorder {0} before {1}"), OtherCategoryText, ThisCategoryText);
-						
-						FScopedTransaction Transaction(FormattedMessage);
-						UncookedOnly::FUtils::GetEditorData(ThisOwner)->ReorderCategory(OtherCategoryName, ThisCategoryName);
-					}
-				}
-			}
-
-			return FReply::Handled();
-		}
-		
-		
-		return FReply::Unhandled();
-	}
-
-	virtual void OnDragEnter( const FGeometry& MyGeometry, const FDragDropEvent& DragDropEvent ) override
-	{
-		TSharedPtr<FVariableDragDropOp> GraphDropOp = DragDropEvent.GetOperationAs<FVariableDragDropOp>();
-		if (GraphDropOp.IsValid())
-		{
-			if (const TSharedPtr<FVariablesOutlinerCategoryItem> CategoryItem = TreeItem.Pin())
-			{
-				TSharedPtr<FVariablesOutlinerEntryItem> EntryItem = GraphDropOp->WeakItem.Pin();
-				if (EntryItem.IsValid())
-				{
-					if (UAnimNextVariableEntry* VariableEntry = EntryItem->WeakEntry.Get())
-					{
-						const FText OtherVariableNameText = FText::FromName(VariableEntry->GetVariableName());
-						const FText ThisVariableCategoryText = FText::FromString(CategoryItem->GetDisplayString());
-
-						const UAnimNextRigVMAssetEditorData* VariableEditorData = VariableEntry->GetTypedOuter<UAnimNextRigVMAssetEditorData>();
-						const UAnimNextRigVMAssetEditorData* CategoryEditorData = UncookedOnly::FUtils::GetEditorData(CategoryItem->WeakOwner.Get());
-
-						if (VariableEditorData == CategoryEditorData)
-						{
-							if (VariableEntry->GetVariableCategory() == CategoryItem->CategoryName)
-							{
-								const FText FormattedMessage = FText::Format(LOCTEXT("VariableAlreadyPartOfCategoryFormat", "{0} is already part of Category {1}"), OtherVariableNameText, ThisVariableCategoryText);
-								
-								const FSlateBrush* ErrorSymbol = FAppStyle::GetBrush(TEXT("Graph.ConnectorFeedback.Error"));
-								GraphDropOp->SetSimpleFeedbackMessage(ErrorSymbol, FSlateColor(GetForegroundColor()), FormattedMessage);
-							}
-							else if (!VariableEntry->GetVariableCategory().IsEmpty())
-							{
-								const FText FormattedMessage = FText::Format(LOCTEXT("MoveVariableToCategoryFormat", "Moving {0} to Category {1}"), OtherVariableNameText, ThisVariableCategoryText);							
-								GraphDropOp->SetSimpleFeedbackMessage(nullptr, FSlateColor(GetForegroundColor()), FormattedMessage);
-							}
-							else
-							{
-								const FText FormattedMessage = FText::Format(LOCTEXT("AddVariableToCategoryFormat", "Add {0} to Category {1}"), OtherVariableNameText, ThisVariableCategoryText);
-								GraphDropOp->SetSimpleFeedbackMessage(nullptr, FSlateColor(GetForegroundColor()), FormattedMessage);
-							}
-						}
-						else
-						{
-							const FText ThisAssetText =  FText::FromString(CategoryItem->WeakOwner.Get()->GetName());
-							const FText FormattedMessage = FText::Format(LOCTEXT("MoveVariableToCategoryAndAssetFormat", "Moving {0} Category {1} in {2}"), OtherVariableNameText, ThisVariableCategoryText, ThisAssetText);
-							GraphDropOp->SetSimpleFeedbackMessage(nullptr, FSlateColor(GetForegroundColor()), FormattedMessage);						
-						}
-					}
-				}
-			}
-
-			return;
-		}
-		
-		TSharedPtr<FCategoryDragDropOp> CategoryDropOp = DragDropEvent.GetOperationAs<FCategoryDragDropOp>();
-		if (CategoryDropOp.IsValid())
-		{
-			if (const TSharedPtr<FVariablesOutlinerCategoryItem> CategoryItem = TreeItem.Pin())
-			{
-				TSharedPtr<FVariablesOutlinerCategoryItem> DropCategoryItem = CategoryDropOp->WeakItem.Pin();
-				if (DropCategoryItem.IsValid())
-				{
-					const UAnimNextRigVMAsset* ThisOwner = CategoryItem->WeakOwner.Get();
-					const UAnimNextRigVMAsset* OtherOwner = DropCategoryItem->WeakOwner.Get();
-
-					if (ThisOwner && ThisOwner == OtherOwner)
-					{
-						const FText OtherCategoryText = FText::FromString(DropCategoryItem->GetDisplayString());
-						const FText ThisCategoryText = FText::FromString(CategoryItem->GetDisplayString());
-						
-						const FText FormattedMessage = FText::Format(LOCTEXT("ReorderCategory", "Reorder {0} before {1}"), OtherCategoryText, ThisCategoryText);
-						CategoryDropOp->SetToolTip(FormattedMessage, nullptr);
-					}
-				}
-			}
-		}
-	}
-
-	virtual void OnDragLeave( const FDragDropEvent& DragDropEvent ) override
-	{
-		TSharedPtr<FVariableDragDropOp> GraphDropOp = DragDropEvent.GetOperationAs<FVariableDragDropOp>();
-		if (GraphDropOp.IsValid())
-		{
-			GraphDropOp->HoverTargetChanged();
-		}
-	}
-
-	TWeakPtr<FVariablesOutlinerCategoryItem> TreeItem;
-	TSharedPtr<SInlineEditableTextBlock> TextBlock;
-};
-
-FVariablesOutlinerCategoryItem::FVariablesOutlinerCategoryItem(const FItemData& InItemData) : ISceneOutlinerTreeItem(FVariablesOutlinerCategoryItem::Type), WeakOwner(InItemData.InOwner), CategoryName(InItemData.InCategoryName), ParentCategoryName(InItemData.InParentCategoryName), CategoryPath(InItemData.InCategoryPath)
-{
-}
-
-bool FVariablesOutlinerCategoryItem::IsValid() const
-{
-	return WeakOwner.Get() != nullptr;
-}
-
-FSceneOutlinerTreeItemID FVariablesOutlinerCategoryItem::GetID() const
-{
-	UAnimNextRigVMAsset* Owner = WeakOwner.Get();
-	const FSoftObjectPath OwnerPath = Owner;
-	return HashCombine(GetTypeHash(OwnerPath), GetTypeHash(CategoryPath));	
-}
-
-FString FVariablesOutlinerCategoryItem::GetDisplayString() const
-{
-	return CategoryName;
-}
-
-TSharedRef<SWidget> FVariablesOutlinerCategoryItem::GenerateLabelWidget(ISceneOutliner& Outliner, const STableRow<FSceneOutlinerTreeItemPtr>& InRow)
-{
-	TSharedRef<SVariablesOutlinerCategoryLabel> LabelWidget = SNew(SVariablesOutlinerCategoryLabel, *this, Outliner);
-	RenameRequestEvent.BindSP(LabelWidget->TextBlock.Get(), &SInlineEditableTextBlock::EnterEditingMode);
-
-	return LabelWidget;
-}
-
-FString FVariablesOutlinerCategoryItem::GetPackageName() const
-{
-	if (UAnimNextRigVMAsset* Owner = WeakOwner.Get())
-	{
-		return Owner->GetPackage()->GetName();
-	}
-	
-	return ISceneOutlinerTreeItem::GetPackageName();
 }
 
 void FVariablesOutlinerCategoryItem::Rename(const FText& InNewName) const
@@ -297,18 +31,40 @@ void FVariablesOutlinerCategoryItem::Rename(const FText& InNewName) const
 		if (UAnimNextRigVMAssetEditorData* EditorData = UE::UAF::UncookedOnly::FUtils::GetEditorData(Owner))
 		{
 			FScopedTransaction Transaction(LOCTEXT("RenameCategory", "Rename category"));
-			EditorData->RenameCategory(CategoryName, InNewName.ToString());						
+			
+			FString NewCategoryPath = ParentCategoryName;
+			if (!NewCategoryPath.IsEmpty())
+			{
+				NewCategoryPath.Append(TEXT("|"));
+			}
+			NewCategoryPath.Append(InNewName.ToString());
+
+			EditorData->RenameCategory(CategoryPath, NewCategoryPath);
 		}
-	}	
+	}
 }
 
 bool FVariablesOutlinerCategoryItem::ValidateName(const FText& InNewName, FText& OutErrorMessage) const
 {
+	if (InNewName.IsEmpty())
+	{		
+		OutErrorMessage = LOCTEXT("NameEmpty", "Category name cannot be empty");
+		return false;
+	}
+
 	if(UAnimNextRigVMAsset* Owner = WeakOwner.Get())
 	{
 		if (UAnimNextRigVMAssetEditorData* EditorData = UE::UAF::UncookedOnly::FUtils::GetEditorData(Owner))
 		{
-			if (EditorData->VariableAndFunctionCategories.Contains(InNewName.ToString()))
+			FString NewCategoryPath = ParentCategoryName;
+			if (!NewCategoryPath.IsEmpty())
+			{
+				NewCategoryPath.Append(TEXT("|"));
+			}
+			NewCategoryPath.Append(InNewName.ToString());
+		
+			if (EditorData->VariableAndFunctionCategories.Contains(NewCategoryPath) && CategoryPath != NewCategoryPath
+			)
 			{
 				OutErrorMessage = LOCTEXT("NameExistsError", "Category name already exists in this asset");
 				return false;
diff --git a/Engine/Plugins/Experimental/UAF/UAF/Source/UAFEditor/Private/Variables/Outliner/VariablesOutlinerCategoryItem.h b/Engine/Plugins/Experimental/UAF/UAF/Source/UAFEditor/Private/Variables/Outliner/VariablesOutlinerCategoryItem.h
--- a/Engine/Plugins/Experimental/UAF/UAF/Source/UAFEditor/Private/Variables/Outliner/VariablesOutlinerCategoryItem.h
+++ b/Engine/Plugins/Experimental/UAF/UAF/Source/UAFEditor/Private/Variables/Outliner/VariablesOutlinerCategoryItem.h
@@ -3,46 +3,21 @@
 #pragma once
 
 #include "ISceneOutlinerTreeItem.h"
+#include "Common/Outliner/OutlinerCategoryItem.h"
 
 class UAnimNextRigVMAsset;
 
 namespace UE::UAF::Editor
 {
 
-struct FVariablesOutlinerCategoryItem : ISceneOutlinerTreeItem
+struct FVariablesOutlinerCategoryItem : FOutlinerCategoryItem
 {
 	static const FSceneOutlinerTreeItemType Type;
 
-	struct FItemData
-	{
-		UAnimNextRigVMAsset* InOwner;
-		FStringView InCategoryName;
-		FStringView InParentCategoryName;
-		FStringView InCategoryPath;
-	};
-
-	FVariablesOutlinerCategoryItem(const FItemData& InItemData);
-
-	// Begin ISceneOutlinerTreeItem overrides
-	virtual bool IsValid() const override;
-	virtual FSceneOutlinerTreeItemID GetID() const override;
-	virtual FString GetDisplayString() const override;
-	virtual bool CanInteract() const override { return true; }
-	virtual TSharedRef<SWidget> GenerateLabelWidget(ISceneOutliner& Outliner, const STableRow<FSceneOutlinerTreeItemPtr>& InRow) override;
-	virtual FString GetPackageName() const override;
-	// End ISceneOutlinerTreeItem overrides
-
-	// Renames the item to the specified name
-	void Rename(const FText& InNewName) const;
-
-	// Validates the new item name
-	bool ValidateName(const FText& InNewName, FText& OutErrorMessage) const;
-	
-	TWeakObjectPtr<UAnimNextRigVMAsset> WeakOwner;
-	FString CategoryName;
-	FString ParentCategoryName;
-	FString CategoryPath;
-	int32 SortValue = 0;
+	FVariablesOutlinerCategoryItem(const FOutlinerCategoryItem::FItemData& IntemData);
+
+	virtual void Rename(const FText& InNewName) const override final;
+	virtual bool ValidateName(const FText& InNewName, FText& OutErrorMessage) const override final;
 };
 
 }
diff --git a/Engine/Plugins/Experimental/UAF/UAF/Source/UAFEditor/Private/Variables/Outliner/VariablesOutlinerColumns.cpp b/Engine/Plugins/Experimental/UAF/UAF/Source/UAFEditor/Private/Variables/Outliner/VariablesOutlinerColumns.cpp
--- a/Engine/Plugins/Experimental/UAF/UAF/Source/UAFEditor/Private/Variables/Outliner/VariablesOutlinerColumns.cpp
+++ b/Engine/Plugins/Experimental/UAF/UAF/Source/UAFEditor/Private/Variables/Outliner/VariablesOutlinerColumns.cpp
@@ -13,7 +13,6 @@
 #include "ScopedTransaction.h"
 #include "SPinTypeSelector.h"
 #include "UncookedOnlyUtils.h"
-#include "VariablesOutlinerAssetItem.h"
 #include "VariablesOutlinerEntryItem.h"
 #include "Entries/AnimNextSharedVariablesEntry.h"
 #include "Entries/AnimNextRigVMAssetEntry.h"
@@ -25,6 +24,7 @@
 #include "Widgets/Input/SButton.h"
 #include "Widgets/Input/SCheckBox.h"
 #include "Variables/Outliner/VariablesOutlinerMode.h"
+#include "Common/Outliner/OutlinerAssetItem.h"
 
 #define LOCTEXT_NAMESPACE "VariablesOutlinerColumns"
 
@@ -123,7 +123,7 @@ SHeaderRow::FColumn::FArguments FVariablesOutlinerTypeColumn::ConstructHeaderRow
 
 const TSharedRef<SWidget> FVariablesOutlinerTypeColumn::ConstructRowWidget(FSceneOutlinerTreeItemRef Item, const STableRow<FSceneOutlinerTreeItemPtr>& Row)
 {
-	const FVariablesOutlinerAssetItem* AssetItem = Item->CastTo<FVariablesOutlinerAssetItem>();
+	const FOutlinerAssetItem* AssetItem = Item->CastTo<FOutlinerAssetItem>();
 	if (AssetItem)
 	{
 		return SNullWidget::NullWidget;
@@ -417,29 +417,7 @@ class SVariablesOutlinerValue : public SCompoundWidget, public FNotifyHook
 	}
 
 	// FNotifyHook interface
-	virtual void NotifyPostChange(const FPropertyChangedEvent& PropertyChangedEvent, FProperty* PropertyThatChanged) override
-	{
-		TSharedPtr<FVariablesOutlinerEntryItem> TreeItem = WeakTreeItem.Pin();
-		if(!TreeItem.IsValid())
-		{
-			return;
-		}
-
-		UAnimNextVariableEntry* VariableEntry = Cast<UAnimNextVariableEntry>(TreeItem->WeakEntry.Get());
-		if(VariableEntry == nullptr)
-		{
-			return;
-		}
-
-		UAnimNextRigVMAssetEditorData* EditorData = VariableEntry->GetTypedOuter<UAnimNextRigVMAssetEditorData>();
-		if (EditorData == nullptr)
-		{
-			return;
-		}
-
-		VariableEntry->MarkPackageDirty();
-		VariableEntry->BroadcastModified(EAnimNextEditorDataNotifType::VariableDefaultValueChanged);
-	}
+	virtual void NotifyPostChange(const FPropertyChangedEvent& PropertyChangedEvent, FProperty* PropertyThatChanged) override;
 
 	void HandleModified(UAnimNextRigVMAssetEditorData* InEditorData, EAnimNextEditorDataNotifType InType, UObject* InSubject)
 	{
@@ -471,185 +449,51 @@ class SVariablesOutlinerValue : public SCompoundWidget, public FNotifyHook
 	FInstancedPropertyBag InternalPropertyBag;
 };
 
-const TSharedRef<SWidget> FVariablesOutlinerValueColumn::ConstructRowWidget(FSceneOutlinerTreeItemRef Item, const STableRow<FSceneOutlinerTreeItemPtr>& Row)
+void SVariablesOutlinerValue::NotifyPostChange(const FPropertyChangedEvent& PropertyChangedEvent, FProperty* PropertyThatChanged)
 {
-	const FVariablesOutlinerAssetItem* AssetItem = Item->CastTo<FVariablesOutlinerAssetItem>();
-	if (AssetItem)
+	TSharedPtr<FVariablesOutlinerEntryItem> TreeItem = WeakTreeItem.Pin();
+	if(!TreeItem.IsValid())
 	{
-		return SNullWidget::NullWidget;
+		return;
 	}
 
-	// Return a dummy box widget, with ShouldGenerateEmptyWidgetForSpacing settings this will prevent us from skipping outside of the header, avoiding alignment issues
-	FVariablesOutlinerEntryItem* TreeItem = Item->CastTo<FVariablesOutlinerEntryItem>();
-	if (TreeItem == nullptr)
+	UAnimNextVariableEntry* VariableEntry = Cast<UAnimNextVariableEntry>(TreeItem->WeakEntry.Get());
+	if(VariableEntry == nullptr)
 	{
-		return SNew(SBox);
+		return;
 	}
-	IAnimNextRigVMVariableInterface* Variable = Cast<IAnimNextRigVMVariableInterface>(TreeItem->WeakEntry.Get());
-	if(Variable == nullptr)
+
+	UAnimNextRigVMAssetEditorData* EditorData = VariableEntry->GetTypedOuter<UAnimNextRigVMAssetEditorData>();
+	if (EditorData == nullptr)
 	{
-		return SNew(SBox);
+		return;
 	}
-	
-	return SNew(SVariablesOutlinerValue, *TreeItem, WeakSceneOutliner.Pin().ToSharedRef().Get(), Row);
-}
-
-FLazyName VariablesOutlinerAccessSpecifier("AccessSpecifier");
 
-FName FVariablesOutlinerAccessSpecifierColumn::GetID()
-{
-	return VariablesOutlinerAccessSpecifier;
+	VariableEntry->MarkPackageDirty();
+	VariableEntry->BroadcastModified(EAnimNextEditorDataNotifType::VariableDefaultValueChanged);
 }
 
-SHeaderRow::FColumn::FArguments FVariablesOutlinerAccessSpecifierColumn::ConstructHeaderRowColumn()
-{
-	return
-		SHeaderRow::Column(GetColumnID())
-		.FixedWidth(24.0f)
-		.HAlignHeader(HAlign_Center)
-		.VAlignHeader(VAlign_Center)
-		.HAlignCell(HAlign_Fill)
-		.VAlignCell(VAlign_Fill)
-		[
-			SNew(SBox) 
-			.VAlign(VAlign_Center)
-			.HAlign(HAlign_Center)
-			[
-				SNew(SImage)
-				.ColorAndOpacity(FSlateColor::UseForeground())
-				.Image(FAppStyle::GetBrush("Level.VisibleIcon16x"))
-				.ToolTipText(LOCTEXT("AccessSpecifierAccessLevelTooltip", "Access level of this entry"))
-			]
-		];
-}
-
-class SVariablesOutlinerAccessSpecifier : public SCompoundWidget, public FNotifyHook
+const TSharedRef<SWidget> FVariablesOutlinerValueColumn::ConstructRowWidget(FSceneOutlinerTreeItemRef Item, const STableRow<FSceneOutlinerTreeItemPtr>& Row)
 {
-	SLATE_BEGIN_ARGS(SVariablesOutlinerAccessSpecifier) {}
-	SLATE_END_ARGS()
-
-	void Construct(const FArguments& InArgs, FVariablesOutlinerEntryItem& InTreeItem, ISceneOutliner& SceneOutliner, const STableRow<FSceneOutlinerTreeItemPtr>& InRow)
-	{
-		WeakTreeItem = StaticCastSharedRef<FVariablesOutlinerEntryItem>(InTreeItem.AsShared());
-
-		ChildSlot
-		[
-			SNew(SBox)
-			.IsEnabled(!InTreeItem.WeakSharedVariablesEntry.IsValid())
-			.WidthOverride(16.0f)
-			.HeightOverride(16.0f)
-			.VAlign(VAlign_Center)
-			.HAlign(HAlign_Center)
-			.IsEnabled(this, &SVariablesOutlinerAccessSpecifier::IsButtonEnabled)
-			[
-				SNew(SButton)
-				.ButtonStyle(FAppStyle::Get(), "HoverHintOnly")
-				.OnClicked(this, &SVariablesOutlinerAccessSpecifier::OnClicked)
-				.Content()
-				[
-					SNew(SImage)
-					.ColorAndOpacity(FSlateColor::UseForeground())
-					.Image(this, &SVariablesOutlinerAccessSpecifier::GetImage)
-					.ToolTipText(this, &SVariablesOutlinerAccessSpecifier::GetTooltipText)
-				]
-			]
-		];
-	}
-
-	FReply OnClicked() const
-	{
-		if (TSharedPtr<FVariablesOutlinerEntryItem> TreeItem = WeakTreeItem.Pin())
-		{
-			if (TSharedPtr<ISceneOutliner> SceneOutliner = TreeItem->WeakSceneOutliner.Pin())
-			{
-				if (const FVariablesOutlinerMode* Mode = static_cast<const FVariablesOutlinerMode*>(SceneOutliner->GetMode()))
-				{
-					Mode->ToggleSelectedVariablesExportAccessSpecifiers(TreeItem);
-				}
-			}
-		}
-
-		return FReply::Unhandled();	// Fall through so we dont deselect our item
-	}
-
-	const FSlateBrush* GetImage() const
-	{
-		TSharedPtr<FVariablesOutlinerEntryItem> TreeItem = WeakTreeItem.Pin();
-		if(!TreeItem.IsValid())
-		{
-			return nullptr;
-		}
-
-		IAnimNextRigVMExportInterface* Export = Cast<IAnimNextRigVMExportInterface>(TreeItem->WeakEntry.Get());
-		if(Export == nullptr)
-		{
-			return nullptr;
-		}
-
-		return Export->GetExportAccessSpecifier() == EAnimNextExportAccessSpecifier::Public ?
-			FAppStyle::GetBrush("Level.VisibleIcon16x") :
-			FAppStyle::GetBrush("Level.NotVisibleHighlightIcon16x");
-	}
-
-	bool IsButtonEnabled() const
-	{
-		TSharedPtr<FVariablesOutlinerEntryItem> TreeItem = WeakTreeItem.Pin();
-		if(!TreeItem.IsValid())
-		{
-			return false;
-		}
-
-		if (UAnimNextVariableEntry* Entry = TreeItem->WeakEntry.Get())
-		{
-			if (UAnimNextRigVMAsset* OuterAsset = Entry->GetTypedOuter<UAnimNextRigVMAsset>())
-			{
-				// Disable toggling public/private for pure UAnimNextSharedVariables objects				
-				return ExactCast<UAnimNextSharedVariables>(OuterAsset) == nullptr;	
-			}
-		}
-		
-		return false;		
-	}	
-
-	FText GetTooltipText() const
+	const FOutlinerAssetItem* AssetItem = Item->CastTo<FOutlinerAssetItem>();
+	if (AssetItem)
 	{
-		TSharedPtr<FVariablesOutlinerEntryItem> TreeItem = WeakTreeItem.Pin();
-		if(!TreeItem.IsValid())
-		{
-			return FText::GetEmpty();
-		}
-
-		IAnimNextRigVMExportInterface* Export = Cast<IAnimNextRigVMExportInterface>(TreeItem->WeakEntry.Get());
-		if(Export == nullptr)
-		{
-			return FText::GetEmpty();
-		}
-
-		FText AccessSpecifier = Export->GetExportAccessSpecifier() == EAnimNextExportAccessSpecifier::Public ? LOCTEXT("PublicSpecifier", "public") : LOCTEXT("PrivateSpecifier", "private");
-		FText AccessSpecifierDesc = Export->GetExportAccessSpecifier() == EAnimNextExportAccessSpecifier::Public ?
-			LOCTEXT("PublicSpecifierDesc", "This means that the entry is usable from gameplay and from other AnimNext assets") :
-			LOCTEXT("PrivateSpecifierDesc", "This means that the entry is only usable inside this asset");
-		return FText::Format(LOCTEXT("AccessSpecifierEntryTooltip", "This entry is {0}.\n{1}"), AccessSpecifier, AccessSpecifierDesc);
+		return SNullWidget::NullWidget;
 	}
 
-	TWeakPtr<FVariablesOutlinerEntryItem> WeakTreeItem;
-};
-
-const TSharedRef<SWidget> FVariablesOutlinerAccessSpecifierColumn::ConstructRowWidget(FSceneOutlinerTreeItemRef Item, const STableRow<FSceneOutlinerTreeItemPtr>& Row)
-{
+	// Return a dummy box widget, with ShouldGenerateEmptyWidgetForSpacing settings this will prevent us from skipping outside of the header, avoiding alignment issues
 	FVariablesOutlinerEntryItem* TreeItem = Item->CastTo<FVariablesOutlinerEntryItem>();
 	if (TreeItem == nullptr)
 	{
-		return SNullWidget::NullWidget;
+		return SNew(SBox);
 	}
-
 	IAnimNextRigVMVariableInterface* Variable = Cast<IAnimNextRigVMVariableInterface>(TreeItem->WeakEntry.Get());
 	if(Variable == nullptr)
 	{
-		return SNullWidget::NullWidget;
+		return SNew(SBox);
 	}
 	
-	return SNew(SVariablesOutlinerAccessSpecifier, *TreeItem, WeakSceneOutliner.Pin().ToSharedRef().Get(), Row);
+	return SNew(SVariablesOutlinerValue, *TreeItem, WeakSceneOutliner.Pin().ToSharedRef().Get(), Row);
 }
 
 }
diff --git a/Engine/Plugins/Experimental/UAF/UAF/Source/UAFEditor/Private/Variables/Outliner/VariablesOutlinerColumns.h b/Engine/Plugins/Experimental/UAF/UAF/Source/UAFEditor/Private/Variables/Outliner/VariablesOutlinerColumns.h
--- a/Engine/Plugins/Experimental/UAF/UAF/Source/UAFEditor/Private/Variables/Outliner/VariablesOutlinerColumns.h
+++ b/Engine/Plugins/Experimental/UAF/UAF/Source/UAFEditor/Private/Variables/Outliner/VariablesOutlinerColumns.h
@@ -56,23 +56,6 @@ private:
 	TWeakPtr<ISceneOutliner> WeakSceneOutliner;
 };
 
-class FVariablesOutlinerAccessSpecifierColumn : public ISceneOutlinerColumn
-{
-public:
-	static FName GetID();
-
-	FVariablesOutlinerAccessSpecifierColumn(ISceneOutliner& SceneOutliner)
-		: WeakSceneOutliner(StaticCastSharedRef<ISceneOutliner>(SceneOutliner.AsShared()))
-	{}
 
-	// ISceneOutlinerColumn interface
-	virtual FName GetColumnID() override { return GetID(); }
-	virtual SHeaderRow::FColumn::FArguments ConstructHeaderRowColumn() override;
-	virtual const TSharedRef<SWidget> ConstructRowWidget(FSceneOutlinerTreeItemRef TreeItem, const STableRow<FSceneOutlinerTreeItemPtr>& Row) override;
-	virtual bool SupportsSorting() const override { return false; }
-
-private:
-	TWeakPtr<ISceneOutliner> WeakSceneOutliner;
-};
 
 }
diff --git a/Engine/Plugins/Experimental/UAF/UAF/Source/UAFEditor/Private/Variables/Outliner/VariablesOutlinerCommands.cpp b/Engine/Plugins/Experimental/UAF/UAF/Source/UAFEditor/Private/Variables/Outliner/VariablesOutlinerCommands.cpp
--- a/Engine/Plugins/Experimental/UAF/UAF/Source/UAFEditor/Private/Variables/Outliner/VariablesOutlinerCommands.cpp
+++ b/Engine/Plugins/Experimental/UAF/UAF/Source/UAFEditor/Private/Variables/Outliner/VariablesOutlinerCommands.cpp
@@ -10,15 +10,6 @@ void FVariablesOutlinerCommands::RegisterCommands()
 	UI_COMMAND(AddNewVariable, "Add Variable", "Adds a single new variable to the selected asset.", EUserInterfaceActionType::Button, FInputChord());
 	UI_COMMAND(AddNewVariables, "Add Variable(s)", "Adds variables to assets.\nIf multiple assets are selected, then variables will be added to each.\nIf no assets are selected and there are multiple assets, variables will be added to all assets.", EUserInterfaceActionType::Button, FInputChord());
 
-	UI_COMMAND(FindReferences, "In Project", "Finds all references to the selected variable variables across assets in the project.", EUserInterfaceActionType::Button, FInputChord());
-	UI_COMMAND(FindReferencesInWorkspace, "In Workspace", "Finds all references to the selected variable variables across assets in the current workspace.", EUserInterfaceActionType::Button, FInputChord());
-	UI_COMMAND(FindReferencesInAsset, "In Asset", "Finds all references to the selected variable variables in the currently opened asset.", EUserInterfaceActionType::Button, FInputChord());
-
-	UI_COMMAND(SaveAsset, "Save Asset", "Save the selected asset.", EUserInterfaceActionType::Button, FInputChord());
-
-	UI_COMMAND(SetVariableExportPublic, "Sets Variable Export to Public", "Sets Variable to be Public.", EUserInterfaceActionType::Button, FInputChord());
-	UI_COMMAND(SetVariableExportPrivate, "Sets Variable Export to Private", "Sets Variable to be Private.", EUserInterfaceActionType::Button, FInputChord());
-
 	UI_COMMAND(CreateSharedVariablesAssets, "Create Shared Variables from Selection", "Creates a new Shared Variables asset from the selected set of variables.", EUserInterfaceActionType::Button, FInputChord());
 }
 
diff --git a/Engine/Plugins/Experimental/UAF/UAF/Source/UAFEditor/Private/Variables/Outliner/VariablesOutlinerCommands.h b/Engine/Plugins/Experimental/UAF/UAF/Source/UAFEditor/Private/Variables/Outliner/VariablesOutlinerCommands.h
--- a/Engine/Plugins/Experimental/UAF/UAF/Source/UAFEditor/Private/Variables/Outliner/VariablesOutlinerCommands.h
+++ b/Engine/Plugins/Experimental/UAF/UAF/Source/UAFEditor/Private/Variables/Outliner/VariablesOutlinerCommands.h
@@ -19,12 +19,6 @@ public:
 
 	TSharedPtr<FUICommandInfo> AddNewVariable;
 	TSharedPtr<FUICommandInfo> AddNewVariables;
-	TSharedPtr<FUICommandInfo> FindReferences;
-	TSharedPtr<FUICommandInfo> FindReferencesInWorkspace;
-	TSharedPtr<FUICommandInfo> FindReferencesInAsset;
-	TSharedPtr<FUICommandInfo> SaveAsset;	
-	TSharedPtr<FUICommandInfo> SetVariableExportPublic;
-	TSharedPtr<FUICommandInfo> SetVariableExportPrivate;
 	TSharedPtr<FUICommandInfo> CreateSharedVariablesAssets;
 };
 
diff --git a/Engine/Plugins/Experimental/UAF/UAF/Source/UAFEditor/Private/Variables/Outliner/VariablesOutlinerEntryItem.cpp b/Engine/Plugins/Experimental/UAF/UAF/Source/UAFEditor/Private/Variables/Outliner/VariablesOutlinerEntryItem.cpp
--- a/Engine/Plugins/Experimental/UAF/UAF/Source/UAFEditor/Private/Variables/Outliner/VariablesOutlinerEntryItem.cpp
+++ b/Engine/Plugins/Experimental/UAF/UAF/Source/UAFEditor/Private/Variables/Outliner/VariablesOutlinerEntryItem.cpp
@@ -3,24 +3,15 @@
 #include "VariablesOutlinerEntryItem.h"
 
 #include "ISceneOutliner.h"
-#include "Styling/SlateColor.h"
-#include "Misc/Optional.h"
 #include "SceneOutlinerStandaloneTypes.h"
-#include "AnimNextEditorModule.h"
 #include "AnimNextRigVMAssetEditorData.h"
 #include "EditorUtils.h"
 #include "ScopedTransaction.h"
 #include "UncookedOnlyUtils.h"
 #include "Entries/AnimNextRigVMAssetEntry.h"
 #include "UObject/Package.h"
-#include "Widgets/SCompoundWidget.h"
-#include "Widgets/DeclarativeSyntaxSupport.h"
-#include "Widgets/Images/SImage.h"
-#include "Widgets/Text/SInlineEditableTextBlock.h"
-#include "Widgets/Text/STextBlock.h"
-#include "Widgets/Views/STableRow.h"
 #include "Widgets/Views/SListView.h"
-#include "VariablesOutlinerDragDrop.h"
+#include "Common/Outliner/OutlinerDragDropOps.h"
 #include "Entries/AnimNextVariableEntry.h"
 
 #define LOCTEXT_NAMESPACE "VariablesOutlinerTreeItem"
@@ -28,319 +19,13 @@
 namespace UE::UAF::Editor
 {
 
-const FSceneOutlinerTreeItemType FVariablesOutlinerEntryItem::Type;
+const FSceneOutlinerTreeItemType FVariablesOutlinerEntryItem::Type(&FOutlinerEntryItem::Type);
 
-class SVariablesOutlinerEntryLabel : FSceneOutlinerCommonLabelData, public SCompoundWidget
-{
-	SLATE_BEGIN_ARGS(SVariablesOutlinerEntryLabel) {}
-	SLATE_END_ARGS()
-
-	void Construct(const FArguments& InArgs, FVariablesOutlinerEntryItem& InTreeItem, ISceneOutliner& SceneOutliner, const STableRow<FSceneOutlinerTreeItemPtr>& InRow)
-	{
-		WeakSceneOutliner = StaticCastSharedRef<ISceneOutliner>(SceneOutliner.AsShared());
-		TreeItem = StaticCastSharedRef<FVariablesOutlinerEntryItem>(InTreeItem.AsShared());
-
-		ChildSlot
-		[
-			SNew(SHorizontalBox)
-			+SHorizontalBox::Slot()
-			.AutoWidth()
-			.VAlign(VAlign_Center)
-			.Padding(2.0f, 2.0f)
-			[
-				SAssignNew(TextBlock, SInlineEditableTextBlock)
-				.IsReadOnly(this, &SVariablesOutlinerEntryLabel::IsReadOnly)
-				.Text(this, &SVariablesOutlinerEntryLabel::GetDisplayText)
-				.HighlightText(SceneOutliner.GetFilterHighlightText())
-				.ColorAndOpacity(this, &SVariablesOutlinerEntryLabel::GetForegroundColor)
-				.OnTextCommitted(this, &SVariablesOutlinerEntryLabel::OnTextCommited)
-				.OnVerifyTextChanged(this, &SVariablesOutlinerEntryLabel::OnVerifyTextChanged)
-			]
-			+SHorizontalBox::Slot()
-			.VAlign(VAlign_Center)
-			.Padding(2.0f, 0.0f, 2.0f, 3.0f)
-			.AutoWidth()
-			[
-				SNew(SImage)
-				.ToolTipText(this, &SVariablesOutlinerEntryLabel::GetDirtyTooltipText)
-				.Image(this, &SVariablesOutlinerEntryLabel::GetDirtyImageBrush)
-			]
-		];
-	}
-
-	FText GetDirtyTooltipText() const
-	{
-		if (const TSharedPtr<FVariablesOutlinerEntryItem> Item = TreeItem.Pin())
-		{
-			FTextBuilder TextBuilder;
-			TextBuilder.AppendLine(LOCTEXT("ModifiedTooltip", "Modified"));
-
-			if(UAnimNextVariableEntry* AssetEntry = Item->WeakEntry.Get())
-			{
-				const UPackage* ExternalPackage = AssetEntry->GetExternalPackage();
-				if(ExternalPackage && ExternalPackage->IsDirty())
-				{
-					TextBuilder.AppendLine(FText::FromName(ExternalPackage->GetFName()));
-				}
-			}
-
-			return TextBuilder.ToText();
-		}
-		return FText::GetEmpty();
-	}
-
-	const FSlateBrush* GetDirtyImageBrush() const
-	{
-		if (const TSharedPtr<FVariablesOutlinerEntryItem> Item = TreeItem.Pin())
-		{
-			if(UAnimNextVariableEntry* AssetEntry = Item->WeakEntry.Get())
-			{
-				bool bIsDirty = false;
-				const UPackage* ExternalPackage = AssetEntry->GetExternalPackage();
-				if(ExternalPackage && ExternalPackage->IsDirty())
-				{
-					bIsDirty = true;
-				}
-
-				return bIsDirty ? FAppStyle::GetBrush("Icons.DirtyBadge") : nullptr;
-			}
-		}
-		return nullptr;
-	}
-	
-	FText GetDisplayText() const
-	{
-		if (const TSharedPtr<FVariablesOutlinerEntryItem> Item = TreeItem.Pin())
-		{
-			return FText::FromString(Item->GetDisplayString());
-		}
-		return FText();
-	}
-
-	void OnTextCommited(const FText& InLabel, ETextCommit::Type InCommitInfo) const
-	{
-		if(InCommitInfo == ETextCommit::OnEnter)
-		{
-			if (const TSharedPtr<FVariablesOutlinerEntryItem> Item = TreeItem.Pin())
-			{
-				Item->Rename(InLabel); 
-			}
-		}
-	}
-
-	bool OnVerifyTextChanged(const FText& InLabel, FText& OutErrorMessage) const
-	{
-		if (const TSharedPtr<FVariablesOutlinerEntryItem> Item = TreeItem.Pin())
-		{
-			return Item->ValidateName(InLabel, OutErrorMessage); 
-		}
-		return false;
-	}
-
-	bool IsReadOnly() const
-	{
-		if (const TSharedPtr<FVariablesOutlinerEntryItem> Item = TreeItem.Pin())
-		{
-			return !Item->PropertyPath.IsPathToFieldEmpty() || Item->WeakEntry.Get() == nullptr;
-		}
-		return false;
-	}
-	
-	virtual FSlateColor GetForegroundColor() const override
-	{
-		const TOptional<FLinearColor> BaseColor = FSceneOutlinerCommonLabelData::GetForegroundColor(*TreeItem.Pin());
-		return BaseColor.IsSet() ? BaseColor.GetValue() : FSlateColor::UseForeground();
-	}
-	
-	virtual FReply OnDrop(const FGeometry& MyGeometry, const FDragDropEvent& DragDropEvent) override
-	{
-		TSharedPtr<FVariableDragDropOp> GraphDropOp = DragDropEvent.GetOperationAs<FVariableDragDropOp>();
-		if (GraphDropOp.IsValid())
-		{
-			if (const TSharedPtr<FVariablesOutlinerEntryItem> ThisEntryItem = TreeItem.Pin())
-			{
-				const UAnimNextVariableEntry* ThisEntry = ThisEntryItem->WeakEntry.Get();
-				TSharedPtr<FVariablesOutlinerEntryItem> EntryItem = GraphDropOp->WeakItem.Pin();
-				if (ThisEntryItem->HasStructOwner())
-				{
-					const FSlateBrush* ErrorSymbol = FAppStyle::GetBrush(TEXT("Graph.ConnectorFeedback.Error"));
-					GraphDropOp->SetSimpleFeedbackMessage(ErrorSymbol, FSlateColor(GetForegroundColor()), LOCTEXT("StructItemFeedback", "Cannot make changes to Struct Entries"));
-				}
-				else if (EntryItem.IsValid())
-				{
-					UAnimNextVariableEntry* OtherEntry = EntryItem->WeakEntry.Get();
-					if (ThisEntry && OtherEntry && ThisEntry != OtherEntry)
-					{
-						UAnimNextRigVMAsset* ThisAsset = ThisEntry->GetTypedOuter<UAnimNextRigVMAsset>();
-						const UAnimNextRigVMAsset* OtherAsset = OtherEntry->GetTypedOuter<UAnimNextRigVMAsset>();
-						if (ThisAsset && ThisAsset == OtherAsset)
-						{
-							if (ThisEntry->GetVariableCategory() != OtherEntry->GetVariableCategory())
-							{
-								if (OtherEntry && ThisEntry)
-								{
-									UAnimNextRigVMAssetEditorData* VariableEditorData = OtherEntry->GetTypedOuter<UAnimNextRigVMAssetEditorData>();
-									UAnimNextRigVMAssetEditorData* ThisEditorData = ThisEntry->GetTypedOuter<UAnimNextRigVMAssetEditorData>();
-
-									if (VariableEditorData && ThisEditorData)
-									{
-										if (VariableEditorData == ThisEditorData)
-										{
-											// Add or remove to category of this variable entry
-											const bool bRemovingFromCategory = ThisEntry->GetVariableCategory().IsEmpty();
-
-											const FText OtherVariableNameText = FText::FromName(OtherEntry->GetVariableName());					
-											const FText OtherVariableCategoryText = FText::FromString(OtherEntry->GetVariableCategory().GetData());
-											const FText ThisVariableNameText = FText::FromName(ThisEntry->GetVariableName());
-											const FText ThisVariableCategoryText = FText::FromString(ThisEntry->GetVariableCategory().GetData());											
-											const FText AddFormattedMessage = FText::Format(LOCTEXT("AddVariableToCategoryFormat", "Add {0} to Category {1}"), OtherVariableNameText, ThisVariableCategoryText);
-											const FText RemoveFormattedMessage = FText::Format(LOCTEXT("RemoveVariableToCategoryFormat", "Remove {0} from Category {1}"), OtherVariableNameText, OtherVariableCategoryText);
-											
-											FScopedTransaction Transaction(bRemovingFromCategory ? RemoveFormattedMessage : AddFormattedMessage);
-											OtherEntry->SetVariableCategory(ThisEntry->GetVariableCategory(), true);		
-										}
-									}
-								}
-							}
-							else if (OtherEntry)
-							{
-								UAnimNextVariableEntry* VariableEntry = EntryItem->WeakEntry.Get();
-								const UAnimNextVariableEntry* ThisVariableEntry = ThisEntryItem->WeakEntry.Get();
-								if (VariableEntry && ThisVariableEntry)
-								{
-									const FText OtherVariableNameText = FText::FromName(OtherEntry->GetVariableName());					
-									const FText ThisVariableNameText = FText::FromName(ThisEntry->GetVariableName());
-									const FText FormattedMessage = FText::Format(LOCTEXT("ReorderVariableFormat", "Reorder {0} before {1}"), OtherVariableNameText, ThisVariableNameText);
-									FScopedTransaction Transaction(FormattedMessage);
-									
-									UAnimNextRigVMAssetEditorData* VariableEditorData = VariableEntry->GetTypedOuter<UAnimNextRigVMAssetEditorData>();
-									UAnimNextRigVMAssetEditorData* ThisEditorData = ThisVariableEntry->GetTypedOuter<UAnimNextRigVMAssetEditorData>();
-									if (VariableEditorData && VariableEditorData == ThisEditorData)
-									{
-										VariableEditorData->ReorderVariable(VariableEntry, ThisVariableEntry);
-									}
-								}								
-							}
-						}
-						else if (ThisAsset)
-						{
-							if (UAnimNextVariableEntry* VariableEntry = EntryItem->WeakEntry.Get())
-							{
-								const FText VariableNameText = FText::FromName(OtherEntry->GetVariableName());
-								const FText AssetNameText = FText::FromName(ThisAsset->GetFName());
-																		
-								const FText FormattedMessage = FText::Format(LOCTEXT("MoveVariableToAssetFormatTransaction", "Moving {0} to {1}"), VariableNameText, AssetNameText);
-								FScopedTransaction Transaction(FormattedMessage);
-								UncookedOnly::FUtils::MoveVariableToAsset(VariableEntry, ThisAsset, true, true);
-							}
-						}	
-					}
-				}
-			}
-		}
-		
-		return FReply::Unhandled();
-	}
-
-	virtual void OnDragEnter(const FGeometry& MyGeometry, const FDragDropEvent& DragDropEvent) override
-	{
-		TSharedPtr<FVariableDragDropOp> GraphDropOp = DragDropEvent.GetOperationAs<FVariableDragDropOp>();
-		if (GraphDropOp.IsValid())
-		{
-			if (const TSharedPtr<FVariablesOutlinerEntryItem> ThisEntryItem = TreeItem.Pin())
-			{
-				const UAnimNextVariableEntry* ThisEntry = ThisEntryItem->WeakEntry.Get();
-				TSharedPtr<FVariablesOutlinerEntryItem> EntryItem = GraphDropOp->WeakItem.Pin();
-
-				const FSlateBrush* ErrorSymbol = FAppStyle::GetBrush(TEXT("Graph.ConnectorFeedback.Error"));
-				if (ThisEntryItem->HasStructOwner())
-				{
-					GraphDropOp->SetSimpleFeedbackMessage(ErrorSymbol, FSlateColor(GetForegroundColor()), LOCTEXT("StructItemFeedback", "Cannot make changes to Struct Entries"));
-				}
-				else if (ThisEntry && EntryItem.IsValid())
-				{
-					if (const UAnimNextVariableEntry* OtherEntry = EntryItem->WeakEntry.Get())
-					{
-						const FText OtherVariableNameText = FText::FromName(OtherEntry->GetVariableName());					
-						const FText OtherVariableCategoryText = FText::FromString(OtherEntry->GetVariableCategory().GetData());
-						const FText ThisVariableNameText = FText::FromName(ThisEntry->GetVariableName());
-						const FText ThisVariableCategoryText = FText::FromString(ThisEntry->GetVariableCategory().GetData());
-						
-						if (ThisEntry && ThisEntry != OtherEntry)
-						{
-							UAnimNextRigVMAsset* ThisAsset = ThisEntryItem->WeakEntry.Get()->GetTypedOuter<UAnimNextRigVMAsset>();
-							const UAnimNextRigVMAsset* OtherAsset = EntryItem->WeakEntry.Get()->GetTypedOuter<UAnimNextRigVMAsset>();
-							if (ThisAsset && ThisAsset == OtherAsset)
-							{
-								if (ThisEntry->GetVariableCategory() != OtherEntry->GetVariableCategory())
-								{
-									if (ThisEntry->GetVariableCategory().IsEmpty())
-									{
-										// Remove from category
-										const FText FormattedMessage = FText::Format(LOCTEXT("RemoveVariableToCategoryFormat", "Remove {0} from Category {1}"), OtherVariableNameText, OtherVariableCategoryText);
-										GraphDropOp->SetSimpleFeedbackMessage(nullptr, FSlateColor(GetForegroundColor()), FormattedMessage);
-									}
-									else
-									{
-										// Add to category
-										const FText FormattedMessage = FText::Format(LOCTEXT("AddVariableToCategoryFormat", "Add {0} to Category {1}"), OtherVariableNameText, ThisVariableCategoryText);
-										GraphDropOp->SetSimpleFeedbackMessage(nullptr, FSlateColor(GetForegroundColor()), FormattedMessage);
-									}
-								}
-								else
-								{
-									// Reorder
-									const FText FormattedMessage = FText::Format(LOCTEXT("ReorderVariableFormat", "Reorder {0} before {1}"), OtherVariableNameText, ThisVariableNameText);
-									GraphDropOp->SetSimpleFeedbackMessage(nullptr, FSlateColor(GetForegroundColor()), FormattedMessage);
-								}
-							}
-							else if (ThisAsset)
-							{
-								// Reorder
-								GraphDropOp->SetSimpleFeedbackMessage(nullptr, FSlateColor(GetForegroundColor()), LOCTEXT("ReorderDifferentOuterVariableLabel", "Cannot Reorder Variable in different Asset"));
-
-								const FText VariableNameText = FText::FromName(OtherEntry->GetVariableName());
-								const FText AssetNameText = FText::FromName(ThisAsset->GetFName());
-								
-								// [TODO] should we only allow move variables "up" the asset chain? If so validate that according to the workspace export chain, that VariableOuter is a child of Asset
-								const FText FormattedMessage = FText::Format(LOCTEXT("MoveVariableToAssetFormat", "Move {0} to {1}"), VariableNameText, AssetNameText);
-								GraphDropOp->SetSimpleFeedbackMessage(nullptr, FSlateColor(GetForegroundColor()), FormattedMessage);
-							}						
-						}
-						else
-						{
-							// Reorder
-							GraphDropOp->SetSimpleFeedbackMessage(ErrorSymbol, FSlateColor(GetForegroundColor()), LOCTEXT("ReorderSelfVariableLabel", "Cannot Reorder Variable before itself"));
-						}
-					}					
-				}
-			}
-		}
-	}
-
-	virtual void OnDragLeave( const FDragDropEvent& DragDropEvent ) override
-	{
-		TSharedPtr<FVariableDragDropOp> GraphDropOp = DragDropEvent.GetOperationAs<FVariableDragDropOp>();
-		if (GraphDropOp.IsValid())
-		{
-			GraphDropOp->HoverTargetChanged();
-		}
-	}
-
-	TWeakPtr<FVariablesOutlinerEntryItem> TreeItem;
-	TSharedPtr<SInlineEditableTextBlock> TextBlock;
-};
-
-FVariablesOutlinerEntryItem::FVariablesOutlinerEntryItem(UAnimNextVariableEntry* InEntry)
-	: ISceneOutlinerTreeItem(FVariablesOutlinerEntryItem::Type)
-	, WeakEntry(InEntry)
-{
-}
-
-FVariablesOutlinerEntryItem::FVariablesOutlinerEntryItem(const FProperty* InProperty)
-	: ISceneOutlinerTreeItem(FVariablesOutlinerEntryItem::Type)
-	, PropertyPath(InProperty)
-	, bStructOwner(true)
+FVariablesOutlinerEntryItem::FVariablesOutlinerEntryItem(const FVariablesOutlinerEntryItem::FItemData& InItemData)
+: FOutlinerEntryItem(FVariablesOutlinerEntryItem::Type, { InItemData.Entry ? InItemData.Entry ->GetTypedOuter<UAnimNextRigVMAsset>() : nullptr, InItemData.SortValue })
+	, WeakEntry(InItemData.Entry)
+	, PropertyPath(InItemData.Property)
+	, bStructOwner(InItemData.Property != nullptr)
 {
 }
 
@@ -375,13 +60,6 @@ FString FVariablesOutlinerEntryItem::GetDisplayString() const
 	return FString();
 }
 
-TSharedRef<SWidget> FVariablesOutlinerEntryItem::GenerateLabelWidget(ISceneOutliner& Outliner, const STableRow<FSceneOutlinerTreeItemPtr>& InRow)
-{
-	TSharedRef<SVariablesOutlinerEntryLabel> LabelWidget = SNew(SVariablesOutlinerEntryLabel, *this, Outliner, InRow);
-	RenameRequestEvent.BindSP(LabelWidget->TextBlock.Get(), &SInlineEditableTextBlock::EnterEditingMode);
-	return LabelWidget;
-}
-
 FString FVariablesOutlinerEntryItem::GetPackageName() const
 {
 	if(UAnimNextRigVMAssetEntry* Entry = WeakEntry.Get())
@@ -440,6 +118,58 @@ bool FVariablesOutlinerEntryItem::ValidateName(const FText& InNewName, FText& Ou
 	return true;
 }
 
+void FVariablesOutlinerEntryItem::SetAccessSpecifier(const EAnimNextExportAccessSpecifier& InSpecifier) const
+{
+	IAnimNextRigVMExportInterface* Export = Cast<IAnimNextRigVMExportInterface>(WeakEntry.Get());
+	if(Export == nullptr)
+	{
+		return;
+	}
+
+	FScopedTransaction Transaction(LOCTEXT("SetAccessSpecifierTransaction", "Set Access Specifier"));
+	Export->SetExportAccessSpecifier(InSpecifier);
+}
+
+EAnimNextExportAccessSpecifier FVariablesOutlinerEntryItem::GetAccessSpecifier() const
+{
+	IAnimNextRigVMExportInterface* Export = Cast<IAnimNextRigVMExportInterface>(WeakEntry.Get());
+	if(Export == nullptr)
+	{
+		return EAnimNextExportAccessSpecifier::Private;
+	}
+
+	FScopedTransaction Transaction(LOCTEXT("SetAccessSpecifierTransaction", "Set Access Specifier"));
+	return Export->GetExportAccessSpecifier();
+}
+
+bool FVariablesOutlinerEntryItem::CanSetAccessSpecifier() const
+{
+	if (UAnimNextVariableEntry* Entry = WeakEntry.Get())
+	{
+		if (UAnimNextRigVMAsset* OuterAsset = Entry->GetTypedOuter<UAnimNextRigVMAsset>())
+		{
+			// Disable toggling public/private for pure UAnimNextSharedVariables objects				
+			return ExactCast<UAnimNextSharedVariables>(OuterAsset) == nullptr;	
+		}
+	}
+
+	return false;
+}
+
+FStringView FVariablesOutlinerEntryItem::GetCategoryPath() const
+{
+	if (UAnimNextVariableEntry* Entry = WeakEntry.Get())
+	{
+		return Entry->GetVariableCategory();
+	}
+	
+	return FStringView();
+}
+
+bool FVariablesOutlinerEntryItem::IsReadOnly() const
+{
+	return !PropertyPath.IsPathToFieldEmpty() || WeakEntry.Get() == nullptr;
+}
 }
 
 #undef LOCTEXT_NAMESPACE // "VariablesOutlinerTreeItem"
\ No newline at end of file
diff --git a/Engine/Plugins/Experimental/UAF/UAF/Source/UAFEditor/Private/Variables/Outliner/VariablesOutlinerEntryItem.h b/Engine/Plugins/Experimental/UAF/UAF/Source/UAFEditor/Private/Variables/Outliner/VariablesOutlinerEntryItem.h
--- a/Engine/Plugins/Experimental/UAF/UAF/Source/UAFEditor/Private/Variables/Outliner/VariablesOutlinerEntryItem.h
+++ b/Engine/Plugins/Experimental/UAF/UAF/Source/UAFEditor/Private/Variables/Outliner/VariablesOutlinerEntryItem.h
@@ -3,6 +3,7 @@
 #pragma once
 
 #include "ISceneOutlinerTreeItem.h"
+#include "Common/Outliner/OutlinerEntryItem.h"
 
 class UAnimNextSharedVariablesEntry;
 class UAnimNextVariableEntry;
@@ -10,27 +11,39 @@ class UAnimNextVariableEntry;
 namespace UE::UAF::Editor
 {
 
-struct FVariablesOutlinerEntryItem : ISceneOutlinerTreeItem
+struct FVariablesOutlinerEntryItem : FOutlinerEntryItem
 {
 	static const FSceneOutlinerTreeItemType Type;
+	
+	struct FItemData 
+	{
+		UAnimNextVariableEntry* Entry;
+		const FProperty* Property;
+		int32 SortValue;
+	};
 
-	FVariablesOutlinerEntryItem(UAnimNextVariableEntry* InEntry);
-	FVariablesOutlinerEntryItem(const FProperty* InProperty);
+	FVariablesOutlinerEntryItem() = delete;
+	FVariablesOutlinerEntryItem(const FVariablesOutlinerEntryItem::FItemData& InItemData);
 
 	// Begin ISceneOutlinerTreeItem overrides
 	virtual bool IsValid() const override;
 	virtual FSceneOutlinerTreeItemID GetID() const override;
 	virtual FString GetDisplayString() const override;
 	virtual bool CanInteract() const override { return true; }
-	virtual TSharedRef<SWidget> GenerateLabelWidget(ISceneOutliner& Outliner, const STableRow<FSceneOutlinerTreeItemPtr>& InRow) override;
 	virtual FString GetPackageName() const override;
 	// End ISceneOutlinerTreeItem overrides
 
 	// Renames the item to the specified name
-	void Rename(const FText& InNewName) const;
+	virtual void Rename(const FText& InNewName) const override;
 
 	// Validates the new item name
-	bool ValidateName(const FText& InNewName, FText& OutErrorMessage) const;
+	virtual bool ValidateName(const FText& InNewName, FText& OutErrorMessage) const override;
+
+	virtual void SetAccessSpecifier(const EAnimNextExportAccessSpecifier& InSpecifier) const override;
+	virtual EAnimNextExportAccessSpecifier GetAccessSpecifier() const override;
+	virtual bool CanSetAccessSpecifier() const override;
+	virtual FStringView GetCategoryPath() const override;
+	virtual bool IsReadOnly() const override;
 
 	bool HasStructOwner() const { return bStructOwner; }
 
@@ -44,8 +57,6 @@ struct FVariablesOutlinerEntryItem : ISceneOutlinerTreeItem
 	TFieldPath<const FProperty> PropertyPath;
 
 	bool bStructOwner = false;
-
-	int32 SortValue = 0;
 };
 
 }
diff --git a/Engine/Plugins/Experimental/UAF/UAF/Source/UAFEditor/Private/Variables/Outliner/VariablesOutlinerHierarchy.cpp b/Engine/Plugins/Experimental/UAF/UAF/Source/UAFEditor/Private/Variables/Outliner/VariablesOutlinerHierarchy.cpp
--- a/Engine/Plugins/Experimental/UAF/UAF/Source/UAFEditor/Private/Variables/Outliner/VariablesOutlinerHierarchy.cpp
+++ b/Engine/Plugins/Experimental/UAF/UAF/Source/UAFEditor/Private/Variables/Outliner/VariablesOutlinerHierarchy.cpp
@@ -5,11 +5,12 @@
 #include "AnimNextRigVMAsset.h"
 #include "ISceneOutlinerMode.h"
 #include "AnimNextRigVMAssetEditorData.h"
-#include "VariablesOutlinerAssetItem.h"
 #include "VariablesOutlinerCategoryItem.h"
 #include "VariablesOutlinerStructSharedVariablesItem.h"
 #include "VariablesOutlinerMode.h"
 #include "VariablesOutlinerEntryItem.h"
+#include "Common/Outliner/OutlinerAssetItem.h"
+#include "Common/Outliner/SCommonOutliner.h"
 #include "Entries/AnimNextSharedVariablesEntry.h"
 #include "Entries/AnimNextRigVMAssetEntry.h"
 #include "Entries/AnimNextVariableEntry.h"
@@ -19,42 +20,21 @@
 namespace UE::UAF::Editor
 {
 
-TAutoConsoleVariable<bool> GOutputHierarchyDebugInformationCvar(TEXT("VariablesOutliner.OutputHierarchyDebugInformation"), false, TEXT("If enabled will output debug information about the hierarchy generation to the output log")); 
-
 FVariablesOutlinerHierarchy::FVariablesOutlinerHierarchy(ISceneOutlinerMode* Mode)
-	: ISceneOutlinerHierarchy(Mode)
-{
-}
-
-void FVariablesOutlinerHierarchy::CreateItems(TArray<FSceneOutlinerTreeItemPtr>& OutItems) const
+	: FOutlinerHierarchy(Mode)
 {
-	SVariablesOutliner* Outliner = static_cast<FVariablesOutlinerMode* const>(Mode)->GetOutliner();
-
-	WeakProcessedAssets.Reset();
-	WeakProcessedStructs.Reset();
-
-	if(Outliner->Assets.Num())
-	{
-		// Process first asset only for now
-		const int32 NumAssets = Outliner->Assets.Num();
-		for(int32 Index = 0; Index < NumAssets; Index++)
-		{
-			const TSoftObjectPtr<UAnimNextRigVMAsset>& SoftAsset = Outliner->Assets[Index];
-			TArray<const FSoftObjectPath> SharedVariableSourcePaths;
-			ProcessAsset(SoftAsset, (NumAssets - Index), OutItems, SharedVariableSourcePaths);
-		}
-	}
 }
 
 FSceneOutlinerTreeItemPtr FVariablesOutlinerHierarchy::FindOrCreateParentItem(const ISceneOutlinerTreeItem& Item, const TMap<FSceneOutlinerTreeItemID, FSceneOutlinerTreeItemPtr>& Items, bool bCreate)
 {
 	uint32 ParentHash = INDEX_NONE;
+	FSoftObjectPath OwnerPath;
+	FStringView CategoryPath;
 	
 	if (const FVariablesOutlinerEntryItem* EntryItem = Item.CastTo<FVariablesOutlinerEntryItem>())
 	{
 		const UAnimNextSharedVariablesEntry* SharedVariablesEntry = EntryItem->WeakSharedVariablesEntry.Get();
 		
-		FSoftObjectPath OwnerPath;
 		if(SharedVariablesEntry != nullptr)
 		{
 			OwnerPath = SharedVariablesEntry->GetObjectPath();
@@ -68,7 +48,7 @@ FSceneOutlinerTreeItemPtr FVariablesOutlinerHierarchy::FindOrCreateParentItem(co
 			}
 
 			UAnimNextRigVMAsset* Asset = Entry->GetTypedOuter<UAnimNextRigVMAsset>();
-			if (Entry == nullptr)
+			if (Asset == nullptr)
 			{
 				return nullptr;
 			}
@@ -85,22 +65,23 @@ FSceneOutlinerTreeItemPtr FVariablesOutlinerHierarchy::FindOrCreateParentItem(co
 		{
 			if (!VariableEntry->GetVariableCategory().IsEmpty())
 			{
-				CategoryHash = GetTypeHash(VariableEntry->GetVariableCategory());
+				CategoryPath = VariableEntry->GetVariableCategory();
+				CategoryHash = GetTypeHash(CategoryPath);
 			}
 		}
 
 		check(OwnerPath.IsValid());
-		ParentHash = CategoryHash != INDEX_NONE ? HashCombine(GetTypeHash(OwnerPath), CategoryHash) : GetTypeHash(OwnerPath);		
-		
+		ParentHash = CategoryHash != INDEX_NONE ? HashCombine(GetTypeHash(OwnerPath), CategoryHash) : GetTypeHash(OwnerPath);
 	}
 
 	if (const FVariablesOutlinerCategoryItem* CategoryItem = Item.CastTo<FVariablesOutlinerCategoryItem>())
 	{
 		if (UAnimNextRigVMAsset* Asset = CategoryItem->WeakOwner.Get())
 		{
-			const FSoftObjectPath OwnerPath = Asset;
+			OwnerPath = Asset;
 			const uint32 AssetHash = GetTypeHash(OwnerPath);
 			const FStringView ParentCategory = CategoryItem->ParentCategoryName;
+			CategoryPath = ParentCategory;
 			ParentHash = ParentCategory.Len() ? HashCombine(AssetHash, GetTypeHash(ParentCategory)) : AssetHash;
 		}
 	}
@@ -110,36 +91,59 @@ FSceneOutlinerTreeItemPtr FVariablesOutlinerHierarchy::FindOrCreateParentItem(co
 		return *ParentItem;
 	}
 
-	return nullptr;
-}
+	TSoftObjectPtr<UAnimNextRigVMAsset> SoftAsset = Cast<UAnimNextRigVMAsset>(OwnerPath.ResolveObject());
 
-void FVariablesOutlinerHierarchy::ProcessAsset(const TSoftObjectPtr<UAnimNextRigVMAsset>& InSoftAsset, int32 Depth, TArray<FSceneOutlinerTreeItemPtr>& OutItems, TArray<const FSoftObjectPath>& OutSharedVariableSourcePaths) const
-{
-	UAnimNextRigVMAsset* Asset = InSoftAsset.Get();
-	if(Asset == nullptr)
+	if (!CategoryPath.IsEmpty())
 	{
-		return;
-	}
+		UAnimNextRigVMAssetEditorData* EditorData = UncookedOnly::FUtils::GetEditorData(SoftAsset.Get());
+		const int32 CategoryIndex = EditorData->VariableAndFunctionCategories.IndexOfByKey(CategoryPath);
+		ensure(CategoryIndex != INDEX_NONE);
 		
-	const UAnimNextRigVMAssetEditorData* EditorData = UncookedOnly::FUtils::GetEditorData<UAnimNextRigVMAssetEditorData>(Asset);
-	if (EditorData == nullptr)
-	{
-		return;
-	}
+		FStringView ParentCategoryName;
+		FStringView SubCategoryName;
 
-	if (WeakProcessedAssets.Contains(Asset))
+		int32 Index = INDEX_NONE;
+		if (CategoryPath.FindLastChar('|', Index))
+		{
+			ParentCategoryName = CategoryPath.Left(Index);
+			SubCategoryName = CategoryPath.RightChop(Index+1);
+		}
+		else
+		{
+			SubCategoryName = CategoryPath;
+		}
+		
+		return Mode->CreateItemFor<FVariablesOutlinerCategoryItem>(FOutlinerCategoryItem::FItemData{UncookedOnly::FUtils::GetAsset<UAnimNextRigVMAsset>(EditorData), CategoryIndex, SubCategoryName, ParentCategoryName, CategoryPath}, bCreate);
+	}
+	else if (OwnerPath.IsValid())
 	{
-		return;
+		TArray<const FSoftObjectPath> SharedVariableSourcePaths;
+		
+		const UScriptStruct* Struct = Cast<UScriptStruct>(OwnerPath.ResolveObject());
+
+		if (SoftAsset.IsValid())
+		{
+			RetrieveSharedAssetPaths(SoftAsset, SharedVariableSourcePaths);
+			const FSceneOutlinerTreeItemPtr AssetItem = Mode->CreateItemFor<FOutlinerAssetItem>(FOutlinerAssetItem::FItemData{SoftAsset, 0, SharedVariableSourcePaths}, bCreate);
+			
+			UE_CLOG(GOutputHierarchyDebugInformationCvar.GetValueOnAnyThread(), LogAnimation, Warning, TEXT("* FOutlinerAssetItem: %s %i"), *SoftAsset.GetAssetName(), 0);
+			return AssetItem;
+		}
+		else if (Struct)
+		{
+			return Mode->CreateItemFor<FVariablesOutlinerStructSharedVariablesItem>(Struct, bCreate);
+		}
 	}
 
-	WeakProcessedAssets.Add(Asset);
+	return nullptr;
+}
 
-	UE_CLOG(GOutputHierarchyDebugInformationCvar.GetValueOnAnyThread(),LogAnimation, Warning, TEXT("ProcessAsset: %s %i"), *InSoftAsset.GetAssetName(), Depth);
-		
+void FVariablesOutlinerHierarchy::CreateItemsForAsset(const UAnimNextRigVMAssetEditorData* EditorData, TArray<FSceneOutlinerTreeItemPtr>& OutItems) const
+{
 	TArray<FSceneOutlinerTreeItemPtr> ParentItems;
-	EditorData->ForEachEntryOfType<UAnimNextSharedVariablesEntry>([this, &OutSharedVariableSourcePaths, &ParentItems, Depth, &OutItems](UAnimNextSharedVariablesEntry* InSharedVariablesEntry)
+	EditorData->ForEachEntryOfType<UAnimNextSharedVariablesEntry>([this, /*&OutSharedVariableSourcePaths,*/ &ParentItems, /*Depth,*/ &OutItems](UAnimNextSharedVariablesEntry* InSharedVariablesEntry)
 	{
-		SVariablesOutliner* Outliner = static_cast<FVariablesOutlinerMode* const>(Mode)->GetOutliner();
+		SCommonOutliner* Outliner = static_cast<FCommonOutlinerMode* const>(Mode)->GetOutliner();
 
 		switch (InSharedVariablesEntry->Type)
 		{
@@ -147,13 +151,11 @@ void FVariablesOutlinerHierarchy::ProcessAsset(const TSoftObjectPtr<UAnimNextRig
 			{
 				if (const UAnimNextSharedVariables* SharedVariables = InSharedVariablesEntry->GetAsset())
 				{
-					OutSharedVariableSourcePaths.Add(SharedVariables);
-				
 					// Would have expected shared variable (assets) to be part of the outer export chain, so should have been handled as regular assets.
-					if (!Outliner->Assets.Contains(SharedVariables))
+					if (!Outliner->GetAssetsView().Contains(SharedVariables))
 					{	
 						TArray<const FSoftObjectPath> SharedVariableSourcePaths;
-						ProcessAsset( const_cast<UAnimNextSharedVariables*>(SharedVariables), Depth - 1, ParentItems, SharedVariableSourcePaths);
+						ProcessAsset( const_cast<UAnimNextSharedVariables*>(SharedVariables), 0, ParentItems, SharedVariableSourcePaths);
 					}
 				}				
 			}
@@ -170,26 +172,22 @@ void FVariablesOutlinerHierarchy::ProcessAsset(const TSoftObjectPtr<UAnimNextRig
 
 					WeakProcessedStructs.Add(Struct);
 
-					if (FSceneOutlinerTreeItemPtr Item = Mode->CreateItemFor<FVariablesOutlinerStructSharedVariablesItem>(InSharedVariablesEntry))
+					if (FSceneOutlinerTreeItemPtr Item = Mode->CreateItemFor<FVariablesOutlinerStructSharedVariablesItem>(Struct))
 					{
 						OutItems.Add(Item);
 					}
 					
-					UE_CLOG(GOutputHierarchyDebugInformationCvar.GetValueOnAnyThread(), LogAnimation, Warning, TEXT("~ FVariablesOutlinerStructSharedVariablesItem: %s %i"), *InSharedVariablesEntry->GetObjectPath().ToString(), Depth);
+					UE_CLOG(GOutputHierarchyDebugInformationCvar.GetValueOnAnyThread(), LogAnimation, Warning, TEXT("~ FVariablesOutlinerStructSharedVariablesItem: %s %i"), *InSharedVariablesEntry->GetObjectPath().ToString(), 0);
 					
-					OutSharedVariableSourcePaths.Add(Struct);
 					int32 Index = 0;
 					for (TFieldIterator<FProperty> It(Struct); It; ++It)
 					{
 						const FProperty* Property = *It;
-						if (FSceneOutlinerTreeItemPtr Item = Mode->CreateItemFor<FVariablesOutlinerEntryItem>(Property))
+						if (FSceneOutlinerTreeItemPtr Item = Mode->CreateItemFor<FVariablesOutlinerEntryItem>(FVariablesOutlinerEntryItem::FItemData{nullptr, Property, Index++}))
 						{
 							FVariablesOutlinerEntryItem* EntryItem = Item->CastTo<FVariablesOutlinerEntryItem>();
 							EntryItem->WeakSharedVariablesEntry = InSharedVariablesEntry;
-							EntryItem->SortValue = Index++;
 							OutItems.Add(Item);
-
-							UE_CLOG(GOutputHierarchyDebugInformationCvar.GetValueOnAnyThread(), LogAnimation, Warning, TEXT("\t- FVariablesOutlinerEntryItem: %s %i"), *Property->GetName(), EntryItem->SortValue);
 						}
 					}
 				}
@@ -202,8 +200,7 @@ void FVariablesOutlinerHierarchy::ProcessAsset(const TSoftObjectPtr<UAnimNextRig
 
 	OutItems.Append(ParentItems);
 
-	TMap<FString, TSharedPtr<FVariablesOutlinerCategoryItem>> CategoryItems;
-
+	TSet<FString> ParsedCategories;
 	for (int32 CategoryIndex = 0; CategoryIndex < EditorData->VariableAndFunctionCategories.Num(); ++CategoryIndex)
 	{
 		const FString& CategoryName = EditorData->VariableAndFunctionCategories[CategoryIndex];
@@ -223,12 +220,12 @@ void FVariablesOutlinerHierarchy::ProcessAsset(const TSoftObjectPtr<UAnimNextRig
 			
 			const FString FullCategoryName = RebuildCategories;
 
-			if (!CategoryItems.Contains(FullCategoryName))
+			if (!ParsedCategories.Contains(FullCategoryName))
 			{
-				if (TSharedPtr<FVariablesOutlinerCategoryItem> Item = StaticCastSharedPtr<FVariablesOutlinerCategoryItem>( Mode->CreateItemFor<FVariablesOutlinerCategoryItem>(FVariablesOutlinerCategoryItem::FItemData{Asset, SubCategoryName, ParentCategoryName, FullCategoryName})))
+				if (TSharedPtr<FVariablesOutlinerCategoryItem> Item = StaticCastSharedPtr<FVariablesOutlinerCategoryItem>(Mode->CreateItemFor<FVariablesOutlinerCategoryItem>(FVariablesOutlinerCategoryItem::FItemData{UncookedOnly::FUtils::GetAsset<UAnimNextRigVMAsset>(EditorData), CategoryIndex, SubCategoryName, ParentCategoryName, FullCategoryName})))
 				{
-					Item->CastTo<FVariablesOutlinerCategoryItem>()->SortValue = CategoryIndex;
-					CategoryItems.Add(FullCategoryName, Item);
+					ParsedCategories.Add(FullCategoryName);
+					OutItems.AddUnique(Item);
 
 					UE_CLOG(GOutputHierarchyDebugInformationCvar.GetValueOnAnyThread(),LogAnimation, Warning, TEXT("+ FVariablesOutlinerCategoryItem: %s  %i"), *FullCategoryName, CategoryIndex);
 				}
@@ -236,47 +233,24 @@ void FVariablesOutlinerHierarchy::ProcessAsset(const TSoftObjectPtr<UAnimNextRig
 		}
 	}
 
-	// Adjust asset depth to ensure SharedVariables are always pinned to 0 depth
-	const int32 AssetDepth = Asset->GetClass() && Asset->GetClass() == UAnimNextSharedVariables::StaticClass() ? 0 : Depth;
-
 	constexpr int32 NumberOfVariableSortValueEntriesInCategory = 1024;
-
-	auto EnsureCategoryItemExists = [&CategoryItems, &OutItems](const FString& CategoryName)
-	{		 
-		TSharedPtr<FVariablesOutlinerCategoryItem>& CategoryItemToAdd = CategoryItems.FindChecked(CategoryName);
-		OutItems.AddUnique(CategoryItemToAdd);
-
-		auto CategoryPtr = CategoryItems.Find(CategoryItemToAdd->ParentCategoryName);
-		while (CategoryPtr)
-		{
-			OutItems.AddUnique(*CategoryPtr);
-			CategoryPtr = CategoryItems.Find((*CategoryPtr)->ParentCategoryName);
-		}
-	};
-
-	if (FSceneOutlinerTreeItemPtr AssetItem = Mode->CreateItemFor<FVariablesOutlinerAssetItem>(FVariablesOutlinerAssetItem::FItemData{InSoftAsset, AssetDepth, OutSharedVariableSourcePaths}))
+	
 	{
-		OutItems.Add(AssetItem);
-
-		UE_CLOG(GOutputHierarchyDebugInformationCvar.GetValueOnAnyThread(), LogAnimation, Warning, TEXT("* FVariablesOutlinerAssetItem: %s %i"), *InSoftAsset.GetAssetName(), AssetDepth);
-
 		int32 Index = 1;
-		EditorData->ForEachEntryOfType<UAnimNextVariableEntry>([this, &OutItems, EnsureCategoryItemExists, &Index, EditorData](UAnimNextVariableEntry* InVariable)
+		EditorData->ForEachEntryOfType<UAnimNextVariableEntry>([this, &OutItems, &Index, EditorData](UAnimNextVariableEntry* InVariable)
 		{
-			if (FSceneOutlinerTreeItemPtr Item = Mode->CreateItemFor<FVariablesOutlinerEntryItem>(CastChecked<UAnimNextVariableEntry>(InVariable)))
+			const int32 CategoryIndex = EditorData->VariableAndFunctionCategories.IndexOfByKey(InVariable->Category);
+			if (FSceneOutlinerTreeItemPtr Item = Mode->CreateItemFor<FVariablesOutlinerEntryItem>(
+				FVariablesOutlinerEntryItem::FItemData
+				{
+					CastChecked<UAnimNextVariableEntry>(InVariable),
+					nullptr,
+					(CategoryIndex != INDEX_NONE ? ((CategoryIndex + 1) * NumberOfVariableSortValueEntriesInCategory) : NumberOfVariableSortValueEntriesInCategory) + Index++
+				}))
 			{
-				FVariablesOutlinerEntryItem* EntryItem = Item->CastTo<FVariablesOutlinerEntryItem>();
-				const int32 CategoryIndex = EditorData->VariableAndFunctionCategories.IndexOfByKey(InVariable->Category);
-				EntryItem->SortValue = (CategoryIndex != INDEX_NONE ? ((CategoryIndex + 1) * NumberOfVariableSortValueEntriesInCategory) : 0) + Index++;
 				OutItems.Add(Item);
 				
 				UE_CLOG(GOutputHierarchyDebugInformationCvar.GetValueOnAnyThread(), LogAnimation, Warning, TEXT("\t- FVariablesOutlinerEntryItem: %s %i %s"), *InVariable->GetEntryName().ToString(), Item->CastTo<FVariablesOutlinerEntryItem>()->SortValue, InVariable->GetVariableCategory().GetData());
-
-				// Add category item on demand (to ensure we do not end up with empty category item hierarchies)
-				if (!InVariable->Category.IsEmpty())
-				{
-					EnsureCategoryItemExists(InVariable->Category);
-				}
 			}
 			return true;
 		});
diff --git a/Engine/Plugins/Experimental/UAF/UAF/Source/UAFEditor/Private/Variables/Outliner/VariablesOutlinerHierarchy.h b/Engine/Plugins/Experimental/UAF/UAF/Source/UAFEditor/Private/Variables/Outliner/VariablesOutlinerHierarchy.h
--- a/Engine/Plugins/Experimental/UAF/UAF/Source/UAFEditor/Private/Variables/Outliner/VariablesOutlinerHierarchy.h
+++ b/Engine/Plugins/Experimental/UAF/UAF/Source/UAFEditor/Private/Variables/Outliner/VariablesOutlinerHierarchy.h
@@ -4,6 +4,7 @@
 
 #include "CoreMinimal.h"
 #include "ISceneOutlinerHierarchy.h"
+#include "Common/Outliner/OutlinerHierarchy.h"
 
 class UAnimNextRigVMAssetEditorData;
 class UAnimNextRigVMAsset;
@@ -11,7 +12,7 @@ class UAnimNextRigVMAsset;
 namespace UE::UAF::Editor
 {
 
-class FVariablesOutlinerHierarchy : public ISceneOutlinerHierarchy
+class FVariablesOutlinerHierarchy : public FOutlinerHierarchy
 {
 public:
 	FVariablesOutlinerHierarchy(ISceneOutlinerMode* Mode);
@@ -19,15 +20,10 @@ public:
 	FVariablesOutlinerHierarchy& operator=(const FVariablesOutlinerHierarchy&) = delete;
 
 	// Begin ISceneOutlinerHierarchy overrides
-	virtual void CreateItems(TArray<FSceneOutlinerTreeItemPtr>& OutItems) const override;
-	virtual void CreateChildren(const FSceneOutlinerTreeItemPtr& Item, TArray<FSceneOutlinerTreeItemPtr>& OutChildren) const override {}
 	virtual FSceneOutlinerTreeItemPtr FindOrCreateParentItem(const ISceneOutlinerTreeItem& Item, const TMap<FSceneOutlinerTreeItemID, FSceneOutlinerTreeItemPtr>& Items, bool bCreate = false) override;
 	// End ISceneOutlinerHierarchy overrides
-protected:
-	void ProcessAsset(const TSoftObjectPtr<UAnimNextRigVMAsset>& InSoftAsset, int32 Depth, TArray<FSceneOutlinerTreeItemPtr>& OutItems, TArray<const FSoftObjectPath>& OutSharedVariableSourcePaths) const;
 
-	mutable TSet<TWeakObjectPtr<const UAnimNextRigVMAsset>> WeakProcessedAssets;
-	mutable TSet<TWeakObjectPtr<const UScriptStruct>> WeakProcessedStructs;
+	virtual void CreateItemsForAsset(const UAnimNextRigVMAssetEditorData* EditorData, TArray<FSceneOutlinerTreeItemPtr>& OutItems) const override;
 };
 
 }
\ No newline at end of file
