Added common outliner implementation to be shared between variables/functions sharing:

diff --git a/Engine/Plugins/Experimental/UAF/UAF/Source/UAFEditor/Private/Variables/Outliner/VariablesOutlinerMode.cpp b/Engine/Plugins/Experimental/UAF/UAF/Source/UAFEditor/Private/Variables/Outliner/VariablesOutlinerMode.cpp
--- a/Engine/Plugins/Experimental/UAF/UAF/Source/UAFEditor/Private/Variables/Outliner/VariablesOutlinerMode.cpp
+++ b/Engine/Plugins/Experimental/UAF/UAF/Source/UAFEditor/Private/Variables/Outliner/VariablesOutlinerMode.cpp
@@ -16,20 +16,20 @@
 #include "ToolMenus.h"
 #include "Framework/Commands/GenericCommands.h"
 #include "Entries/AnimNextRigVMAssetEntry.h"
-#include "Variables/AnimNextVariableItemMenuContext.h"
 #include "AnimNextRigVMAsset.h"
 #include "ContentBrowserModule.h"
 #include "FileHelpers.h"
 #include "IWorkspaceEditorModule.h"
 #include "StructViewerFilter.h"
 #include "StructViewerModule.h"
-#include "VariablesOutlinerAssetItem.h"
 #include "VariablesOutlinerCommands.h"
 #include "VariablesOutlinerStructSharedVariablesItem.h"
 #include "VariablesOutlinerCategoryItem.h"
-#include "VariablesOutlinerDragDrop.h"
+#include "Common/Outliner/OutlinerDragDropOps.h"
 #include "Common/AnimNextAssetFindReplaceVariables.h"
 #include "Common/GraphEditorSchemaActions.h"
+#include "Common/Outliner/OutlinerAssetItem.h"
+#include "Common/Outliner/OutlinerCommands.h"
 #include "Entries/AnimNextVariableEntry.h"
 #include "Variables/SVariablesView.h"
 #include "Entries/AnimNextSharedVariablesEntry.h"
@@ -39,216 +39,658 @@
 #include "Variables/AnimNextVariableSettings.h"
 #include "Widgets/Docking/SDockTab.h"
 #include "AnimNextScopedCompileJob.h"
+#include "Framework/Notifications/NotificationManager.h"
+#include "HAL/PlatformApplicationMisc.h"
+#include "Common/Outliner/OutlinerItemMenuContexts.h"
+#include "Common/Outliner/SCommonOutliner.h"
+#include "Widgets/Notifications/SNotificationList.h"
 
 #define LOCTEXT_NAMESPACE "VariablesOutlinerMode"
 
 namespace UE::UAF::Editor
 {
-
-FVariablesOutlinerMode::FVariablesOutlinerMode(SVariablesOutliner* InVariablesOutliner, const TSharedRef<UE::Workspace::IWorkspaceEditor>& InWorkspaceEditor)
-	: ISceneOutlinerMode(InVariablesOutliner)
-	, WeakWorkspaceEditor(InWorkspaceEditor)
+void FVariablesOutlinerMode::RegisterToolMenus()
 {
-	CommandList = MakeShared<FUICommandList>();
-
-	static const FName MenuName("VariablesOutliner.AddVariablesMenu");
+	FCommonOutlinerMode::RegisterToolMenus();
+	
 	UToolMenus* ToolMenus = UToolMenus::Get();
-	if (!ToolMenus->IsMenuRegistered(MenuName))
 	{
 		FToolMenuOwnerScoped ToolMenuOwnerScope(this);
+		
+		if (UToolMenu* Menu = ToolMenus->ExtendMenu(ContextToolMenuName))
+		{
+			FToolMenuSection& AssetSection = Menu->AddDynamicSection("DynamicAssetSection", FNewToolMenuDelegate::CreateLambda([](UToolMenu* InMenu)
+			{
+				const UCommonOutlinerItemMenuContext* MenuContext = InMenu->FindContext<UCommonOutlinerItemMenuContext>();			
+				const UVariablesOutlinerItemMenuContext* VariablesContext = InMenu->FindContext<UVariablesOutlinerItemMenuContext>();
+				if (MenuContext == nullptr || VariablesContext == nullptr)
+				{
+					return;
+				}
+				
+				TSharedPtr<SCommonOutliner> SharedOutliner = MenuContext->WeakOutliner.Pin();
+				if (SharedOutliner.IsValid() && MenuContext->WeakEditorDatas.Num() == 1)
+				{
+					FToolMenuSection& AssetsSection = InMenu->AddSection(ContextToolAssetSectionName, LOCTEXT("AssetsSectionLabel", "Assets"));
+				
+					const FVariablesOutlinerMode* Outliner = static_cast<const FVariablesOutlinerMode*>(SharedOutliner->GetMode());
+					AssetsSection.AddMenuEntry("AddNewCategory",  LOCTEXT("AddCategoryLabel", "Add New Category"), LOCTEXT("AddCategoryTooltip", "Adds a new Category to the selected Asset"), FSlateIcon(FAppStyle::Get().GetStyleSetName(), "Icons.PlusCircle"), FUIAction(FExecuteAction::CreateRaw(Outliner, &FVariablesOutlinerMode::AddCategory, MenuContext->WeakEditorDatas[0])));
+				}
+			}));
+			
+			
+			FToolMenuSection& VariablesSection = Menu->AddDynamicSection("Variables", FNewToolMenuDelegate::CreateLambda([](UToolMenu* InMenu)
+			{
+				const UAssetEditorToolkitMenuContext* EditorContext = InMenu->FindContext<UAssetEditorToolkitMenuContext>();
+				const UCommonOutlinerItemMenuContext* OutlinerContext = InMenu->FindContext<UCommonOutlinerItemMenuContext>();
+				const UVariablesOutlinerItemMenuContext* MenuContext = InMenu->FindContext<UVariablesOutlinerItemMenuContext>();
+				if(EditorContext == nullptr || OutlinerContext == nullptr || MenuContext == nullptr)
+				{
+					return;
+				}
+				
+				FToolMenuSection& VariablesSection = InMenu->AddSection("Variables", LOCTEXT("VariablesSectionLabel", "Variables"));
 
-		if (UToolMenu* Menu = ToolMenus->RegisterMenu(MenuName))
+				const int32 NumSelectedVariables = MenuContext->WeakEntries.Num();
+				/** Variable specific actions */			
+				{
+					// Populates entries for adding of different variable types/sources					
+					PopulateNewVariableToolMenuEntries(InMenu, false);
+					
+					if (NumSelectedVariables >= 1)
+					{
+						VariablesSection.AddMenuEntryWithCommandList(FVariablesOutlinerCommands::Get().CreateSharedVariablesAssets, OutlinerContext->WeakCommandList.Pin(), TAttribute<FText>(), TAttribute<FText>(), FSlateIcon(FUAFStyle::Get().GetStyleSetName(), "ClassIcon.AnimNextSharedVariables"));
+					}
+				}
+			}), FToolMenuInsert(FCommonOutlinerMode::ContextToolAssetSectionName, EToolMenuInsertType::After));
+		}
+		
+		
+		if (UToolMenu* Menu = ToolMenus->ExtendMenu(AddEntryContextMenuName))
 		{
-			Menu->AddDynamicSection(TEXT("Variables"), FNewToolMenuDelegate::CreateLambda([](UToolMenu* InMenu)
+			Menu->AddDynamicSection("VariablesDynamic", FNewToolMenuDelegate::CreateLambda([this](UToolMenu* InMenu)
 			{
-				PopulateNewVariableToolMenuEntries(InMenu, false);				
+				const UCommonOutlinerItemMenuContext* OutlinerContext = InMenu->FindContext<UCommonOutlinerItemMenuContext>();
+				const UVariablesOutlinerItemMenuContext* MenuContext = InMenu->FindContext<UVariablesOutlinerItemMenuContext>();
+				if(MenuContext == nullptr)
+				{
+					return;
+				}
+				
+				TSharedPtr<SCommonOutliner> SharedOutliner = OutlinerContext->WeakOutliner.Pin();
+				if (SharedOutliner.IsValid() && OutlinerContext->WeakEditorDatas.Num() == 1)
+				{
+					FToolMenuSection& AssetSection = InMenu->AddSection(ContextToolAssetSectionName, LOCTEXT("AssetSectionLabel", "Asset"));
+					{
+						const FVariablesOutlinerMode* Outliner = static_cast<const FVariablesOutlinerMode*>(SharedOutliner->GetMode());
+						AssetSection.AddMenuEntry("AddNewCategory",  LOCTEXT("AddCategoryLabel", "Add New Category"), LOCTEXT("AddCategoryTooltip", "Adds a new Category to the selected Asset"), FSlateIcon(FAppStyle::Get().GetStyleSetName(), "Icons.PlusCircle"), FUIAction(FExecuteAction::CreateRaw(Outliner, &FVariablesOutlinerMode::AddCategory, OutlinerContext->WeakEditorDatas[0])));
+					}
+				}
+			
+				PopulateNewVariableToolMenuEntries(InMenu, false);			
 			}));
 		}
 	}
 }
 
-void FVariablesOutlinerMode::Rebuild()
+void FVariablesOutlinerMode::InitToolMenuContext(FToolMenuContext& InOutContext) const
 {
-	Hierarchy = CreateHierarchy();
+	FCommonOutlinerMode::InitToolMenuContext(InOutContext);
+	
+	UVariablesOutlinerItemMenuContext* MenuContext = NewObject<UVariablesOutlinerItemMenuContext>();
+	TArray<FSceneOutlinerTreeItemPtr> SelectedItems = GetOutliner()->GetSelectedItems();
+	for(const FSceneOutlinerTreeItemPtr& Item : SelectedItems)
+	{
+		if (const FVariablesOutlinerEntryItem* EntryItem = Item->CastTo<FVariablesOutlinerEntryItem>())
+		{
+			UAnimNextRigVMAssetEntry* Entry = EntryItem->WeakEntry.Get();
+			if(Entry == nullptr)
+			{
+				continue;
+			}
+
+			MenuContext->WeakEntries.Add(Entry);
+		}
+	}
+	
+	InOutContext.AddObject(MenuContext);
 }
 
-TSharedPtr<SWidget> FVariablesOutlinerMode::CreateContextMenu()
+FVariablesOutlinerMode::FVariablesOutlinerMode(SCommonOutliner* InOutliner, const TSharedRef<Workspace::IWorkspaceEditor>& InWorkspaceEditor)
+	: FCommonOutlinerMode(InOutliner, InWorkspaceEditor)
 {
-	static const FName MenuName("VariablesOutliner.ItemContextMenu");
+	static const FName MenuName("VariablesOutliner.AddVariablesMenu");
 	UToolMenus* ToolMenus = UToolMenus::Get();
 	if (!ToolMenus->IsMenuRegistered(MenuName))
 	{
 		FToolMenuOwnerScoped ToolMenuOwnerScope(this);
+
 		if (UToolMenu* Menu = ToolMenus->RegisterMenu(MenuName))
 		{
-			Menu->AddDynamicSection(TEXT("Assets"), FNewToolMenuDelegate::CreateLambda([](UToolMenu* InMenu)
+			Menu->AddDynamicSection(TEXT("Variables"), FNewToolMenuDelegate::CreateLambda([](UToolMenu* InMenu)
 			{
-				const UAssetEditorToolkitMenuContext* EditorContext = InMenu->FindContext<UAssetEditorToolkitMenuContext>();
-				const UAnimNextVariableItemMenuContext* MenuContext = InMenu->FindContext<UAnimNextVariableItemMenuContext>();
-				if(EditorContext == nullptr || MenuContext == nullptr)
-				{
-					return;
-				}
+				PopulateNewVariableToolMenuEntries(InMenu, false);				
+			}));
+		}
+	}
+	
+	RegisterDragDropOperationConstructor<FVariablesOutlinerEntryItem>([](const FSceneOutlinerTreeItemPtr& InItem)-> TSharedPtr<FDragDropOperation>
+	{
+		FVariablesOutlinerEntryItem* VariableItem = InItem->CastTo<FVariablesOutlinerEntryItem>();
+		check(VariableItem);
+		
+		if(UAnimNextVariableEntry* Entry = Cast<UAnimNextVariableEntry>(VariableItem->WeakEntry.Get()))
+		{
+			const UAnimNextSharedVariables* Asset = nullptr;
+			if (VariableItem->WeakSharedVariablesEntry.Get())
+			{
+				Asset = VariableItem->WeakSharedVariablesEntry->GetAsset();
+			}
+			else
+			{
+				Asset = Entry->GetTypedOuter<UAnimNextSharedVariables>();
+			}
+	
+			if (Asset)
+			{
+				TSharedPtr<FAnimNextSchemaAction_Variable> Action = MakeShared<FAnimNextSchemaAction_Variable>(Entry->GetVariableName(), Asset, Entry->GetType(), FAnimNextSchemaAction_Variable::EVariableAccessorChoice::Deferred);
 
-				const int32 NumSelectedVariables = MenuContext->WeakEntries.Num();
+				return FVariableDragDropOp::New(Action, StaticCastSharedRef<FVariablesOutlinerEntryItem>(VariableItem->AsShared()));
+			}
+		}
+		else if (const FProperty* Property = VariableItem->PropertyPath.Get())
+		{
+			if (const UScriptStruct* Struct = Property->GetOwner<UScriptStruct>())
+			{
+				TSharedPtr<FAnimNextSchemaAction_Variable> Action = MakeShared<FAnimNextSchemaAction_Variable>(Property->GetFName(), Struct, FAnimNextParamType::FromProperty(Property), FAnimNextSchemaAction_Variable::EVariableAccessorChoice::Deferred);
+			
+				TSharedPtr<FVariableDragDropOp> Op = FVariableDragDropOp::New(Action, StaticCastSharedRef<FVariablesOutlinerEntryItem>(VariableItem->AsShared()));
+				return Op;
+			}
+		}
+		
+		return nullptr;
+	});
+	
+	RegisterDragDropOperationParser<FVariableDragDropOp>([](FSceneOutlinerDragDropPayload& OutPayload, const FDragDropOperation& InDragDropOperation)
+	{
+		check(InDragDropOperation.IsOfType<FVariableDragDropOp>());
+		const FVariableDragDropOp& VariableDragDropOp = static_cast<const FVariableDragDropOp&>(InDragDropOperation);
+
+		OutPayload.DraggedItems.Add(VariableDragDropOp.WeakItem.Pin());
+	});
+	
+	RegisterDragDropOperator<FVariablesOutlinerEntryItem, FOutlinerCategoryItem, FVariableDragDropOp>([](const ISceneOutlinerTreeItem& DragOverTreeItem, TSharedRef<FDragDropOperation> DragDropOperation, FSceneOutlinerDragValidationInfo& OutValidationInfo) -> FReply
+	{
+		TSharedRef<FVariableDragDropOp> VariableDropOp = StaticCastSharedRef<FVariableDragDropOp>(DragDropOperation);			
+		const FOutlinerCategoryItem* DragOverItem = DragOverTreeItem.CastTo<FOutlinerCategoryItem>();
+		
+		if (VariableDropOp->WeakItem.IsValid())
+		{
+			TSharedPtr<FVariablesOutlinerEntryItem> DropVariableItem = VariableDropOp->WeakItem.Pin();
+			
+			if (DropVariableItem->HasStructOwner())
+			{
+				const FText FormattedMessage = LOCTEXT("StructItemFeedback", "Cannot make changes to Struct Entries");						
+				OutValidationInfo = FSceneOutlinerDragValidationInfo(ESceneOutlinerDropCompatibility::IncompatibleGeneric, FormattedMessage);
 				
-				/** Asset specific actions */
-				FToolMenuSection& AssetSection = InMenu->AddSection("Asset", LOCTEXT("AssetSectionLabel", "Asset"));
+				return FReply::Handled();
+			}
+			
+			if (UAnimNextVariableEntry* VariableEntry = DropVariableItem->WeakEntry.Get())
+			{
+				const FText OtherVariableNameText = FText::FromName(VariableEntry->GetVariableName());
+				const FText ThisVariableCategoryText = FText::FromString(DragOverItem->GetDisplayString());
+
+				const UAnimNextRigVMAssetEditorData* VariableEditorData = VariableEntry->GetTypedOuter<UAnimNextRigVMAssetEditorData>();
+				const UAnimNextRigVMAssetEditorData* CategoryEditorData = UncookedOnly::FUtils::GetEditorData(DragOverItem->WeakOwner.Get());
+
+				if (VariableEditorData == CategoryEditorData)
 				{
-					AssetSection.AddMenuEntryWithCommandList(FVariablesOutlinerCommands::Get().SaveAsset, MenuContext->WeakCommandList.Pin(), TAttribute<FText>(), TAttribute<FText>(), FSlateIcon(FAppStyle::Get().GetStyleSetName(), "AssetEditor.SaveAsset"));
-					AssetSection.AddMenuEntryWithCommandList(FGenericCommands::Get().Paste, MenuContext->WeakCommandList.Pin());
+					if (VariableEntry->GetVariableCategory() == DragOverItem->CategoryPath)
+					{
+						const FText FormattedMessage = FText::Format(LOCTEXT("VariableAlreadyPartOfCategoryFormat", "{0} is already part of Category {1}"), OtherVariableNameText, ThisVariableCategoryText);
+						
+						OutValidationInfo = FSceneOutlinerDragValidationInfo(ESceneOutlinerDropCompatibility::Incompatible, FormattedMessage);
+					}
+					else if (!VariableEntry->GetVariableCategory().IsEmpty())
+					{
+						const FText FormattedMessage = FText::Format(LOCTEXT("MoveVariableToCategoryFormat", "Moving {0} to Category {1}"), OtherVariableNameText, ThisVariableCategoryText);
+						OutValidationInfo = FSceneOutlinerDragValidationInfo(ESceneOutlinerDropCompatibility::CompatibleGeneric, FormattedMessage);
+					}
+					else
+					{
+						const FText FormattedMessage = FText::Format(LOCTEXT("AddVariableToCategoryFormat", "Add {0} to Category {1}"), OtherVariableNameText, ThisVariableCategoryText);
+						OutValidationInfo = FSceneOutlinerDragValidationInfo(ESceneOutlinerDropCompatibility::CompatibleGeneric, FormattedMessage);
+					}
 				}
+				else
+				{
+					const FText ThisAssetText = FText::FromString(DragOverItem->GetDisplayString());
+					const FText FormattedMessage = FText::Format(LOCTEXT("MoveVariableToCategoryAndAssetFormat", "Moving {0} to Category {1} in {2}"), OtherVariableNameText, ThisVariableCategoryText, ThisAssetText);
+						OutValidationInfo = FSceneOutlinerDragValidationInfo(ESceneOutlinerDropCompatibility::CompatibleGeneric, FormattedMessage);		
+				}
+				
+				return FReply::Handled();
+			}
+		}
+		
+	
+		return FReply::Unhandled();
+	},
+	[this](ISceneOutlinerTreeItem& DroppedOnTreeItem, TSharedRef<const FDragDropOperation> DragDropOperation)
+	{			
+		TSharedRef<const FVariableDragDropOp> VariableDropOp = StaticCastSharedRef<const FVariableDragDropOp>(DragDropOperation);
+		const FOutlinerCategoryItem* DroppedOnCategory = DroppedOnTreeItem.CastTo<FOutlinerCategoryItem>();			
+		TSharedPtr<FVariablesOutlinerEntryItem> DropVariableItem = VariableDropOp->WeakItem.Pin();
+		
+		if (DropVariableItem.IsValid())
+		{
+			if (UAnimNextVariableEntry* VariableEntry = DropVariableItem->WeakEntry.Get())
+			{
+				UAnimNextRigVMAsset* VariableAsset = VariableEntry->GetTypedOuter<UAnimNextRigVMAsset>();
+				UAnimNextRigVMAsset* CategoryAsset = DroppedOnCategory->WeakOwner.Get();
 
-				/** Variable specific actions */
-				FToolMenuSection& VariablesSection = InMenu->AddSection("Variables", LOCTEXT("VariablesSectionLabel", "Variables"));
+				const FText OtherVariableNameText = FText::FromName(VariableEntry->GetVariableName());
+				const FText ThisVariableCategoryText = FText::FromString(DroppedOnCategory->GetDisplayString());
+				const FText FormattedMessage = FText::Format(LOCTEXT("MoveVariableToCategoryFormat", "Moving {0} to Category {1}"), OtherVariableNameText, ThisVariableCategoryText);	
+				FScopedTransaction Transaction(FormattedMessage);
+
+				UAnimNextVariableEntry* NewVariableEntry = VariableEntry;
+					
+				if (VariableAsset != CategoryAsset)
 				{
-					// Populates entries for adding of different variable types/sources					
-					PopulateNewVariableToolMenuEntries(InMenu, true);
+					const FName VariableEntryName = VariableEntry->GetVariableName();
+					UncookedOnly::FUtils::MoveVariableToAsset(VariableEntry, CategoryAsset, true, true);
+					NewVariableEntry = Cast<UAnimNextVariableEntry>(UncookedOnly::FUtils::GetEditorData(CategoryAsset)->FindEntry(VariableEntryName));
+				}
 					
-					VariablesSection.AddSubMenu("FindReferencesMenu", LOCTEXT("FindReferencesMenuLabel", "Find References..."), LOCTEXT("FindReferencesMenuToolTip", "Find references to selected variables by different means."),
-					FNewToolMenuDelegate::CreateLambda(
-						[MenuContext](UToolMenu* InSubmenu)
-							{
-								InSubmenu->AddMenuEntry(NAME_None, FToolMenuEntry::InitMenuEntryWithCommandList(FVariablesOutlinerCommands::Get().FindReferences, MenuContext->WeakCommandList.Pin(), TAttribute<FText>(), TAttribute<FText>(), UAnimNextAssetFindReplaceVariables::GetIconFromSearchScope(ESearchScope::Global)));
+				check(NewVariableEntry);
+					
+				// Only add if not already part of category
+				if (NewVariableEntry->GetVariableCategory() != DroppedOnCategory->CategoryPath)
+				{
+					NewVariableEntry->SetVariableCategory(DroppedOnCategory->CategoryPath);
+				}
+				
+				return true;
+			}
+		}
+		
+		return false;
+	});	
+	
+	RegisterDragDropOperator<FVariablesOutlinerEntryItem, FVariablesOutlinerEntryItem, FVariableDragDropOp>([](const ISceneOutlinerTreeItem& DragOverTreeItem, TSharedRef<FDragDropOperation> DragDropOperation, FSceneOutlinerDragValidationInfo& OutValidationInfo) -> FReply
+	{
+		TSharedRef<FVariableDragDropOp> VariableDropOp = StaticCastSharedRef<FVariableDragDropOp>(DragDropOperation);			
+		const FVariablesOutlinerEntryItem* DragOverItem = DragOverTreeItem.CastTo<FVariablesOutlinerEntryItem>();			
+		TSharedPtr<FVariablesOutlinerEntryItem> DropVariableItem = VariableDropOp->WeakItem.Pin();
 
-								InSubmenu->AddMenuEntry(NAME_None, FToolMenuEntry::InitMenuEntryWithCommandList(FVariablesOutlinerCommands::Get().FindReferencesInWorkspace, MenuContext->WeakCommandList.Pin(), TAttribute<FText>(), TAttribute<FText>(),
-									UAnimNextAssetFindReplaceVariables::GetIconFromSearchScope(ESearchScope::Workspace)));
-							
-								InSubmenu->AddMenuEntry(NAME_None, FToolMenuEntry::InitMenuEntryWithCommandList(FVariablesOutlinerCommands::Get().FindReferencesInAsset, MenuContext->WeakCommandList.Pin(), TAttribute<FText>(), TAttribute<FText>(),
-									UAnimNextAssetFindReplaceVariables::GetIconFromSearchScope(ESearchScope::Asset)));
-							}),
-							FUIAction(FExecuteAction(), FCanExecuteAction(), FGetActionCheckState(), FIsActionButtonVisible::CreateLambda([MenuContext]() -> bool
-							{
-								if (TSharedPtr<UE::UAF::Editor::SVariablesOutliner> SharedOutliner = MenuContext->WeakOutliner.Pin())
-								{
-									if (const FVariablesOutlinerMode* Mode = static_cast<const FVariablesOutlinerMode*>(SharedOutliner->GetMode()))
-									{
-										return Mode->CanFindReferences();
-									}
-								}
+		if (DragOverItem && DropVariableItem.IsValid())
+		{
+			const UAnimNextVariableEntry* ThisEntry = DragOverItem->WeakEntry.Get();
 
-								return false;
-							})),
-							EUserInterfaceActionType::Button,
-							false,
-							FSlateIcon(FAppStyle::GetAppStyleSetName(), TEXT("Icons.Search"))
-						);
-
-					VariablesSection.AddMenuEntryWithCommandList(FGenericCommands::Get().Rename, MenuContext->WeakCommandList.Pin());
-					VariablesSection.AddMenuEntryWithCommandList(FGenericCommands::Get().Delete, MenuContext->WeakCommandList.Pin());
-					VariablesSection.AddMenuEntryWithCommandList(FGenericCommands::Get().Duplicate, MenuContext->WeakCommandList.Pin());
+			if (DragOverItem->HasStructOwner() || DropVariableItem->HasStructOwner())
+			{
+				const FText FormattedMessage = LOCTEXT("StructItemFeedback", "Cannot make changes to Struct Entries");						
+				OutValidationInfo = FSceneOutlinerDragValidationInfo(ESceneOutlinerDropCompatibility::IncompatibleGeneric, FormattedMessage);
+			}
+			else if (ThisEntry)
+			{
+				if (const UAnimNextVariableEntry* OtherEntry = DropVariableItem->WeakEntry.Get())
+				{
+					const FText OtherVariableNameText = FText::FromName(OtherEntry->GetVariableName());					
+					const FText OtherVariableCategoryText = FText::FromStringView(OtherEntry->GetVariableCategory());
+					const FText ThisVariableNameText = FText::FromName(ThisEntry->GetVariableName());
+					const FText ThisVariableCategoryText = FText::FromStringView(ThisEntry->GetVariableCategory());
 					
-					VariablesSection.AddMenuEntryWithCommandList(FGenericCommands::Get().Copy, MenuContext->WeakCommandList.Pin());
-
-					// Check which export access related context menu entries we need to populate
-					if (NumSelectedVariables > 0)
+					if (ThisEntry && ThisEntry != OtherEntry)
 					{
-						bool bHasPublicVariables = false;
-						bool bHasPrivateVariables = false;
-						for (const TWeakObjectPtr<UAnimNextRigVMAssetEntry> WeakEntry : MenuContext->WeakEntries)
+						UAnimNextRigVMAsset* ThisAsset = DragOverItem->WeakEntry.Get()->GetTypedOuter<UAnimNextRigVMAsset>();
+						const UAnimNextRigVMAsset* OtherAsset = DropVariableItem->WeakEntry.Get()->GetTypedOuter<UAnimNextRigVMAsset>();
+						if (ThisAsset && ThisAsset == OtherAsset)
 						{
-							if (const UAnimNextVariableEntry* VariableEntry = Cast<UAnimNextVariableEntry>(WeakEntry))
+							if (ThisEntry->GetVariableCategory() != OtherEntry->GetVariableCategory())
 							{
-								if (VariableEntry->GetExportAccessSpecifier() == EAnimNextExportAccessSpecifier::Public)
+								if (ThisEntry->GetVariableCategory().IsEmpty())
 								{
-									bHasPublicVariables = true;
+									// Remove from category
+									const FText FormattedMessage = FText::Format(LOCTEXT("RemoveVariableToCategoryFormat", "Remove {0} from Category {1}"), OtherVariableNameText, OtherVariableCategoryText);
+									OutValidationInfo = FSceneOutlinerDragValidationInfo(ESceneOutlinerDropCompatibility::CompatibleGeneric, FormattedMessage);
 								}
-								else if (VariableEntry->GetExportAccessSpecifier() == EAnimNextExportAccessSpecifier::Private)
+								else
 								{
-									bHasPrivateVariables = true;
+									// Add to category
+									const FText FormattedMessage = FText::Format(LOCTEXT("AddVariableToCategoryFormat", "Add {0} to Category {1}"), OtherVariableNameText, ThisVariableCategoryText);
+									OutValidationInfo = FSceneOutlinerDragValidationInfo(ESceneOutlinerDropCompatibility::CompatibleGeneric, FormattedMessage);
 								}
 							}
-
-							if (bHasPrivateVariables && bHasPublicVariables)
+							else
 							{
-								break;
+								// Reorder
+								const FText FormattedMessage = FText::Format(LOCTEXT("ReorderVariableFormat", "Reorder {0} before {1}"), OtherVariableNameText, ThisVariableNameText);
+								OutValidationInfo = FSceneOutlinerDragValidationInfo(ESceneOutlinerDropCompatibility::CompatibleGeneric, FormattedMessage);
 							}
 						}
-
-						if (bHasPrivateVariables)
+						else if (ThisAsset)
 						{
-							VariablesSection.AddMenuEntryWithCommandList(FVariablesOutlinerCommands::Get().SetVariableExportPublic, MenuContext->WeakCommandList.Pin(), LOCTEXT("MakePublicMenuLabel", "Expose as Public"), TAttribute<FText>(), FSlateIcon(FAppStyle::Get().GetStyleSetName(), "Level.VisibleIcon16x"));
+							const FText VariableNameText = FText::FromName(OtherEntry->GetVariableName());
+							const FText AssetNameText = FText::FromName(ThisAsset->GetFName());
+							
+							// [TODO] should we only allow move variables "up" the asset chain? If so validate that according to the workspace export chain, that VariableOuter is a child of Asset
+							const FText FormattedMessage = FText::Format(LOCTEXT("MoveVariableToAssetFormat", "Move {0} to {1}"), VariableNameText, AssetNameText);
+							
+							OutValidationInfo = FSceneOutlinerDragValidationInfo(ESceneOutlinerDropCompatibility::CompatibleGeneric, FormattedMessage);
 						}
+					}
+					else
+					{
+						// Reorder
+						const FText FormattedMessage = LOCTEXT("ReorderSelfVariableLabel", "Cannot Reorder Variable before itself");
+						OutValidationInfo = FSceneOutlinerDragValidationInfo(ESceneOutlinerDropCompatibility::IncompatibleGeneric, FormattedMessage);
+					}
+				}					
+			}
+			
+			return FReply::Handled();
+		}
+	
+		return FReply::Unhandled();
+	},
+	[this](ISceneOutlinerTreeItem& DroppedOnTreeItem, TSharedRef<const FDragDropOperation> DragDropOperation)
+	{
+		TSharedRef<const FVariableDragDropOp> VariableDropOp = StaticCastSharedRef<const FVariableDragDropOp>(DragDropOperation);			
+		const FVariablesOutlinerEntryItem* DroppedOnItem = DroppedOnTreeItem.CastTo<FVariablesOutlinerEntryItem>();			
+		TSharedPtr<FVariablesOutlinerEntryItem> DropVariableItem = VariableDropOp->WeakItem.Pin();
 
-						if (bHasPublicVariables)
+		if (DroppedOnItem)
+		{
+			const UAnimNextVariableEntry* ThisEntry = DroppedOnItem->WeakEntry.Get();
+			TSharedPtr<FVariablesOutlinerEntryItem> EntryItem = VariableDropOp->WeakItem.Pin();
+
+			if (ThisEntry && EntryItem.IsValid())
+			{
+				if (UAnimNextVariableEntry* OtherEntry = EntryItem->WeakEntry.Get())
+				{
+					const FText OtherVariableNameText = FText::FromName(OtherEntry->GetVariableName());					
+					const FText OtherVariableCategoryText = FText::FromStringView(OtherEntry->GetVariableCategory());
+					const FText ThisVariableNameText = FText::FromName(ThisEntry->GetVariableName());
+					const FText ThisVariableCategoryText = FText::FromStringView(ThisEntry->GetVariableCategory());
+					
+					if (ThisEntry && ThisEntry != OtherEntry)
+					{
+						UAnimNextRigVMAsset* ThisAsset = DroppedOnItem->WeakEntry.Get()->GetTypedOuter<UAnimNextRigVMAsset>();
+						const UAnimNextRigVMAsset* OtherAsset = EntryItem->WeakEntry.Get()->GetTypedOuter<UAnimNextRigVMAsset>();
+						if (ThisAsset && ThisAsset == OtherAsset)
+						{
+							UAnimNextVariableEntry* VariableEntry = EntryItem->WeakEntry.Get();
+							const UAnimNextVariableEntry* ThisVariableEntry = DroppedOnItem->WeakEntry.Get();
+								
+							if (VariableEntry && ThisVariableEntry)
+							{	
+								if (ThisEntry->GetVariableCategory() != OtherEntry->GetVariableCategory())
+								{
+									UAnimNextRigVMAssetEditorData* VariableEditorData = OtherEntry->GetTypedOuter<UAnimNextRigVMAssetEditorData>();
+									UAnimNextRigVMAssetEditorData* ThisEditorData = ThisEntry->GetTypedOuter<UAnimNextRigVMAssetEditorData>();																			
+									if (VariableEditorData && ThisEditorData)
+									{
+										if (VariableEditorData == ThisEditorData)
+										{
+											// Add or remove to category of this variable entry
+											const bool bRemovingFromCategory = ThisEntry->GetVariableCategory().IsEmpty();
+																				
+											const FText AddFormattedMessage = FText::Format(LOCTEXT("AddVariableToCategoryFormat", "Add {0} to Category {1}"), OtherVariableNameText, ThisVariableCategoryText);
+											const FText RemoveFormattedMessage = FText::Format(LOCTEXT("RemoveVariableToCategoryFormat", "Remove {0} from Category {1}"), OtherVariableNameText, OtherVariableCategoryText);
+												
+											FScopedTransaction Transaction(bRemovingFromCategory ? RemoveFormattedMessage : AddFormattedMessage);
+											OtherEntry->SetVariableCategory(ThisEntry->GetVariableCategory(), true);		
+											// Also reorder variables such that dropped item is above dropped-on item
+											VariableEditorData->ReorderVariable(VariableEntry, ThisVariableEntry);
+											
+											return true;
+										}
+									}
+								}	
+								else
+								{
+									const FText FormattedMessage = FText::Format(LOCTEXT("ReorderVariableFormat", "Reorder {0} before {1}"), OtherVariableNameText, ThisVariableNameText);
+									FScopedTransaction Transaction(FormattedMessage);
+								
+									UAnimNextRigVMAssetEditorData* VariableEditorData = VariableEntry->GetTypedOuter<UAnimNextRigVMAssetEditorData>();
+									UAnimNextRigVMAssetEditorData* ThisEditorData = ThisVariableEntry->GetTypedOuter<UAnimNextRigVMAssetEditorData>();
+									if (VariableEditorData && VariableEditorData == ThisEditorData)
+									{
+										VariableEditorData->ReorderVariable(VariableEntry, ThisVariableEntry);
+										return true;
+									}
+								}	
+							}	
+						}
+						else if (ThisAsset)
 						{
-							VariablesSection.AddMenuEntryWithCommandList(FVariablesOutlinerCommands::Get().SetVariableExportPrivate, MenuContext->WeakCommandList.Pin(), LOCTEXT("MakePrivateMenuLabel", "Make Private"), TAttribute<FText>(), FSlateIcon(FAppStyle::Get().GetStyleSetName(), "Level.NotVisibleHighlightIcon16x"));
+							if (UAnimNextVariableEntry* VariableEntry = EntryItem->WeakEntry.Get())
+							{
+								const FText VariableNameText = FText::FromName(OtherEntry->GetVariableName());
+								const FText AssetNameText = FText::FromName(ThisAsset->GetFName());
+																	
+								const FText FormattedMessage = FText::Format(LOCTEXT("MoveVariableToAssetFormatTransaction", "Moving {0} to {1}"), VariableNameText, AssetNameText);
+								FScopedTransaction Transaction(FormattedMessage);
+								UncookedOnly::FUtils::MoveVariableToAsset(VariableEntry, ThisAsset, true, true);
+								
+								return true;
+							}
 						}
 					}
-					
-					AssetSection.AddMenuEntryWithCommandList(FVariablesOutlinerCommands::Get().CreateSharedVariablesAssets, MenuContext->WeakCommandList.Pin(), TAttribute<FText>(), TAttribute<FText>(), FSlateIcon(FUAFStyle::Get().GetStyleSetName(), "ClassIcon.AnimNextSharedVariables"));
-				}
-
-				if (TSharedPtr<SVariablesOutliner> Outliner = MenuContext->WeakOutliner.Pin())
-				{
-					Outliner->AddSourceControlMenuOptions(InMenu);
-				}
-			}));
+				}					
+			}
 		}
-	}
-
+	
+		return false;
+	});
+	
+	RegisterDragDropOperator<FVariablesOutlinerEntryItem, FVariablesOutlinerStructSharedVariablesItem, FVariableDragDropOp>([](const ISceneOutlinerTreeItem& DragOverTreeItem, TSharedRef<FDragDropOperation> DragDropOperation, FSceneOutlinerDragValidationInfo& OutValidationInfo) -> FReply
 	{
-		UAnimNextVariableItemMenuContext* MenuContext = NewObject<UAnimNextVariableItemMenuContext>();
-		MenuContext->WeakWorkspaceEditor = WeakWorkspaceEditor;
-		MenuContext->WeakOutliner = StaticCastSharedRef<SVariablesOutliner>(SceneOutliner->AsShared());
-		MenuContext->WeakCommandList = CommandList;
-		TArray<FSceneOutlinerTreeItemPtr> SelectedItems = GetOutliner()->GetSelectedItems();
-		for(const FSceneOutlinerTreeItemPtr& Item : SelectedItems)
+		const FText FormattedMessage = LOCTEXT("StructItemFeedback", "Cannot make changes to Struct Entries");						
+		OutValidationInfo = FSceneOutlinerDragValidationInfo(ESceneOutlinerDropCompatibility::IncompatibleGeneric, FormattedMessage);			
+		return FReply::Handled();
+	},
+	[this](ISceneOutlinerTreeItem& DroppedOnTreeItem, TSharedRef<const FDragDropOperation> DragDropOperation){ return false; }
+	);
+	
+	RegisterDragDropOperator<FVariablesOutlinerEntryItem, FOutlinerAssetItem, FVariableDragDropOp>([](const ISceneOutlinerTreeItem& DragOverTreeItem, TSharedRef<FDragDropOperation> DragDropOperation, FSceneOutlinerDragValidationInfo& OutValidationInfo) -> FReply
+	{	
+		TSharedRef<FVariableDragDropOp> VariableDropOp = StaticCastSharedRef<FVariableDragDropOp>(DragDropOperation);			
+		const FOutlinerAssetItem* DragOverItem = DragOverTreeItem.CastTo<FOutlinerAssetItem>();			
+		TSharedPtr<FVariablesOutlinerEntryItem> DropVariableItem = VariableDropOp->WeakItem.Pin();
+		
+		if (DropVariableItem->HasStructOwner())
 		{
-			if (const FVariablesOutlinerAssetItem* AssetItem = Item->CastTo<FVariablesOutlinerAssetItem>())
+			const FText FormattedMessage = LOCTEXT("StructItemFeedback", "Cannot make changes to Struct Entries");						
+			OutValidationInfo = FSceneOutlinerDragValidationInfo(ESceneOutlinerDropCompatibility::IncompatibleGeneric, FormattedMessage);
+			return FReply::Handled();
+		}
+		
+		if (const UAnimNextRigVMAsset* Asset = DragOverItem->SoftAsset.Get())
+		{
+			if (DropVariableItem.IsValid())
 			{
-				UAnimNextRigVMAsset* Asset = AssetItem->SoftAsset.Get();
-				if(Asset == nullptr)
+				if (const UAnimNextVariableEntry* VariableEntry = DropVariableItem->WeakEntry.Get())
 				{
-					continue;
-				}
-
-				UAnimNextRigVMAssetEditorData* EditorData = UncookedOnly::FUtils::GetEditorData<UAnimNextRigVMAssetEditorData>(Asset);
-				if(EditorData == nullptr)
-				{
-					continue;
-				}
+					const FText VariableNameText = FText::FromName(VariableEntry->GetVariableName());
+					const FText AssetNameText = FText::FromName(Asset->GetFName());
 				
-				MenuContext->WeakEditorDatas.Add(EditorData);
+					const UAnimNextRigVMAsset* VariableOuter = VariableEntry->GetTypedOuter<UAnimNextRigVMAsset>();
+					if (VariableOuter != Asset)
+					{
+						// [TODO] should we only allow move variables "up" the asset chain? If so validate that according to the workspace export chain, that VariableOuter is a child of Asset
+						const FText FormattedMessage = FText::Format(LOCTEXT("MoveVariableToAssetFormat", "Move {0} to {1}"), VariableNameText, AssetNameText);
+						OutValidationInfo = FSceneOutlinerDragValidationInfo(ESceneOutlinerDropCompatibility::CompatibleGeneric, FormattedMessage);	
+					}
+					else
+					{
+						const FStringView CategoryPath = VariableEntry->GetVariableCategory();
+						if (CategoryPath.IsEmpty())
+						{
+							const FText FormattedMessage = FText::Format(LOCTEXT("VariableAlreadyPartOfAssetFormat", "{0} is already part of {1}"), VariableNameText, AssetNameText);								
+							OutValidationInfo = FSceneOutlinerDragValidationInfo(ESceneOutlinerDropCompatibility::Incompatible, FormattedMessage);	
+						}
+						else
+						{
+							const FText FormattedMessage = FText::Format(LOCTEXT("RemoveVariableFromCategoryFormat", "Removing {0} from Category {1}"), VariableNameText, FText::FromStringView(CategoryPath));	
+							OutValidationInfo = FSceneOutlinerDragValidationInfo(ESceneOutlinerDropCompatibility::CompatibleGeneric, FormattedMessage);
+						}							
+					}
+					
+					return FReply::Handled();
+				}
 			}
-			else if (const FVariablesOutlinerEntryItem* EntryItem = Item->CastTo<FVariablesOutlinerEntryItem>())
+		}
+		
+		return FReply::Unhandled();
+	},
+	[this](ISceneOutlinerTreeItem& DroppedOnTreeItem, TSharedRef<const FDragDropOperation> DragDropOperation)
+	{
+		if (DragDropOperation->IsOfType<FVariableDragDropOp>() && DroppedOnTreeItem.IsA<FOutlinerAssetItem>())
+		{
+			TSharedRef<const FVariableDragDropOp> VariableDropOp = StaticCastSharedRef<const FVariableDragDropOp>(DragDropOperation);			
+			const FOutlinerAssetItem* DroppedOnItem = DroppedOnTreeItem.CastTo<FOutlinerAssetItem>();			
+			TSharedPtr<FVariablesOutlinerEntryItem> DropVariableItem = VariableDropOp->WeakItem.Pin();
+			
+			if (UAnimNextRigVMAsset* Asset = DroppedOnItem->SoftAsset.Get())
 			{
-				UAnimNextRigVMAssetEntry* Entry = EntryItem->WeakEntry.Get();
-				if(Entry == nullptr)
+				if (DropVariableItem.IsValid())
 				{
-					continue;
+					if (UAnimNextVariableEntry* VariableEntry = DropVariableItem->WeakEntry.Get())
+					{
+						const UAnimNextRigVMAsset* VariableOuter = VariableEntry->GetTypedOuter<UAnimNextRigVMAsset>();
+						const FText VariableNameText = FText::FromName(VariableEntry->GetVariableName());
+						
+						if (VariableOuter != Asset)
+						{								
+							const FText AssetNameText = FText::FromName(Asset->GetFName());						
+							const FText FormattedMessage = FText::Format(LOCTEXT("MoveVariableToAssetFormatTransaction", "Moving {0} to {1}"), VariableNameText, AssetNameText);
+							FScopedTransaction Transaction(FormattedMessage);
+							UncookedOnly::FUtils::MoveVariableToAsset(VariableEntry, Asset, true, true);
+							
+							return true;
+						}
+						else
+						{
+							const FStringView CategoryPath = VariableEntry->GetVariableCategory();
+							if (!CategoryPath.IsEmpty())
+							{
+								const FText FormattedMessage = FText::Format(LOCTEXT("RemoveVariableFromCategoryFormat", "Removing {0} from Category {1}"), VariableNameText, FText::FromStringView(CategoryPath));	
+								
+								FScopedTransaction Transaction(FormattedMessage);
+								VariableEntry->SetVariableCategory(TEXT(""));
+								
+								return true;
+							}							
+						}
+					}
 				}
+			}
+		}
+		
+		return false;
+	});
+}
 
-				MenuContext->WeakEntries.Add(Entry);
+bool FVariablesOutlinerMode::UpdateOperationDecorator(const FDragDropEvent& Event, const FSceneOutlinerDragValidationInfo* ValidationInfo) const
+{
+	TSharedPtr<FVariableDragDropOp> GraphDropOp = Event.GetOperationAs<FVariableDragDropOp>();
+	if (GraphDropOp.IsValid())
+	{
+		if (ValidationInfo)
+		{
+			const FSlateBrush* ValidHoverIcon = FAppStyle::GetBrush(TEXT("Graph.ConnectorFeedback.OK"));
+			const FSlateBrush* InvalidSymbol = FAppStyle::GetBrush(TEXT("Graph.ConnectorFeedback.Error"));			
+			GraphDropOp->SetSimpleFeedbackMessage(ValidationInfo->IsValid() ? ValidHoverIcon : InvalidSymbol, FLinearColor::White, ValidationInfo->ValidationText);		
+		}
+		else
+		{
+			GraphDropOp->HoverTargetChanged();
+		}
+		
+		return true;
+	}
+	
+	return false;
+}
+
+void FVariablesOutlinerMode::ReorderCategoryItem(TSharedPtr<FOutlinerCategoryItem> EntryItem, TSharedPtr<FOutlinerCategoryItem> BeforeItem) const
+{
+	const UAnimNextRigVMAssetEditorData* ThisEditorData = UncookedOnly::FUtils::GetEditorData(EntryItem->WeakOwner.Get());
+	 UAnimNextRigVMAssetEditorData* EditorData = UncookedOnly::FUtils::GetEditorData(BeforeItem->WeakOwner.Get());
+	if (EditorData && ThisEditorData)
+	{
+		if (EditorData == ThisEditorData)
+		{		
+			EditorData->ReorderCategory(EntryItem->CategoryPath, BeforeItem->CategoryPath);
+		}
+	}
+}
+
+void FVariablesOutlinerMode::ReorderCategory(const FString& CategoryPath, const FString& BeforeCategoryPath, UAnimNextRigVMAsset* Asset) const
+{
+	if (UAnimNextRigVMAssetEditorData* EditorData = UncookedOnly::FUtils::GetEditorData(Asset))
+	{
+		EditorData->ReorderCategory(CategoryPath, BeforeCategoryPath);
+	}
+}
+
+void FVariablesOutlinerMode::MoveCategoryToAsset(const FString& CategoryPath, UAnimNextRigVMAsset* FromAsset, UAnimNextRigVMAsset* ToAsset) const
+{
+	UAnimNextRigVMAssetEditorData* ThisEditorData = UncookedOnly::FUtils::GetEditorData(FromAsset);
+	UAnimNextRigVMAssetEditorData* EditorData = UncookedOnly::FUtils::GetEditorData(ToAsset);
+	
+	TArray<FString> CategoriesToMove;
+	Algo::TransformIf(ThisEditorData->VariableAndFunctionCategories, CategoriesToMove, [&CategoryPath](const FString& CategoryName)
+	{
+		return CategoryName.StartsWith(CategoryPath);
+	},
+	[](const FString& CategoryName)
+	{
+		return CategoryName;
+	});	
+	
+	TArray<UAnimNextVariableEntry*> EntriesToMove;
+	ThisEditorData->ForEachEntryOfType<UAnimNextVariableEntry>([&EntriesToMove, &CategoriesToMove](UAnimNextVariableEntry* Entry)
+	{
+		if (CategoriesToMove.Contains(Entry->GetVariableCategory()))
+		{
+			EntriesToMove.Add(Entry);
+		}
+	
+		return true;
+	});
+	
+	for (const FString& CategoryName : CategoriesToMove)
+	{	
+		EditorData->AddCategory(CategoryName);
+		ThisEditorData->RemoveCategory(CategoryName);
+	}
+	
+	if (EntriesToMove.Num())
+	{
+		for (UAnimNextVariableEntry* ToMoveEntry : EntriesToMove)
+		{
+			if (UAnimNextVariableEntry* NewEntry = UncookedOnly::FUtils::MoveVariableToAsset(ToMoveEntry, UncookedOnly::FUtils::GetAsset(EditorData)))
+			{
+				NewEntry->SetVariableCategory(ToMoveEntry->GetVariableCategory());
 			}
 		}
-
-		FToolMenuContext Context;
-		Context.AddObject(MenuContext);
-		WeakWorkspaceEditor.Pin()->InitToolMenuContext(Context);
-		return UToolMenus::Get()->GenerateWidget(MenuName, Context);
-	}
-}
-
-FReply FVariablesOutlinerMode::OnKeyDown(const FKeyEvent& InKeyEvent)
-{
-	if (CommandList->ProcessCommandBindings(InKeyEvent))
-	{
-		return FReply::Handled();
 	}
-
-	return FReply::Unhandled();
 }
 
-void FVariablesOutlinerMode::OnItemClicked(FSceneOutlinerTreeItemPtr Item)
+void FVariablesOutlinerMode::Rebuild()
 {
-	const FSceneOutlinerItemSelection& Selection = SceneOutliner->GetSelection();
-	HandleItemSelection(Selection);
+	Hierarchy = CreateHierarchy();
 }
 
 void FVariablesOutlinerMode::HandleItemSelection(const FSceneOutlinerItemSelection& Selection) const
 {
-	if (TSharedPtr<UE::Workspace::IWorkspaceEditor> WorkspaceEditor = WeakWorkspaceEditor.Pin())
+	if (TSharedPtr<Workspace::IWorkspaceEditor> WorkspaceEditor = WeakWorkspaceEditor.Pin())
 	{
 		TArray<FSceneOutlinerTreeItemPtr> SelectedItems;
 		Selection.Get(SelectedItems);
@@ -263,36 +705,19 @@ void FVariablesOutlinerMode::HandleItemSelection(const FSceneOutlinerItemSelecti
 					EntriesToShow.Add(VariableEntry);
 				}
 			}
-			else if(const FVariablesOutlinerStructSharedVariablesItem* SharedVariablesItem = Item->CastTo<FVariablesOutlinerStructSharedVariablesItem>())
-			{
-				if(const UAnimNextSharedVariablesEntry* SharedVariablesEntry = SharedVariablesItem->WeakEntry.Get())
-				{
-					EntriesToShow.Add(const_cast<UAnimNextSharedVariablesEntry*>(SharedVariablesEntry));
-				}
-			}
 		}
 
 		WorkspaceEditor->SetDetailsObjects(EntriesToShow);
 	}
 }
 
-void FVariablesOutlinerMode::OnItemSelectionChanged(FSceneOutlinerTreeItemPtr Item, ESelectInfo::Type SelectionType, const FSceneOutlinerItemSelection& Selection)
-{	
-	if (TSharedPtr<UE::Workspace::IWorkspaceEditor> SharedWorkspaceEditor = WeakWorkspaceEditor.Pin())
-	{
-		SharedWorkspaceEditor->SetGlobalSelection(SceneOutliner->AsShared(), UE::Workspace::FOnClearGlobalSelection::CreateRaw(this, &FVariablesOutlinerMode::ResetOutlinerSelection));
-	}
-	
-	HandleItemSelection(Selection);
-}
-
 bool FVariablesOutlinerMode::CanRenameItem(const ISceneOutlinerTreeItem& Item) const
 {
 	if (const FVariablesOutlinerEntryItem* EntryItem = Item.CastTo<FVariablesOutlinerEntryItem>())
 	{
 		return !EntryItem->HasStructOwner();
 	}
-	else if (Item.IsA<FVariablesOutlinerAssetItem>())
+	else if (Item.IsA<FOutlinerAssetItem>())
 	{
 		return true;
 	}
@@ -306,95 +731,8 @@ bool FVariablesOutlinerMode::CanRenameItem(const ISceneOutlinerTreeItem& Item) c
 
 void FVariablesOutlinerMode::BindCommands(const TSharedRef<FUICommandList>& OutCommandList)
 {
-	CommandList->MapAction( 
-		FGenericCommands::Get().Rename, 
-		FExecuteAction::CreateRaw(this, &FVariablesOutlinerMode::Rename),
-		FCanExecuteAction::CreateRaw(this, &FVariablesOutlinerMode::CanRename),
-		FIsActionChecked(),
-		FIsActionButtonVisible::CreateRaw(this, &FVariablesOutlinerMode::CanRename)
-	);
-
-	CommandList->MapAction(
-		FGenericCommands::Get().Delete,
-		FExecuteAction::CreateRaw(this, &FVariablesOutlinerMode::Delete),
-		FCanExecuteAction::CreateRaw(this, &FVariablesOutlinerMode::CanDelete),
-		FIsActionChecked(),
-		FIsActionButtonVisible::CreateRaw(this, &FVariablesOutlinerMode::CanDelete)	
-		);
-
-	// [TODO] implement copy/paste for variables
-	CommandList->MapAction(
-		FGenericCommands::Get().Copy, 
-		FExecuteAction::CreateRaw(this, &FVariablesOutlinerMode::Copy),
-		FCanExecuteAction::CreateRaw(this, &FVariablesOutlinerMode::CanCopy),
-		FIsActionChecked(),
-		FIsActionButtonVisible::CreateRaw(this, &FVariablesOutlinerMode::CanCopy)
-	);
-
-	CommandList->MapAction( 
-		FGenericCommands::Get().Paste, 
-		FExecuteAction::CreateRaw(this, &FVariablesOutlinerMode::Paste),
-		FCanExecuteAction::CreateRaw(this, &FVariablesOutlinerMode::CanPaste),
-		FIsActionChecked(),
-		FIsActionButtonVisible::CreateRaw(this, &FVariablesOutlinerMode::CanPaste)
-	);
-
-	CommandList->MapAction(
-		FGenericCommands::Get().Duplicate,
-		FExecuteAction::CreateRaw(this, &FVariablesOutlinerMode::Duplicate),
-		FCanExecuteAction::CreateRaw(this, &FVariablesOutlinerMode::CanDuplicate),
-		FIsActionChecked(),
-		FIsActionButtonVisible::CreateRaw(this, &FVariablesOutlinerMode::CanDuplicate)
-	);
-
-	CommandList->MapAction(
-		FVariablesOutlinerCommands::Get().FindReferences,
-		FExecuteAction::CreateRaw(this, &FVariablesOutlinerMode::FindReferences, ESearchScope::Global),
-		FCanExecuteAction::CreateRaw(this, &FVariablesOutlinerMode::CanFindReferences),
-		FIsActionChecked(),
-		FIsActionButtonVisible::CreateRaw(this, &FVariablesOutlinerMode::IsFindReferencesVisible, ESearchScope::Global)
-	);
-	
-	CommandList->MapAction(
-		FVariablesOutlinerCommands::Get().FindReferencesInWorkspace,
-		FExecuteAction::CreateRaw(this, &FVariablesOutlinerMode::FindReferences, ESearchScope::Workspace),
-		FCanExecuteAction::CreateRaw(this, &FVariablesOutlinerMode::CanFindReferences),
-		FIsActionChecked(),
-	FIsActionButtonVisible::CreateRaw(this, &FVariablesOutlinerMode::IsFindReferencesVisible, ESearchScope::Workspace)
-	);
-	
-	CommandList->MapAction(
-		FVariablesOutlinerCommands::Get().FindReferencesInAsset,
-		FExecuteAction::CreateRaw(this, &FVariablesOutlinerMode::FindReferences, ESearchScope::Asset),
-		FCanExecuteAction::CreateRaw(this, &FVariablesOutlinerMode::CanFindReferences),
-		FIsActionChecked(),
-		FIsActionButtonVisible::CreateRaw(this, &FVariablesOutlinerMode::IsFindReferencesVisible, ESearchScope::Asset)
-	);
-
-	CommandList->MapAction(
-		FVariablesOutlinerCommands::Get().SaveAsset,
-		FExecuteAction::CreateRaw(this, &FVariablesOutlinerMode::SaveAsset),
-		FCanExecuteAction::CreateRaw(this, &FVariablesOutlinerMode::CanSaveAsset),
-		FIsActionChecked(),
-		FIsActionButtonVisible::CreateRaw(this, &FVariablesOutlinerMode::IsAsset)
-		);
+	FCommonOutlinerMode::BindCommands(OutCommandList);
 	
-	CommandList->MapAction(
-		FVariablesOutlinerCommands::Get().SetVariableExportPublic,
-		FExecuteAction::CreateRaw(this, &FVariablesOutlinerMode::SetVariableExportAccessSpecifiersOnSelection, EAnimNextExportAccessSpecifier::Public),
-		FCanExecuteAction::CreateRaw(this, &FVariablesOutlinerMode::CanModifyVariableExport),
-		FIsActionChecked(),
-		FIsActionButtonVisible::CreateRaw(this, &FVariablesOutlinerMode::CanModifyVariableExport)
-	);
-
-	CommandList->MapAction(
-		FVariablesOutlinerCommands::Get().SetVariableExportPrivate,
-		FExecuteAction::CreateRaw(this, &FVariablesOutlinerMode::SetVariableExportAccessSpecifiersOnSelection, EAnimNextExportAccessSpecifier::Private),
-		FCanExecuteAction::CreateRaw(this, &FVariablesOutlinerMode::CanModifyVariableExport),
-		FIsActionChecked(),
-		FIsActionButtonVisible::CreateRaw(this, &FVariablesOutlinerMode::CanModifyVariableExport)
-	);
-
 	CommandList->MapAction(		
 		FVariablesOutlinerCommands::Get().CreateSharedVariablesAssets,
 		FExecuteAction::CreateRaw(this, &FVariablesOutlinerMode::CreateSharedVariablesAssets),
@@ -409,16 +747,6 @@ TUniquePtr<ISceneOutlinerHierarchy> FVariablesOutlinerMode::CreateHierarchy()
 	return MakeUnique<FVariablesOutlinerHierarchy>(this);
 }
 
-void FVariablesOutlinerMode::ResetOutlinerSelection() const
-{
-	SceneOutliner->ClearSelection();
-}
-
-SVariablesOutliner* FVariablesOutlinerMode::GetOutliner() const
-{
-	return static_cast<SVariablesOutliner*>(SceneOutliner);
-}
-
 void FVariablesOutlinerMode::Rename() const
 {
 	const FSceneOutlinerItemSelection& Selection = SceneOutliner->GetSelection();
@@ -449,6 +777,9 @@ void FVariablesOutlinerMode::Delete() const
 {
 	int32 NumEntries = 0;
 	TMap<UAnimNextRigVMAssetEditorData*, TArray<UAnimNextRigVMAssetEntry*>> EntriesToDeletePerAsset;
+	
+	TMap<UAnimNextRigVMAssetEditorData*, TArray<FString>> CategoryPathsToDeletePerAsset;
+	
 	TArray<FSceneOutlinerTreeItemPtr> SelectedItems = SceneOutliner->GetSelectedItems();
 	for(const FSceneOutlinerTreeItemPtr& Item : SelectedItems)
 	{
@@ -474,152 +805,656 @@ void FVariablesOutlinerMode::Delete() const
 				NumEntries++;
 			}
 		}
-		else if (const FVariablesOutlinerStructSharedVariablesItem* SharedVariablesItem = Item->CastTo<FVariablesOutlinerStructSharedVariablesItem>())
+		else if (const FVariablesOutlinerCategoryItem* CategoryItem = Item->CastTo<FVariablesOutlinerCategoryItem>())
 		{
-			UAnimNextSharedVariablesEntry* SharedVariablesEntry = SharedVariablesItem->WeakEntry.Get();
-			if(SharedVariablesEntry == nullptr)
+			if (UAnimNextRigVMAsset* Asset = CategoryItem->WeakOwner.Get())
 			{
-				continue;
+				UAnimNextRigVMAssetEditorData* EditorData = UncookedOnly::FUtils::GetEditorData(Asset);
+				CategoryPathsToDeletePerAsset.FindOrAdd(EditorData).Add(CategoryItem->CategoryPath);
+				NumEntries++;
+
+				EditorData->ForEachEntryOfType<UAnimNextVariableEntry>([CategoryPath = CategoryItem->CategoryPath, EditorData, &EntriesToDeletePerAsset, &NumEntries, &CategoryPathsToDeletePerAsset](UAnimNextVariableEntry* Entry)
+				{
+					if (Entry->GetVariableCategory().StartsWith(CategoryPath))
+					{
+						CategoryPathsToDeletePerAsset.FindOrAdd(EditorData).AddUnique(Entry->GetVariableCategory().GetData());
+						EntriesToDeletePerAsset.FindOrAdd(EditorData).Add(Entry);
+						NumEntries++;
+					}
+					
+					return true;
+				});
+			}		
+		}
+	}
+
+	if(EntriesToDeletePerAsset.Num() || CategoryPathsToDeletePerAsset.Num())
+	{
+		// [TODO] prompt user when deletion impacts assets not part of the current workspace (local context), as these have a hidden (read unknown) impact
+		FScopedTransaction Transaction(FText::Format(LOCTEXT("DeleteVariablesFormat", "Delete {0}|plural(one=item, other=items)"), NumEntries));
+		for(const TPair<UAnimNextRigVMAssetEditorData*, TArray<UAnimNextRigVMAssetEntry*>>& EntriesPair : EntriesToDeletePerAsset)
+		{
+			for (UAnimNextRigVMAssetEntry* Entry : EntriesPair.Value)
+			{
+				UncookedOnly::FUtils::DeleteVariable(CastChecked<UAnimNextVariableEntry>(Entry), true, true);
 			}
+		}
 
-			UAnimNextRigVMAssetEditorData* EditorData = SharedVariablesEntry->GetTypedOuter<UAnimNextRigVMAssetEditorData>();
-			if(EditorData == nullptr)
+		for(const TPair<UAnimNextRigVMAssetEditorData*, TArray<FString>>& CategoriesPair : CategoryPathsToDeletePerAsset)
+		{
+			for (const FString& CategoryPath : CategoriesPair.Value)
 			{
-				continue;
+				CategoriesPair.Key->RemoveCategory(CategoryPath);
+			}
+		}
+	}	
+}
+
+bool FVariablesOutlinerMode::CanDelete() const
+{
+	TArray<FSceneOutlinerTreeItemPtr> SelectedItems = SceneOutliner->GetSelectedItems();
+
+	bool bCanDelete = false;
+	for(const FSceneOutlinerTreeItemPtr& Item : SelectedItems)
+	{
+		if (const FVariablesOutlinerEntryItem* VariablesItem = Item->CastTo<FVariablesOutlinerEntryItem>())
+		{
+			if (!VariablesItem->HasStructOwner())
+			{
+				bCanDelete = true;
 			}
+		}
+		
+		if (const FVariablesOutlinerCategoryItem* CategoryItem = Item->CastTo<FVariablesOutlinerCategoryItem>())
+		{
+			bCanDelete = true;
+		}
+	}
+	
+	return bCanDelete;
+}
 
-			TArray<UAnimNextRigVMAssetEntry*>& EntriesToDelete = EntriesToDeletePerAsset.FindOrAdd(EditorData);
-			EntriesToDelete.Add(SharedVariablesEntry);
-			NumEntries++;
+bool RecursivelyFindVariableItem(const FVariablesOutlinerCategoryItem& Category)
+{	
+	const TSet<TWeakPtr<ISceneOutlinerTreeItem>>& Children = Category.GetChildren();
+	for (const TWeakPtr<ISceneOutlinerTreeItem>& ChildItem : Children)
+	{
+		if (TSharedPtr<ISceneOutlinerTreeItem> SharedChild = ChildItem.Pin())
+		{
+			if (SharedChild->IsA<FVariablesOutlinerEntryItem>())
+			{
+				return true;
+			}
+			
+			if (FVariablesOutlinerCategoryItem* ChildCategoryItem = SharedChild->CastTo<FVariablesOutlinerCategoryItem>())
+			{
+				if (RecursivelyFindVariableItem(*ChildCategoryItem))
+				{
+					return true;
+				}
+			}
 		}
 	}
+	
+	return false;
+}
 
-	if(NumEntries > 0)
+bool FVariablesOutlinerMode::CanCopy() const
+{
+	/**
+	 *	- Variable(s) from a graph, module of shared variables. (from multiple assets at the same time)
+	 *	- Entire category (by category selection only)
+	 *	- Variables from structs
+	 */
+	const FSceneOutlinerItemSelection& Selection = SceneOutliner->GetSelection();
+	const bool bSelectionContainsVariable = Selection.Has<FVariablesOutlinerEntryItem>();
+	const bool bSelectionContainsCategory = Selection.Has<FVariablesOutlinerCategoryItem>();	
+	const bool bSelectedCategoryContainsVariables = bSelectionContainsCategory && [&Selection]()-> bool
 	{
-		// [TODO] prompt user when deletion impacts assets not part of the current workspace (local context), as these have a hidden (read unknown) impact
-		FScopedTransaction Transaction(FText::Format(LOCTEXT("DeleteVariablesFormat", "Delete {0}|plural(one=variable, other=variables)"), NumEntries));
-		for(const TPair<UAnimNextRigVMAssetEditorData*, TArray<UAnimNextRigVMAssetEntry*>>& EntriesPair : EntriesToDeletePerAsset)
+		bool bContainsVariable = false;
+		Selection.ForEachItem<FVariablesOutlinerCategoryItem>([&bContainsVariable](const FVariablesOutlinerCategoryItem& Category)
+		{
+			bContainsVariable |= RecursivelyFindVariableItem(Category);
+			return !bContainsVariable;
+		});
+		
+		return bContainsVariable;
+	}();
+	
+	const bool bSelectionContainsAssetItem = Selection.Has<FOutlinerAssetItem>();
+	const bool bSelectionContainsStructItem = Selection.Has<FVariablesOutlinerStructSharedVariablesItem>();
+	
+	return !bSelectionContainsAssetItem && !bSelectionContainsStructItem && (bSelectionContainsVariable || (bSelectionContainsCategory && bSelectedCategoryContainsVariables));
+}
+
+void FVariablesOutlinerMode::Copy() const
+{
+	FVariablesOutlinerClipboardData ClipboardData;
+
+	auto CopyVariableEntry = [&ClipboardData](const UAnimNextVariableEntry* Entry, const UAnimNextRigVMAssetEditorData* EditorData)
+	{
+		FVariableClipboardData VariableData;
+		VariableData.SoftSourceObjectPath = EditorData;
+
+		VariableData.ParameterName = Entry->GetVariableName();
+		VariableData.Type = Entry->GetType();
+		if (Entry->Binding.IsValid())
+		{
+			VariableData.Binding = Entry->Binding;
+		}
+		VariableData.Comment = Entry->Comment;
+		VariableData.Category = Entry->Category;
+		VariableData.Access = Entry->GetExportAccessSpecifier();
+
+		Entry->GetDefaultValueString(VariableData.DefaultValue);
+
+		ClipboardData.Variables.AddUnique(VariableData);
+	};
+
+	auto CopyStructProperty = [&ClipboardData](const FProperty* Property)
+	{		
+		if (const UScriptStruct* Struct = Property->GetOwner<UScriptStruct>())
+		{
+			FVariableClipboardData VariableData;
+			VariableData.SoftSourceObjectPath = Struct;
+			
+			FStructOnScope DefaultStruct;
+			DefaultStruct.Initialize(Struct);
+			VariableData.ParameterName = Property->GetFName();
+			VariableData.Type = FAnimNextParamType::FromProperty(Property);			
+			VariableData.Access = EAnimNextExportAccessSpecifier::Public;
+
+			Property->ExportTextItem_InContainer(VariableData.DefaultValue, DefaultStruct.GetStructMemory(), nullptr, nullptr, 0);
+			
+			ClipboardData.Variables.AddUnique(VariableData);
+		}
+	};
+	
+	const FSceneOutlinerItemSelection& Selection = SceneOutliner->GetSelection();
+	Selection.ForEachItem<FVariablesOutlinerCategoryItem>([&CopyVariableEntry, &ClipboardData](const FVariablesOutlinerCategoryItem& SelectedItem)
+	{
+		const FString& CategoryPath = SelectedItem.CategoryPath;
+		if (UAnimNextRigVMAsset* Asset = SelectedItem.WeakOwner.Get())
+		{
+			if (UAnimNextRigVMAssetEditorData* EditorData = UncookedOnly::FUtils::GetEditorData(Asset))
+			{
+				EditorData->ForEachEntryOfType<UAnimNextVariableEntry>([&CopyVariableEntry, &CategoryPath, &EditorData](const UAnimNextVariableEntry* VariableEntry)
+				{
+					// Category is stored as set of concatenated (sub)categories, so this will cover _all_ variables
+					if (VariableEntry->GetVariableCategory().StartsWith(CategoryPath))
+					{
+						CopyVariableEntry(VariableEntry, EditorData);
+					}
+					return true;
+				});
+			}
+
+			ClipboardData.bCategorySelectedForOperation = true;
+		}
+	});
+
+	Selection.ForEachItem<FVariablesOutlinerEntryItem>([&CopyVariableEntry, &CopyStructProperty](const FVariablesOutlinerEntryItem& SelectedItem)
+	{
+		if (UAnimNextVariableEntry* Entry = SelectedItem.WeakEntry.Get())
+		{
+			if (UAnimNextRigVMAssetEditorData* EditorData = CastChecked<UAnimNextRigVMAssetEditorData>(Entry->GetOuter()))
+			{
+				CopyVariableEntry(Entry, EditorData);
+			}
+		}
+		else if (SelectedItem.bStructOwner)
+		{
+			if (const FProperty* Property = SelectedItem.PropertyPath.Get())
+			{
+				CopyStructProperty(Property);
+			}
+		}
+	});
+	
+	ensure(ClipboardData.Variables.Num());
+	{		
+		FString ClipboardDataString;
+		FVariablesOutlinerClipboardData::StaticStruct()->ExportText(ClipboardDataString, &ClipboardData, nullptr, nullptr, 0, nullptr);	
+		FPlatformApplicationMisc::ClipboardCopy(*ClipboardDataString);
+	}
+}
+
+bool FVariablesOutlinerMode::CanCut() const
+{
+	/**
+	 *	- Variable(s) from a graph, module of shared variables. (from multiple assets at the same time?) BUT not from Structs
+	 *	- Entire category (by category selection only)
+	*/
+	const FSceneOutlinerItemSelection& Selection = SceneOutliner->GetSelection();
+	const bool bSelectionContainsVariable = Selection.Has<FVariablesOutlinerEntryItem>();
+	const bool bSelectionContainsCategory = Selection.Has<FVariablesOutlinerCategoryItem>();
+	const bool bSelectedCategoryContainsVariables = bSelectionContainsCategory && [&Selection]()-> bool
+	{
+		bool bContainsVariable = false;
+		Selection.ForEachItem<FVariablesOutlinerCategoryItem>([&bContainsVariable](const FVariablesOutlinerCategoryItem& Category)
+		{
+			bContainsVariable |= RecursivelyFindVariableItem(Category);
+			return !bContainsVariable;
+		});
+	
+		return bContainsVariable;
+	}();
+	const bool bSelectionContainsAssetItem = Selection.Has<FOutlinerAssetItem>();
+
+	bool bSelectionContainsVariableFromStruct = false;
+	if (bSelectionContainsVariable)
+	{
+		TArray<FVariablesOutlinerEntryItem*> SelectedItems;
+		Selection.Get<FVariablesOutlinerEntryItem>(SelectedItems);
+
+		for (FVariablesOutlinerEntryItem* Item : SelectedItems)
+		{
+			bSelectionContainsVariableFromStruct |= Item->HasStructOwner();
+		}
+	}
+	if (bSelectionContainsCategory)
+	{
+		TArray<FVariablesOutlinerCategoryItem*> SelectedItems;
+		Selection.Get<FVariablesOutlinerCategoryItem>(SelectedItems);
+
+		for (FVariablesOutlinerCategoryItem* Item : SelectedItems)
+		{
+			bSelectionContainsVariableFromStruct |= Item->WeakOwner.Get() == nullptr;
+		}
+	}
+
+	return !bSelectionContainsAssetItem && !bSelectionContainsVariableFromStruct && (bSelectionContainsVariable || (bSelectionContainsCategory && bSelectedCategoryContainsVariables));
+}
+
+void FVariablesOutlinerMode::Cut() const
+{
+	FVariablesOutlinerClipboardData ClipboardData;
+	ClipboardData.bCutOperation = true;
+
+	auto CutVariableEntry = [&ClipboardData](UAnimNextVariableEntry* Entry)
+	{
+		UAnimNextRigVMAssetEditorData* EditorData = CastChecked<UAnimNextRigVMAssetEditorData>(Entry->GetOuter());
+		
+		FVariableClipboardData& VariableData = ClipboardData.Variables.AddDefaulted_GetRef();
+		VariableData.SoftSourceObjectPath = EditorData;
+
+		VariableData.ParameterName = Entry->GetVariableName();
+		VariableData.Type = Entry->GetType();
+		if (Entry->Binding.IsValid())
+		{
+			VariableData.Binding = Entry->Binding;
+		}
+		VariableData.Comment = Entry->Comment;
+		VariableData.Category = Entry->Category;
+		VariableData.Access = Entry->GetExportAccessSpecifier();
+		Entry->GetDefaultValueString(VariableData.DefaultValue);
+
+		EditorData->RemoveEntry(Entry);
+	};
+
+	FScopedTransaction Transaction(LOCTEXT("CutVariables", "Cut Variable(s)"));
+	TArray<UAnimNextVariableEntry*> EntriesToCut;
+	TArray<UAnimNextRigVMAssetEditorData*> EditorDatas;
+	
+	const FSceneOutlinerItemSelection& Selection = SceneOutliner->GetSelection();
+	Selection.ForEachItem<FVariablesOutlinerCategoryItem>([&CutVariableEntry, &ClipboardData, &EntriesToCut, &EditorDatas](const FVariablesOutlinerCategoryItem& SelectedItem)
+	{
+		const FString& CategoryPath = SelectedItem.CategoryPath;
+		if (UAnimNextRigVMAsset* Asset = SelectedItem.WeakOwner.Get())
+		{
+			if (UAnimNextRigVMAssetEditorData* EditorData = UncookedOnly::FUtils::GetEditorData(Asset))
+			{
+				EditorData->ForEachEntryOfType<UAnimNextVariableEntry>([&CutVariableEntry, &CategoryPath, &EntriesToCut](UAnimNextVariableEntry* VariableEntry)
+				{
+					// Category is stored as set of concatenated (sub)categories, so this will cover _all_ variables
+					if (VariableEntry->GetVariableCategory().StartsWith(CategoryPath))
+					{
+						EntriesToCut.Add(VariableEntry);
+					}
+					return true;
+				});
+
+				EditorDatas.AddUnique(EditorData);
+				EditorData->RemoveCategory(CategoryPath);
+			}
+
+			ClipboardData.bCategorySelectedForOperation = true;
+		}
+	});
+	
+	
+	TMap<UAnimNextRigVMAssetEditorData*, TArray<FString>> CategoriesToCheck;
+	Selection.ForEachItem<FVariablesOutlinerEntryItem>([&EntriesToCut, &EditorDatas, &CategoriesToCheck](const FVariablesOutlinerEntryItem& SelectedItem)
+	{
+		if (UAnimNextVariableEntry* Entry = SelectedItem.WeakEntry.Get())
+		{
+			const FStringView Category = Entry->GetVariableCategory();
+			if (!Category.IsEmpty())
+			{
+				CategoriesToCheck.FindOrAdd(UncookedOnly::FUtils::GetEditorData(SelectedItem.WeakOwner.Get())).Add(FString(Category));
+			}
+		
+			EntriesToCut.Add(Entry);
+			EditorDatas.AddUnique(Entry->GetTypedOuter<UAnimNextRigVMAssetEditorData>());
+		}
+	});
+
+	TArray<TUniquePtr<FRigVMControllerCompileBracketScope>> RigVMScopes;
+	for (UAnimNextRigVMAssetEditorData* EditorData : EditorDatas)
+	{
+		// Standalone SharedVariableAssets don't have a graph, and thus no controller
+		if (URigVMController* Controller = EditorData->GetOrCreateController())
+		{
+			RigVMScopes.Add(MakeUnique<FRigVMControllerCompileBracketScope>(Controller));
+		}
+	}
+
+	for (UAnimNextVariableEntry* VariableEntry : EntriesToCut)
+	{
+		CutVariableEntry(VariableEntry);
+	}
+	
+	// Handle categories for cut variables (as the might end up empty)
+	for (TPair<UAnimNextRigVMAssetEditorData*, TArray<FString>> Pair : CategoriesToCheck)
+	{	
+		TArray<FString>& Categories = Pair.Value;
+		// Remove all non-empty categories
+		Pair.Key->ForEachEntryOfType<UAnimNextVariableEntry>([&Categories](UAnimNextVariableEntry* VariableEntry)
+		{
+			Categories.RemoveAll([CategoryPath = VariableEntry->GetVariableCategory()](const FString& PathToCheck) { return PathToCheck == CategoryPath; });
+			return true;
+		});
+		
+		// Now remove all empty categories
+		for (const FString& CategoryPath : Categories)
+		{
+			Pair.Key->RemoveCategory(CategoryPath);
+		}
+	}
+	
+	ensure(ClipboardData.Variables.Num());
+	{
+		FString ClipboardDataString;
+		FVariablesOutlinerClipboardData::StaticStruct()->ExportText(ClipboardDataString, &ClipboardData, nullptr, nullptr, 0, nullptr);	
+		FPlatformApplicationMisc::ClipboardCopy(*ClipboardDataString);
+	}
+}
+
+void FVariablesOutlinerMode::CheckOverlappingVariableNamesRecursively(const UAnimNextRigVMAssetEditorData* EditorData, const TArray<FName>& BaseVariableNames, TMap<FName, FSoftObjectPath>& OutOverlappingVariableNames) const
+{
+	EditorData->ForEachEntryOfType<UAnimNextVariableEntry>([&BaseVariableNames, &OutOverlappingVariableNames, &EditorData](const UAnimNextVariableEntry* Entry)
+	{
+		if (BaseVariableNames.Contains(Entry->GetVariableName()))
+		{
+			OutOverlappingVariableNames.Add(Entry->GetVariableName(), EditorData->GetOuter());
+		}
+		
+		return true;
+	});	
+	
+	EditorData->ForEachEntryOfType<UAnimNextSharedVariablesEntry>([this, &BaseVariableNames, &OutOverlappingVariableNames](const UAnimNextSharedVariablesEntry* Entry)
+	{
+		if (const UAnimNextSharedVariables* SharedVariables = Entry->Asset.Get())
 		{
-			for (UAnimNextRigVMAssetEntry* Entry : EntriesPair.Value)
+			if (const UAnimNextRigVMAssetEditorData* SharedVariablesEditorData = UncookedOnly::FUtils::GetEditorData<const UAnimNextRigVMAssetEditorData>(SharedVariables))
 			{
-				UncookedOnly::FUtils::DeleteVariable(CastChecked<UAnimNextVariableEntry>(Entry), true, true);
+				CheckOverlappingVariableNamesRecursively(SharedVariablesEditorData, BaseVariableNames, OutOverlappingVariableNames);
 			}
 		}
-	}
+		
+		return true;
+	});
 }
 
-bool FVariablesOutlinerMode::CanDelete() const
+bool FVariablesOutlinerMode::CanPaste() const
 {
-	TArray<FSceneOutlinerTreeItemPtr> SelectedItems = SceneOutliner->GetSelectedItems();
+	/**
+	 *	- Asset (not a struct) selected as target
+	 *	- (Sub)Category selected as target
+	 */
+	const FSceneOutlinerItemSelection& Selection = SceneOutliner->GetSelection();
+	const bool bSelectionIsCategory = Selection.Has<FVariablesOutlinerCategoryItem>();
+	const bool bSelectionIsAsset = Selection.Has<FOutlinerAssetItem>();
 
-	bool bCanDelete = false;
-	for(const FSceneOutlinerTreeItemPtr& Item : SelectedItems)
+	UAnimNextRigVMAsset* TargetAsset = nullptr;
+	if (bSelectionIsCategory)
 	{
-		if (const FVariablesOutlinerEntryItem* VariablesItem = Item->CastTo<FVariablesOutlinerEntryItem>())
+		TArray<FVariablesOutlinerCategoryItem*> SelectedItems;
+		Selection.Get<FVariablesOutlinerCategoryItem>(SelectedItems);
+
+		if (SelectedItems.Num() == 1)
 		{
-			if (!VariablesItem->HasStructOwner())
-			{
-				bCanDelete = true;
-			}
+			const FVariablesOutlinerCategoryItem& SelectedCategory = *SelectedItems[0];
+			TargetAsset = SelectedCategory.WeakOwner.Get();
 		}
 	}
-	
-	return bCanDelete;
-}
+	else
+	{
+		TArray<FOutlinerAssetItem*> SelectedItems;
+		Selection.Get<FOutlinerAssetItem>(SelectedItems);
 
-void FVariablesOutlinerMode::Copy() const
-{
-	const FSceneOutlinerItemSelection& Selection = SceneOutliner->GetSelection();
-	Selection.ForEachItem([](const FSceneOutlinerTreeItemPtr& SelectedItem)
+		// This should have been enforced by CanPaste
+		if (SelectedItems.Num() == 1)
+		{
+			const FOutlinerAssetItem& SelectedAsset = *SelectedItems[0];
+			TargetAsset = SelectedAsset.SoftAsset.Get();
+		}
+	}
+
+	FVariablesOutlinerClipboardData ClipboardData;
+	if (TargetAsset)
 	{
-		
-	});
-}
+		FString ClipboardDataString;
+		FPlatformApplicationMisc::ClipboardPaste(ClipboardDataString);
 
-bool FVariablesOutlinerMode::CanCopy() const
-{
-	// [TODO] implement copy behaviour
-	return false;
+		FVariablesOutlinerClipboardData::StaticStruct()->ImportText(*ClipboardDataString, &ClipboardData, nullptr, 0, nullptr, FVariablesOutlinerClipboardData::StaticStruct()->GetName());
+	}
+
+	const bool bClipboardHasValidPasteData = ClipboardData.Variables.Num() != 0;
+	return bClipboardHasValidPasteData;
 }
 
 void FVariablesOutlinerMode::Paste() const
 {
-}
+	const FSceneOutlinerItemSelection& Selection = SceneOutliner->GetSelection();
+	const bool bSelectionIsCategory = Selection.Has<FVariablesOutlinerCategoryItem>();
+	const bool bSelectionIsAsset = Selection.Has<FOutlinerAssetItem>();
 
-bool FVariablesOutlinerMode::CanPaste() const
-{
-	// [TODO] implement paste behaviour
-	return false;
-}
+	const FVariablesOutlinerCategoryItem* SelectedCategoryPtr = nullptr;
 
-void FVariablesOutlinerMode::SetVariableExportAccessSpecifiersOnSelection(EAnimNextExportAccessSpecifier VariableAccessSpecifier) const
-{
-	if (!SceneOutliner)
+	UAnimNextRigVMAsset* TargetAsset = nullptr;
+	if (bSelectionIsCategory)
 	{
-		return;
+		TArray<FVariablesOutlinerCategoryItem*> SelectedItems;
+		Selection.Get<FVariablesOutlinerCategoryItem>(SelectedItems);
+
+		// This should have been enforced by CanPaste
+		check(SelectedItems.Num() == 1);
+
+		SelectedCategoryPtr = SelectedItems[0];
+		TargetAsset = SelectedCategoryPtr->WeakOwner.Get();
 	}
-	
-	const FSceneOutlinerItemSelection& Selection = SceneOutliner->GetSelection();
-	if (Selection.Num() > 0)
+	else if (bSelectionIsAsset)
+	{
+		TArray<FOutlinerAssetItem*> SelectedItems;
+		Selection.Get<FOutlinerAssetItem>(SelectedItems);
+
+		// This should have been enforced by CanPaste
+		check(SelectedItems.Num() == 1);
+
+		const FOutlinerAssetItem& SelectedAsset = *SelectedItems[0];
+		TargetAsset = SelectedAsset.SoftAsset.Get();
+	}
+
+	if (TargetAsset)
 	{
-		TArray<UObject*> ObjectsToCompile;
-		TArray<IAnimNextRigVMExportInterface*> ExportsToChange;
-		Selection.ForEachItem<FVariablesOutlinerEntryItem>([&ObjectsToCompile, &ExportsToChange](const FVariablesOutlinerEntryItem& SelectedItem)
+		FString ClipboardDataString;
+		FPlatformApplicationMisc::ClipboardPaste(ClipboardDataString);
+
+		FVariablesOutlinerClipboardData ClipboardData;
+
+		FVariablesOutlinerClipboardData::StaticStruct()->ImportText(*ClipboardDataString, &ClipboardData, nullptr, 0, nullptr, FVariablesOutlinerClipboardData::StaticStruct()->GetName());
+		
+		// Verify that to-be-pasted variables don't overlap with any existing within the asset, or its upward dependencies
 		{
-			if (IAnimNextRigVMExportInterface* Export = Cast<IAnimNextRigVMExportInterface>(SelectedItem.WeakEntry.Get()))
+			TArray<FName> PasteVariableNames;
+			Algo::Transform(ClipboardData.Variables, PasteVariableNames, [](const FVariableClipboardData& VariablePair)
+			{
+				return VariablePair.ParameterName;
+			});
+		
+			TMap<FName, FSoftObjectPath> OverlappingVariableNames;
+			UAnimNextRigVMAssetEditorData* TargetEditorData = UncookedOnly::FUtils::GetEditorData(TargetAsset);
+			CheckOverlappingVariableNamesRecursively(TargetEditorData, PasteVariableNames, OverlappingVariableNames);
+			
+			if (OverlappingVariableNames.Num())
 			{
-				if (UAnimNextRigVMAsset* OuterAsset = SelectedItem.WeakEntry.Get()->GetTypedOuter<UAnimNextRigVMAsset>())
+				FString OverlappingVariableNameInfo;
+				for (const TPair<FName, FSoftObjectPath>& Pair : OverlappingVariableNames)
 				{
-					if (OuterAsset->GetClass() && OuterAsset->GetClass() != UAnimNextSharedVariables::StaticClass())
+					if (OverlappingVariableNameInfo.Len())
 					{
-						ObjectsToCompile.AddUnique(OuterAsset);
-						ExportsToChange.AddUnique(Export);
+						OverlappingVariableNameInfo.Append(TEXT("\n"));
 					}
+					OverlappingVariableNameInfo.Append(FString::Printf(TEXT("%s: %s"), *Pair.Key.ToString(), *Pair.Value.GetAssetName()));
 				}
-			}
-		});
 
-		UncookedOnly::FScopedCompileJob CompilerResults(LOCTEXT("ModifiedAssets_SetAccessSpecifierOnSelection", "Modified Assets Set Access Specifier on Selection"), ObjectsToCompile);
+				const FText NotificationText = FText::Format(LOCTEXT("FailedPasteOverlappingNames", "Asset or references already contain Variable Name(s):\n{0}"), FText::FromString(OverlappingVariableNameInfo));
+				
+				FNotificationInfo Notification(NotificationText);
+				Notification.ExpireDuration = 5.0f;
+				FSlateNotificationManager::Get().AddNotification(Notification);	
 
-		FScopedTransaction Transaction(LOCTEXT("SetAccessSpecifierSelectionTransaction", "Set Access Specifier on Selection"));
+				return;
+			}
+		}
 
-		for (IAnimNextRigVMExportInterface* Export : ExportsToChange)
+		UAnimNextRigVMAssetEditorData* TargetEditorData = UncookedOnly::FUtils::GetEditorData(TargetAsset);
+		if (TargetEditorData && ClipboardData.Variables.Num())
 		{
-			if (Export)
+			FScopedTransaction Transaction(LOCTEXT("PasteVariables", "Paste Variable(s)"));
+			
+			// Standalone SharedVariableAssets don't have a graph, and thus no controller
+			TUniquePtr<FRigVMControllerCompileBracketScope> Scope;
+			if (URigVMController* Controller = TargetEditorData->GetOrCreateController())
+			{
+				Scope = MakeUnique<FRigVMControllerCompileBracketScope>(Controller);
+			}
+			
+			TargetEditorData->Modify();
+
+			const bool bShouldRetainCopiedCategories = ClipboardData.bCategorySelectedForOperation;
+			
+			for (FVariableClipboardData& VariableData : ClipboardData.Variables)
 			{
-				Export->SetExportAccessSpecifier(VariableAccessSpecifier);
+				const FSoftObjectPath& SourceObjectPath = VariableData.SoftSourceObjectPath;
+				UAnimNextVariableEntry* NewVariable = TargetEditorData->AddVariable(VariableData.ParameterName, VariableData.Type, VariableData.DefaultValue);
+				NewVariable->SetExportAccessSpecifier(VariableData.Access);
+
+				// In case the user explicitly selected categories when copy/cutting
+				if (bShouldRetainCopiedCategories && !VariableData.Category.IsEmpty())
+				{
+					// Maintain the categories in case of pasting on an asset
+					if (SelectedCategoryPtr == nullptr)
+					{
+						TargetEditorData->AddCategory(VariableData.Category);
+						NewVariable->SetVariableCategory(VariableData.Category);
+					}
+					// If pasting on a category (make the pasted category a child category)
+					else
+					{
+						const FString BaseCategoryPath = SelectedCategoryPtr->CategoryPath;
+						const FString NewCategoryPath = BaseCategoryPath + TEXT("|") + VariableData.Category;
+
+						TargetEditorData->AddCategory(NewCategoryPath);
+						NewVariable->SetVariableCategory(NewCategoryPath);
+					}
+				}
+				// Otherwise, in case the user explicitly selected a category to paste in - set it as the category for all pasted variables
+				else if (bSelectionIsCategory && SelectedCategoryPtr)
+				{
+					NewVariable->SetVariableCategory(SelectedCategoryPtr->CategoryPath);
+				}
+
+				if(VariableData.Binding.IsValid())
+				{
+					NewVariable->SetBinding(MoveTemp(VariableData.Binding.BindingData));
+				}
+
+				NewVariable->Comment = VariableData.Comment;
+
+				if (ClipboardData.bCutOperation)
+				{
+					if (UAnimNextRigVMAssetEditorData* CurrentEditorData = Cast<UAnimNextRigVMAssetEditorData>(SourceObjectPath.ResolveObject()))
+					{
+						UAnimNextRigVMAssetEditorData* NewEditorData = TargetEditorData;
+						
+						const FAnimNextSoftVariableReference FindReference(VariableData.ParameterName, UncookedOnly::FUtils::GetAsset(CurrentEditorData));
+						const FAnimNextSoftVariableReference ReplaceReference(VariableData.ParameterName, UncookedOnly::FUtils::GetAsset(NewEditorData));
+
+						UncookedOnly::FUtils::ReplaceVariableReferencesAcrossProject(FindReference, ReplaceReference);
+					}
+				}
 			}
 		}
 	}
 }
 
-bool FVariablesOutlinerMode::CanModifyVariableExport() const
+
+void FVariablesOutlinerMode::AddCategory(TWeakObjectPtr<UAnimNextRigVMAssetEditorData> InWeakEditorData) const
 {
-	const FSceneOutlinerItemSelection& Selection = SceneOutliner->GetSelection();
-	if (Selection.Num())
+	if (UAnimNextRigVMAssetEditorData* EditorData = InWeakEditorData.Get())
 	{
-		bool bCanModify = false;
-		
-		FScopedTransaction Transaction(LOCTEXT("SetAccessSpecifierTransaction", "Set Access Specifier"));
-		Selection.ForEachItem<FVariablesOutlinerEntryItem>([&bCanModify](const FVariablesOutlinerEntryItem& EntryItem)
+		auto GetCategoryName = [&EditorData]() -> FString
 		{
-			if (UAnimNextVariableEntry* VariableEntry = Cast<UAnimNextVariableEntry>(EntryItem.WeakEntry.Get()))
+			auto NameExists = [ExistingNames = EditorData->VariableAndFunctionCategories](FString InName)
 			{
-				if (UAnimNextRigVMAsset* OuterAsset = VariableEntry->GetTypedOuter<UAnimNextRigVMAsset>())
+				for(FString AdditionalName : ExistingNames)
 				{
-					// Disable modifying of variable export for pure UAnimNextSharedVariables objects				
-					if (OuterAsset->GetClass() && OuterAsset->GetClass() != UAnimNextSharedVariables::StaticClass())
+					if(AdditionalName == InName)
 					{
-						bCanModify = true;
+						return true;
 					}
 				}
+
+				return false;
+			};
+			
+			const FString BaseCategoryName(TEXT("Category"));
+
+			if(!NameExists(BaseCategoryName))
+			{
+				// Early out - name is valid
+				return BaseCategoryName;
 			}
-		});
 
-		return bCanModify;
-	}
+			int32 PostFixIndex = 0;
+			TStringBuilder<128> StringBuilder;
+			while(true)
+			{
+				StringBuilder.Reset();
+				StringBuilder.Append(BaseCategoryName);
+				StringBuilder.Appendf(TEXT("_%d"), PostFixIndex++);
 
-	return false;
+				const FString TestString(StringBuilder.ToString()); 
+				if(!NameExists(TestString))
+				{
+					return TestString;
+				}
+			}
+		};
+
+		FScopedTransaction Transaction(LOCTEXT("AddNewCategory", "Adding New Category"));
+		const FString NewCategoryName = GetCategoryName();
+		check(EditorData->AddCategory(NewCategoryName));
+		
+		if (SceneOutliner)
+		{
+			const FSoftObjectPath OwnerPath = EditorData->GetTypedOuter<UAnimNextRigVMAsset>();
+			const FSceneOutlinerTreeItemID CategoryID = HashCombine(GetTypeHash(OwnerPath), GetTypeHash(NewCategoryName));
+			SceneOutliner->OnItemAdded(CategoryID, SceneOutliner::ENewItemAction::Select | SceneOutliner::ENewItemAction::Rename);
+		}
+	}
 }
 
 void FVariablesOutlinerMode::Duplicate() const
@@ -631,12 +1466,27 @@ void FVariablesOutlinerMode::Duplicate() const
 
 		Selection.ForEachItem<FVariablesOutlinerEntryItem>([](const FVariablesOutlinerEntryItem& EntryItem)
 		{
-			if (const UAnimNextVariableEntry* VariableEntry = EntryItem.WeakEntry.Get())
+			if (UAnimNextVariableEntry* VariableEntry = EntryItem.WeakEntry.Get())
 			{
 				if (UAnimNextRigVMAssetEditorData* EditorData = VariableEntry->GetTypedOuter<UAnimNextRigVMAssetEditorData>())
 				{	
-					const FName VariableName = UncookedOnly::FUtils::GetValidVariableName(EditorData, VariableEntry->GetVariableName());		
-					EditorData->AddVariable(VariableName, VariableEntry->GetType());
+					const FName VariableName = UncookedOnly::FUtils::GetValidVariableName(EditorData, VariableEntry->GetVariableName());
+					FString DefaultValueString;
+					VariableEntry->GetDefaultValueString(DefaultValueString);
+					if (UAnimNextVariableEntry* DuplicatedVariable = EditorData->AddVariable(VariableName, VariableEntry->GetType(), DefaultValueString))
+					{
+						if (!VariableEntry->GetVariableCategory().IsEmpty())
+						{
+							DuplicatedVariable->SetVariableCategory(VariableEntry->GetVariableCategory());
+						}
+
+						if(VariableEntry->GetBinding().IsValid())
+						{
+							TInstancedStruct<FAnimNextVariableBindingData> BindingCopy = VariableEntry->Binding.BindingData;
+							DuplicatedVariable->SetBinding(MoveTemp(BindingCopy));
+						}					
+						DuplicatedVariable->SetExportAccessSpecifier(VariableEntry->GetExportAccessSpecifier());
+					}
 				}
 			}
 		});
@@ -668,13 +1518,13 @@ void FVariablesOutlinerMode::FindReferences(ESearchScope InSearchScope) const
 		FSceneOutlinerTreeItemPtr ItemToFind = Selection.SelectedItems[0].Pin();
 		if (const FVariablesOutlinerEntryItem* EntryItem = ItemToFind->CastTo<FVariablesOutlinerEntryItem>())
 		{
-			if (TSharedPtr<SDockTab> FindAndReplaceTab = WeakWorkspaceEditor.Pin()->GetTabManager()->TryInvokeTab(FTabId(UE::UAF::Editor::FindAndReplaceTabName)))
+			if (TSharedPtr<SDockTab> FindAndReplaceTab = WeakWorkspaceEditor.Pin()->GetTabManager()->TryInvokeTab(FTabId(UAF::Editor::FindAndReplaceTabName)))
 			{
 				TSharedRef<IAnimAssetFindReplace> FindAndReplace = StaticCastSharedRef<IAnimAssetFindReplace>(FindAndReplaceTab->GetContent());
 				FindAndReplace->SetCurrentProcessor(UAnimNextAssetFindReplaceVariables::StaticClass());
 				UAnimNextAssetFindReplaceVariables* AnimNextFindAndReplaceProcessor = FindAndReplace->GetProcessor<UAnimNextAssetFindReplaceVariables>();
 
-				if (TSharedPtr<UE::Workspace::IWorkspaceEditor> SharedWorkspaceEditor = WeakWorkspaceEditor.Pin())
+				if (TSharedPtr<Workspace::IWorkspaceEditor> SharedWorkspaceEditor = WeakWorkspaceEditor.Pin())
 				{
 					AnimNextFindAndReplaceProcessor->SetWorkspaceEditor(SharedWorkspaceEditor.ToSharedRef());
 				}
@@ -709,30 +1559,31 @@ bool FVariablesOutlinerMode::IsFindReferencesVisible(ESearchScope InSearchScope)
 
 void FVariablesOutlinerMode::PopulateNewVariableToolMenuEntries(UToolMenu* InMenu, bool bAddSeparator)
 {
-	const UAnimNextVariableItemMenuContext* MenuContext = InMenu->FindContext<UAnimNextVariableItemMenuContext>();
-	if(MenuContext == nullptr)
+	const UCommonOutlinerItemMenuContext* OutlinerContext = InMenu->FindContext<UCommonOutlinerItemMenuContext>();
+	const UVariablesOutlinerItemMenuContext* MenuContext = InMenu->FindContext<UVariablesOutlinerItemMenuContext>();
+	if(OutlinerContext == nullptr || MenuContext == nullptr)
 	{
 		return;
 	}
 
 	TMap<TWeakObjectPtr<UAnimNextRigVMAssetEditorData>, FString> AssetsAndCategoriesToAddTo;
-	if (MenuContext->WeakEditorDatas.Num())
+	if (OutlinerContext->WeakEditorDatas.Num())
 	{
-		if (UAnimNextRigVMAssetEditorData* EditorData = MenuContext->WeakEditorDatas[0].Get())
+		if (UAnimNextRigVMAssetEditorData* EditorData = OutlinerContext->WeakEditorDatas[0].Get())
 		{
 			AssetsAndCategoriesToAddTo.Add(EditorData, TEXT(""));
 		}
 	}
 	else
 	{
-		TSharedPtr<SVariablesOutliner> SharedOutliner = MenuContext->WeakOutliner.Pin();
+		TSharedPtr<SCommonOutliner> SharedOutliner = OutlinerContext->WeakOutliner.Pin();
 		if (SharedOutliner.IsValid())
 		{
 			TArray<FSceneOutlinerTreeItemPtr> SelectedItems = SharedOutliner->GetSelectedItems();
 		
 			for(const FSceneOutlinerTreeItemPtr& Item : SelectedItems)
 			{
-				const FVariablesOutlinerAssetItem* AssetItem = Item->CastTo<FVariablesOutlinerAssetItem>();
+				const FOutlinerAssetItem* AssetItem = Item->CastTo<FOutlinerAssetItem>();
 				const FVariablesOutlinerCategoryItem* CategoryItem = Item->CastTo<FVariablesOutlinerCategoryItem>();
 				if (AssetItem == nullptr && CategoryItem == nullptr)
 				{
@@ -775,7 +1626,7 @@ void FVariablesOutlinerMode::PopulateNewVariableToolMenuEntries(UToolMenu* InMen
 
 		FToolMenuSection& Section = InMenu->AddSection("Variables", LOCTEXT("VariablesSectionLabel", "Variables"));
 		Section.AddMenuEntry("AddNewVariable", LOCTEXT("VariableLabel", "Add Variable"), LOCTEXT("VariableTooltip", "Adds a new Variable"),FSlateIcon(FAppStyle::GetAppStyleSetName(), "Icons.PlusCircle"),
-			FUIAction(FExecuteAction::CreateLambda([WeakOutliner = MenuContext->WeakOutliner, AssetsAndCategoriesToAddTo]()
+			FUIAction(FExecuteAction::CreateLambda([WeakOutliner = OutlinerContext->WeakOutliner, AssetsAndCategoriesToAddTo]()
 			{
 				for (auto& Pair : AssetsAndCategoriesToAddTo)
 				{
@@ -792,7 +1643,7 @@ void FVariablesOutlinerMode::PopulateNewVariableToolMenuEntries(UToolMenu* InMen
 							Settings->SetLastVariableName(VariableName);
 
 							// Prompt user to rename added variable in outliner
-							TSharedPtr<SVariablesOutliner> SharedOutliner = WeakOutliner.Pin();
+							TSharedPtr<SCommonOutliner> SharedOutliner = WeakOutliner.Pin();
 							if (SharedOutliner.IsValid())
 							{
 								// FVariablesOutlinerEntryItem::SharedVariablesEntry is only populated for Struct based SharedVariables
@@ -974,104 +1825,10 @@ void FVariablesOutlinerMode::PopulateNewVariableToolMenuEntries(UToolMenu* InMen
 	}
 }
 
-void FVariablesOutlinerMode::SetHighlightedItem(FSceneOutlinerTreeItemID InID) const
-{
-	if (SVariablesOutliner* Outliner = static_cast<SVariablesOutliner*>(SceneOutliner))
-	{
-		if (FSceneOutlinerTreeItemPtr ItemPtr = Outliner->GetTreeItem(InID))
-		{
-			Outliner->FrameItem(InID);
-			Outliner->SetHighlightedItem(ItemPtr);
-		}
-	}
-}
-
-void FVariablesOutlinerMode::ClearHighlightedItem(FSceneOutlinerTreeItemID InID) const
-{
-	if (SVariablesOutliner* Outliner = static_cast<SVariablesOutliner*>(SceneOutliner))
-	{
-		if (FSceneOutlinerTreeItemPtr ItemPtr = Outliner->GetTreeItem(InID))
-		{
-			Outliner->ClearHighlightedItem(ItemPtr);
-		}
-	}
-}
-
-void FVariablesOutlinerMode::SetSelectedVariablesExportAccessSpecifiers(FSceneOutlinerTreeItemPtr ItemToChange, EAnimNextExportAccessSpecifier VariableAccessSpecifier) const
-{
-	if (!SceneOutliner || !ItemToChange.IsValid())
-	{
-		return;
-	}
-
-	const FSceneOutlinerItemSelection& Selection = SceneOutliner->GetSelection();
-	const bool bEditChangedItemOnly = !Selection.SelectedItems.Contains(ItemToChange);
-
-	FScopedTransaction Transaction(LOCTEXT("SetAccessSpecifierTransaction", "Set Access Specifier"));
-
-	if (bEditChangedItemOnly)
-	{
-		if (TSharedPtr<FVariablesOutlinerEntryItem> VariablesOutlinerItem = StaticCastSharedPtr<FVariablesOutlinerEntryItem>(ItemToChange))
-		{
-			if (IAnimNextRigVMExportInterface* Export = Cast<IAnimNextRigVMExportInterface>(VariablesOutlinerItem->WeakEntry.Get()))
-			{
-				Export->SetExportAccessSpecifier(VariableAccessSpecifier);
-			}
-		}
-	}
-	else
-	{
-		SetVariableExportAccessSpecifiersOnSelection(VariableAccessSpecifier);
-	}
-}
-
-void FVariablesOutlinerMode::ToggleSelectedVariablesExportAccessSpecifiers(FSceneOutlinerTreeItemPtr ItemToChange) const
-{
-	if (TSharedPtr<FVariablesOutlinerEntryItem> VariablesOutlinerItem = StaticCastSharedPtr<FVariablesOutlinerEntryItem>(ItemToChange))
-	{
-		if (IAnimNextRigVMExportInterface* Export = Cast<IAnimNextRigVMExportInterface>(VariablesOutlinerItem->WeakEntry.Get()))
-		{
-			SetSelectedVariablesExportAccessSpecifiers(ItemToChange, Export->GetExportAccessSpecifier() == EAnimNextExportAccessSpecifier::Public ? EAnimNextExportAccessSpecifier::Private : EAnimNextExportAccessSpecifier::Public);
-		}
-	}	
-}
-
-void FVariablesOutlinerMode::SaveAsset() const
-{
-	const FSceneOutlinerItemSelection& Selection = SceneOutliner->GetSelection();
-	Selection.ForEachItem<FVariablesOutlinerAssetItem>([](const FVariablesOutlinerAssetItem& AssetItem)
-	{
-		if (const UAnimNextRigVMAsset* Asset = AssetItem.SoftAsset.Get())
-		{
-			FEditorFileUtils::PromptForCheckoutAndSave({Asset->GetPackage()}, false, /*bPromptToSave=*/ false);
-		}
-	});	
-}
-
-bool FVariablesOutlinerMode::CanSaveAsset() const
-{
-	const FSceneOutlinerItemSelection& Selection = SceneOutliner->GetSelection();	
-	return Selection.SelectedItems.ContainsByPredicate([](const TWeakPtr<ISceneOutlinerTreeItem>& WeakItem)
-	{
-		if (WeakItem.IsValid())
-		{
-			if (const FVariablesOutlinerAssetItem* AssetItem = WeakItem.Pin()->CastTo<FVariablesOutlinerAssetItem>())
-			{
-				if (const UAnimNextRigVMAsset* Asset = AssetItem->SoftAsset.Get())
-				{
-					return Asset->GetPackage()->IsDirty();
-				}
-			}
-		}
-
-		return false;
-	});
-}
-
 bool FVariablesOutlinerMode::IsAsset() const
 {
 	const FSceneOutlinerItemSelection& Selection = SceneOutliner->GetSelection();	
-	return Selection.SelectedItems.ContainsByPredicate([](const TWeakPtr<ISceneOutlinerTreeItem>& WeakItem) { return WeakItem.IsValid() && WeakItem.Pin()->IsA<FVariablesOutlinerAssetItem>(); });	
+	return Selection.SelectedItems.ContainsByPredicate([](const TWeakPtr<ISceneOutlinerTreeItem>& WeakItem) { return WeakItem.IsValid() && WeakItem.Pin()->IsA<FOutlinerAssetItem>(); });	
 }
 
 void FVariablesOutlinerMode::CreateSharedVariablesAssets() const
@@ -1096,7 +1853,7 @@ void FVariablesOutlinerMode::CreateSharedVariablesAssets() const
 			FSaveAssetDialogConfig SaveAssetDialogConfig;
 			{
 				// Populate default path to match currently focussed assets
-				TSharedPtr<UE::Workspace::IWorkspaceEditor> SharedWorkspaceEditor = WeakWorkspaceEditor.Pin();
+				TSharedPtr<Workspace::IWorkspaceEditor> SharedWorkspaceEditor = WeakWorkspaceEditor.Pin();
 				if (SharedWorkspaceEditor.IsValid())
 				{
 					const Workspace::FWorkspaceDocument& Document =SharedWorkspaceEditor->GetFocusedWorkspaceDocument();
@@ -1195,110 +1952,6 @@ bool FVariablesOutlinerMode::CanCreateSharedVariablesAssets() const
 	return bSelectionOnlyContainsVariables && !bMoveWillLeaveEmptyAsset;
 }
 
-TSharedPtr<FDragDropOperation> FVariablesOutlinerMode::CreateDragDropOperation(const FPointerEvent& MouseEvent, const TArray<FSceneOutlinerTreeItemPtr>& InTreeItems) const
-{
-	if(InTreeItems.Num() != 1)
-	{
-		return nullptr;
-	}
-
-	const FSceneOutlinerTreeItemPtr& SelectedItem = InTreeItems[0];
-	if (FVariablesOutlinerEntryItem* VariableItem = SelectedItem->CastTo<FVariablesOutlinerEntryItem>())
-	{
-		if(UAnimNextVariableEntry* Entry = Cast<UAnimNextVariableEntry>(VariableItem->WeakEntry.Get()))
-		{
-			const UAnimNextSharedVariables* Asset = nullptr;
-			if (VariableItem->WeakSharedVariablesEntry.Get())
-			{
-				Asset = VariableItem->WeakSharedVariablesEntry->GetAsset();
-			}
-			else
-			{
-				Asset = Entry->GetTypedOuter<UAnimNextSharedVariables>();
-			}
-		
-			if (Asset)
-			{
-				TSharedPtr<FAnimNextSchemaAction_Variable> Action = MakeShared<FAnimNextSchemaAction_Variable>(Entry->GetVariableName(), Asset, Entry->GetType(), FAnimNextSchemaAction_Variable::EVariableAccessorChoice::Deferred);
-
-				return FVariableDragDropOp::New(Action, StaticCastSharedRef<FVariablesOutlinerEntryItem>(VariableItem->AsShared()));
-			}
-		}
-		else if (const FProperty* Property = VariableItem->PropertyPath.Get())
-		{
-			if (const UScriptStruct* Struct = Property->GetOwner<UScriptStruct>())
-			{
-				TSharedPtr<FAnimNextSchemaAction_Variable> Action = MakeShared<FAnimNextSchemaAction_Variable>(Property->GetFName(), Struct, FAnimNextParamType::FromProperty(Property), FAnimNextSchemaAction_Variable::EVariableAccessorChoice::Deferred);
-				
-				TSharedPtr<FVariableDragDropOp> Op = FVariableDragDropOp::New(Action, StaticCastSharedRef<FVariablesOutlinerEntryItem>(VariableItem->AsShared()));
-				return Op;
-			}
-		}
-	}
-	else if (FVariablesOutlinerCategoryItem* CategoryItem = SelectedItem->CastTo<FVariablesOutlinerCategoryItem>())
-	{
-		TSharedPtr<FCategoryDragDropOp> Op = FCategoryDragDropOp::New(StaticCastSharedRef<FVariablesOutlinerCategoryItem>(CategoryItem->AsShared()));
-		return Op;
-	}
-
-	return nullptr;
-
-}
-
-bool FVariablesOutlinerMode::ParseDragDrop(FSceneOutlinerDragDropPayload& OutPayload, const FDragDropOperation& Operation) const
-{
-	return ISceneOutlinerMode::ParseDragDrop(OutPayload, Operation);
-}
-
-FSceneOutlinerDragValidationInfo FVariablesOutlinerMode::ValidateDrop(const ISceneOutlinerTreeItem& DropTarget,
-	const FSceneOutlinerDragDropPayload& Payload) const
-{
-	return ISceneOutlinerMode::ValidateDrop(DropTarget, Payload);
-}
-
-void FVariablesOutlinerMode::OnDrop(ISceneOutlinerTreeItem& DropTarget, const FSceneOutlinerDragDropPayload& Payload,
-	const FSceneOutlinerDragValidationInfo& ValidationInfo) const
-{
-	ISceneOutlinerMode::OnDrop(DropTarget, Payload, ValidationInfo);
-}
-
-FReply FVariablesOutlinerMode::OnDragOverItem(const FDragDropEvent& Event, const ISceneOutlinerTreeItem& Item) const
-{
-	return ISceneOutlinerMode::OnDragOverItem(Event, Item);
-}
-
-int32 FVariablesOutlinerMode::GetTypeSortPriority(const ISceneOutlinerTreeItem& Item) const
-{
-	if(const FVariablesOutlinerAssetItem* AssetItem = Item.CastTo<FVariablesOutlinerAssetItem>())
-	{
-		return AssetItem->SortValue;
-	}
-
-	if(const FVariablesOutlinerEntryItem* EntryItem = Item.CastTo<FVariablesOutlinerEntryItem>())
-	{
-		return EntryItem->SortValue;
-	}
-
-	if(const FVariablesOutlinerCategoryItem* CategoryItem = Item.CastTo<FVariablesOutlinerCategoryItem>())
-	{
-		return CategoryItem->SortValue;
-	}
-	
-	return ISceneOutlinerMode::GetTypeSortPriority(Item);
-}
-
-void FVariablesOutlinerMode::OnItemDoubleClick(FSceneOutlinerTreeItemPtr Item)
-{
-	if (Item.IsValid())
-	{
-		if(FVariablesOutlinerAssetItem* AssetItem = Item->CastTo<FVariablesOutlinerAssetItem>())
-		{
-			// [TODO] Open asset from Variables Outliner
-		}
-	}	
-	
-	ISceneOutlinerMode::OnItemDoubleClick(Item);
-}
 }
 
-#undef LOCTEXT_NAMESPACE // "FVariablesOutlinerMode"
\ No newline at end of file
+#undef LOCTEXT_NAMESPACE // "FVariablesOutlinerMode"
diff --git a/Engine/Plugins/Experimental/UAF/UAF/Source/UAFEditor/Private/Variables/Outliner/VariablesOutlinerMode.h b/Engine/Plugins/Experimental/UAF/UAF/Source/UAFEditor/Private/Variables/Outliner/VariablesOutlinerMode.h
--- a/Engine/Plugins/Experimental/UAF/UAF/Source/UAFEditor/Private/Variables/Outliner/VariablesOutlinerMode.h
+++ b/Engine/Plugins/Experimental/UAF/UAF/Source/UAFEditor/Private/Variables/Outliner/VariablesOutlinerMode.h
@@ -5,6 +5,9 @@
 #include "AnimNextRigVMAssetEditorData.h"
 #include "ISceneOutlinerMode.h"
 #include "Common/AnimNextAssetFindReplaceVariables.h"
+#include "Common/Outliner/CommonOutlinerMode.h"
+
+#include "VariablesOutlinerMode.generated.h"
 
 class UAnimNextRigVMAssetEditorData;
 class FVariablesOutlinerMode;
@@ -13,90 +16,122 @@ namespace UE::Workspace
 	class IWorkspaceEditor;
 }
 
-namespace UE::UAF::Editor
+USTRUCT()
+struct FVariableClipboardData
 {
-	class SVariablesOutliner;
-}
+	GENERATED_BODY()
+	
+	UPROPERTY()
+	FSoftObjectPath SoftSourceObjectPath;
+
+	UPROPERTY()
+	EAnimNextExportAccessSpecifier Access = EAnimNextExportAccessSpecifier::Private;
+
+	UPROPERTY()
+	FName ParameterName;
+
+	UPROPERTY()
+	FAnimNextParamType Type;
+
+	UPROPERTY()
+	FAnimNextVariableBinding Binding;
+
+	UPROPERTY()
+	FString Comment;
+
+	UPROPERTY()
+	FString Category;
+
+	UPROPERTY()
+	FString DefaultValue;
+
+	bool operator==(const FVariableClipboardData& Other) const
+	{
+		return ParameterName == Other.ParameterName
+			&& Type == Other.Type
+			&& Binding.BindingData == Other.Binding.BindingData
+			&& Access == Other.Access
+			&& Comment == Other.Comment
+			&& Category == Other.Category
+			&& SoftSourceObjectPath == Other.SoftSourceObjectPath
+			&& DefaultValue == Other.DefaultValue;
+	}
+};
+
+USTRUCT()
+struct FVariablesOutlinerClipboardData
+{
+	GENERATED_BODY()
+
+	UPROPERTY()
+	TArray<FVariableClipboardData> Variables;
+
+	UPROPERTY()
+	bool bCutOperation = false;
+
+	UPROPERTY()
+	bool bCategorySelectedForOperation = false;
+};
 
 namespace UE::UAF::Editor
 {
-class FVariablesOutlinerMode : public ISceneOutlinerMode
+class FVariablesOutlinerMode : public FCommonOutlinerMode
 {
 public:
-	FVariablesOutlinerMode(SVariablesOutliner* InVariablesOutliner, const TSharedRef<UE::Workspace::IWorkspaceEditor>& InWorkspaceEditor);
+	FVariablesOutlinerMode(SCommonOutliner* InOutliner, const TSharedRef<UE::Workspace::IWorkspaceEditor>& InWorkspaceEditor);
+	
+	// Begin FCommonOutlinerMode overrides
+	virtual void HandleItemSelection(const FSceneOutlinerItemSelection& Selection) const override;
+	virtual void RegisterToolMenus() override;
+	virtual void InitToolMenuContext(FToolMenuContext& InOutContext) const override;
+	virtual void ReorderCategoryItem(TSharedPtr<FOutlinerCategoryItem> EntryItem, TSharedPtr<FOutlinerCategoryItem> BeforeItem) const override;
+	virtual void ReorderCategory(const FString& CategoryName, const FString& BeforeCategoryName, UAnimNextRigVMAsset* Asset) const override;
+	virtual void MoveCategoryToAsset(const FString& CategoryPath, UAnimNextRigVMAsset* FromAsset, UAnimNextRigVMAsset* ToAsset) const override;
+	// End FCommonOutlinerMode overrides
 
 	// Begin ISceneOutlinerMode overrides
 	virtual void Rebuild() override;
-	virtual TSharedPtr<SWidget> CreateContextMenu() override;
-	virtual FReply OnKeyDown(const FKeyEvent& InKeyEvent) override;
-	virtual void OnItemClicked(FSceneOutlinerTreeItemPtr Item) override;
-	void HandleItemSelection(const FSceneOutlinerItemSelection& Selection) const;
-	virtual void OnItemSelectionChanged(FSceneOutlinerTreeItemPtr Item, ESelectInfo::Type SelectionType, const FSceneOutlinerItemSelection& Selection) override;
-	virtual bool CanCustomizeToolbar() const override { return true; }
 	virtual bool CanRenameItem(const ISceneOutlinerTreeItem& Item) const override;
-	virtual ESelectionMode::Type GetSelectionMode() const override { return ESelectionMode::Multi; }
 	virtual void BindCommands(const TSharedRef<FUICommandList>& OutCommandList) override;
-	virtual int32 GetTypeSortPriority(const ISceneOutlinerTreeItem& Item) const override;
-	virtual TSharedPtr<FDragDropOperation> CreateDragDropOperation(const FPointerEvent& MouseEvent, const TArray<FSceneOutlinerTreeItemPtr>& InTreeItems) const override;
-	virtual bool ParseDragDrop(FSceneOutlinerDragDropPayload& OutPayload, const FDragDropOperation& Operation) const override;
-	virtual FSceneOutlinerDragValidationInfo ValidateDrop(const ISceneOutlinerTreeItem& DropTarget, const FSceneOutlinerDragDropPayload& Payload) const override;
-	virtual void OnDrop(ISceneOutlinerTreeItem& DropTarget, const FSceneOutlinerDragDropPayload& Payload, const FSceneOutlinerDragValidationInfo& ValidationInfo) const override;
-	virtual FReply OnDragOverItem(const FDragDropEvent& Event, const ISceneOutlinerTreeItem& Item) const override;
-	virtual bool CanSupportDragAndDrop() const override { return true; }
-	virtual void OnItemDoubleClick(FSceneOutlinerTreeItemPtr Item) override;
-	
-	void SetHighlightedItem(FSceneOutlinerTreeItemID InID) const;
-	void ClearHighlightedItem(FSceneOutlinerTreeItemID InID) const;
-
-	// Sets the export access specifier on ItemToChange and the current selection if ItemToChange is part of it
-	void SetSelectedVariablesExportAccessSpecifiers(FSceneOutlinerTreeItemPtr ItemToChange, EAnimNextExportAccessSpecifier VariableAccessSpecifier) const;
-
-	// Toggles the export access specifier on ItemToChange and sets the current selection to match, if ItemToChange is part of it
-	void ToggleSelectedVariablesExportAccessSpecifiers(FSceneOutlinerTreeItemPtr ItemToChange) const;
-
+	virtual bool UpdateOperationDecorator(const FDragDropEvent& Event, const FSceneOutlinerDragValidationInfo* ValidationInfo) const override;	
 	static void PopulateNewVariableToolMenuEntries(UToolMenu* InMenu, bool bAddSeparator);
 protected:
 	virtual TUniquePtr<ISceneOutlinerHierarchy> CreateHierarchy() override;
 	// End ISceneOutlinerMode overrides
 
-	void ResetOutlinerSelection() const;
-
-	SVariablesOutliner* GetOutliner() const;
 
-	void Rename() const;
+	// QQ begin/end overrides
+	virtual void Rename() const override;
 	virtual bool CanRename() const override;
 
-	void Delete() const;
+	virtual void Delete() const override;
 	virtual bool CanDelete() const override;
 
-	void Copy() const;
+	virtual void Copy() const override;
 	virtual bool CanCopy() const override;
 
-	void Paste() const;
+	virtual void Cut() const override;
+	virtual bool CanCut() const override;
+
+	virtual void Paste() const override;
 	virtual bool CanPaste() const override;
 
-	void SetVariableExportAccessSpecifiersOnSelection(EAnimNextExportAccessSpecifier VariableAccessSpecifier) const;
-	bool CanModifyVariableExport() const;
+	void CheckOverlappingVariableNamesRecursively(const UAnimNextRigVMAssetEditorData* EditorData, const TArray<FName>& BaseVariableNames, TMap<FName, FSoftObjectPath>& OutOverlappingVariableNames) const;
+	virtual void AddCategory(TWeakObjectPtr<UAnimNextRigVMAssetEditorData> InWeakEditorData) const override;
 
-	void Duplicate() const;
-	bool CanDuplicate() const;
+	virtual void Duplicate() const override;
+	virtual bool CanDuplicate() const override;
 
-	void FindReferences(ESearchScope InSearchScope) const;
-	bool CanFindReferences() const;
-	bool IsFindReferencesVisible(ESearchScope InSearchScope) const;
+	virtual void FindReferences(ESearchScope InSearchScope) const override;
+	virtual bool CanFindReferences() const override;
+	virtual bool IsFindReferencesVisible(ESearchScope InSearchScope) const override;
 
-	void SaveAsset() const;
-	bool CanSaveAsset() const;
 	bool IsAsset() const;
 
 	void CreateSharedVariablesAssets() const;
 	bool CanCreateSharedVariablesAssets() const;
 private:
 	friend class FVariablesOutlinerHierarchy;
-
-	TWeakPtr<UE::Workspace::IWorkspaceEditor> WeakWorkspaceEditor;
-
-	TSharedPtr<FUICommandList> CommandList;
 };
 
 }
