Added common outliner implementation to be shared between variables/functions sharing:

diff --git a/Engine/Plugins/Experimental/UAF/UAF/Source/UAFEditor/Private/Variables/Outliner/VariablesOutlinerStructSharedVariablesItem.cpp b/Engine/Plugins/Experimental/UAF/UAF/Source/UAFEditor/Private/Variables/Outliner/VariablesOutlinerStructSharedVariablesItem.cpp
--- a/Engine/Plugins/Experimental/UAF/UAF/Source/UAFEditor/Private/Variables/Outliner/VariablesOutlinerStructSharedVariablesItem.cpp
+++ b/Engine/Plugins/Experimental/UAF/UAF/Source/UAFEditor/Private/Variables/Outliner/VariablesOutlinerStructSharedVariablesItem.cpp
@@ -9,7 +9,7 @@
 #include "SceneOutlinerStandaloneTypes.h"
 #include "AnimNextEditorModule.h"
 #include "AnimNextRigVMAssetEditorData.h"
-#include "VariablesOutlinerDragDrop.h"
+#include "Common/Outliner/OutlinerDragDropOps.h"
 #include "Entries/AnimNextSharedVariablesEntry.h"
 #include "Styling/SlateIconFinder.h"
 #include "Styling/StarshipCoreStyle.h"
@@ -27,7 +27,7 @@
 namespace UE::UAF::Editor
 {
 
-const FSceneOutlinerTreeItemType FVariablesOutlinerStructSharedVariablesItem::Type;
+const FSceneOutlinerTreeItemType FVariablesOutlinerStructSharedVariablesItem::Type(&FOutlinerItem::Type);
 
 class SVariablesOutlinerStructSharedVariablesLabel : FSceneOutlinerCommonLabelData, public SCompoundWidget
 {
@@ -40,9 +40,9 @@ class SVariablesOutlinerStructSharedVariablesLabel : FSceneOutlinerCommonLabelDa
 		TreeItem = StaticCastSharedRef<FVariablesOutlinerStructSharedVariablesItem>(InTreeItem.AsShared());
 
 		FText AssetName;
-		if(const UAnimNextSharedVariablesEntry* AssetEntry = InTreeItem.WeakEntry.Get())
+		if(const UScriptStruct* Struct = InTreeItem.Struct.Get())
 		{
-			AssetName = FText::FromString(AssetEntry->GetObjectPath().ToString());
+			AssetName = FText::FromString(Struct->GetPathName());
 		}
 		else
 		{
@@ -84,58 +84,51 @@ class SVariablesOutlinerStructSharedVariablesLabel : FSceneOutlinerCommonLabelDa
 		return FText();
 	}
 
-	virtual void OnDragEnter( const FGeometry& MyGeometry, const FDragDropEvent& DragDropEvent ) override
+	virtual FSlateColor GetForegroundColor() const override
 	{
-		TSharedPtr<FVariableDragDropOp> GraphDropOp = DragDropEvent.GetOperationAs<FVariableDragDropOp>();
-		if (GraphDropOp.IsValid())
+		TOptional<FLinearColor> BaseColor;
+		if (TreeItem.IsValid())
 		{
-			const FSlateBrush* ErrorSymbol = FAppStyle::GetBrush(TEXT("Graph.ConnectorFeedback.Error"));
-			GraphDropOp->SetSimpleFeedbackMessage(ErrorSymbol, FSlateColor(GetForegroundColor()), LOCTEXT("StructSharedVariablesFeedback", "Cannot make changes to Struct SharedVariables"));
+			BaseColor = FSceneOutlinerCommonLabelData::GetForegroundColor(*TreeItem.Pin());
 		}
-	}
-
-	virtual FSlateColor GetForegroundColor() const override
-	{
-		const TOptional<FLinearColor> BaseColor = FSceneOutlinerCommonLabelData::GetForegroundColor(*TreeItem.Pin());
 		return BaseColor.IsSet() ? BaseColor.GetValue() : FSlateColor::UseForeground();
 	}
 
+
 	TWeakPtr<FVariablesOutlinerStructSharedVariablesItem> TreeItem;
 	TSharedPtr<SInlineEditableTextBlock> TextBlock;
 };
 
-FVariablesOutlinerStructSharedVariablesItem::FVariablesOutlinerStructSharedVariablesItem(UAnimNextSharedVariablesEntry* InEntry)
-	: ISceneOutlinerTreeItem(FVariablesOutlinerStructSharedVariablesItem::Type)
-	, WeakEntry(InEntry)
+FVariablesOutlinerStructSharedVariablesItem::FVariablesOutlinerStructSharedVariablesItem(const TObjectPtr<const UScriptStruct> InStruct)
+: FOutlinerItem(FVariablesOutlinerStructSharedVariablesItem::Type, {nullptr, 0})
+	, Struct(InStruct)
 {
 }
 
 bool FVariablesOutlinerStructSharedVariablesItem::IsValid() const
 {
-	return WeakEntry.Get() != nullptr;
+	return Struct != nullptr;
 }
 
 FSceneOutlinerTreeItemID FVariablesOutlinerStructSharedVariablesItem::GetID() const
 {
-	if (const UAnimNextSharedVariablesEntry* SharedVariableEntry = WeakEntry.Get())
+	if (Struct)
 	{
-		check(SharedVariableEntry->GetStruct());
-		const FSoftObjectPath Path = SharedVariableEntry->GetStruct();
+		const FSoftObjectPath Path = Struct;
 		return GetTypeHash(Path);
 	}
 	
-	return GetTypeHash(WeakEntry);
+	return GetTypeHash(Struct);
 }
 
 FString FVariablesOutlinerStructSharedVariablesItem::GetDisplayString() const
 {
-	const UAnimNextRigVMAssetEntry* Entry = WeakEntry.Get();
-	if(Entry == nullptr)
+	if(Struct == nullptr)
 	{
 		return FString();
 	}
 
-	return Entry->GetDisplayName().ToString();
+	return Struct->GetName();
 }
 
 TSharedRef<SWidget> FVariablesOutlinerStructSharedVariablesItem::GenerateLabelWidget(ISceneOutliner& Outliner, const STableRow<FSceneOutlinerTreeItemPtr>& InRow)
@@ -145,15 +138,14 @@ TSharedRef<SWidget> FVariablesOutlinerStructSharedVariablesItem::GenerateLabelWi
 
 FString FVariablesOutlinerStructSharedVariablesItem::GetPackageName() const
 {
-	const UAnimNextRigVMAssetEntry* Entry = WeakEntry.Get();
-	if(Entry == nullptr)
+	if(Struct == nullptr)
 	{
 		return ISceneOutlinerTreeItem::GetPackageName();
 	}
 
-	return Entry->GetPackage()->GetName();
+	return Struct->GetPackage()->GetName();
 }
 
 }
 
-#undef LOCTEXT_NAMESPACE // "VariablesOutlinerTreeItem"
\ No newline at end of file
+#undef LOCTEXT_NAMESPACE // "VariablesOutlinerTreeItem"
diff --git a/Engine/Plugins/Experimental/UAF/UAF/Source/UAFEditor/Private/Variables/Outliner/VariablesOutlinerStructSharedVariablesItem.h b/Engine/Plugins/Experimental/UAF/UAF/Source/UAFEditor/Private/Variables/Outliner/VariablesOutlinerStructSharedVariablesItem.h
--- a/Engine/Plugins/Experimental/UAF/UAF/Source/UAFEditor/Private/Variables/Outliner/VariablesOutlinerStructSharedVariablesItem.h
+++ b/Engine/Plugins/Experimental/UAF/UAF/Source/UAFEditor/Private/Variables/Outliner/VariablesOutlinerStructSharedVariablesItem.h
@@ -3,6 +3,7 @@
 #pragma once
 
 #include "ISceneOutlinerTreeItem.h"
+#include "Common/Outliner/OutlinerEntryItem.h"
 
 class UAnimNextSharedVariablesEntry;
 
@@ -11,11 +12,11 @@ namespace UE::UAF::Editor
 
 class IWorkspaceOutlinerItemDetails;
 
-struct FVariablesOutlinerStructSharedVariablesItem : ISceneOutlinerTreeItem
+struct FVariablesOutlinerStructSharedVariablesItem : FOutlinerItem
 {
 	static const FSceneOutlinerTreeItemType Type;
 
-	FVariablesOutlinerStructSharedVariablesItem(UAnimNextSharedVariablesEntry* InEntry);
+	FVariablesOutlinerStructSharedVariablesItem(const TObjectPtr<const UScriptStruct> InStruct);
 
 	// Begin ISceneOutlinerTreeItem overrides
 	virtual bool IsValid() const override;
@@ -26,8 +27,7 @@ struct FVariablesOutlinerStructSharedVariablesItem : ISceneOutlinerTreeItem
 	virtual FString GetPackageName() const override;
 	// End ISceneOutlinerTreeItem overrides
 
-	// Weak-ptr to the represented entry
-	TWeakObjectPtr<UAnimNextSharedVariablesEntry> WeakEntry;
+	TObjectPtr<const UScriptStruct> Struct;
 };
 
 }
diff --git a/Engine/Plugins/Experimental/UAF/UAF/Source/UAFEditor/Private/Variables/SVariablesView.cpp b/Engine/Plugins/Experimental/UAF/UAF/Source/UAFEditor/Private/Variables/SVariablesView.cpp
--- a/Engine/Plugins/Experimental/UAF/UAF/Source/UAFEditor/Private/Variables/SVariablesView.cpp
+++ b/Engine/Plugins/Experimental/UAF/UAF/Source/UAFEditor/Private/Variables/SVariablesView.cpp
@@ -1,9 +1,10 @@
-// Copyright Epic Games, Inc. All Rights Reserved.
+﻿// Copyright Epic Games, Inc. All Rights Reserved.
 
 #include "SVariablesView.h"
 
 #include "AnimNextAssetWorkspaceAssetUserData.h"
 #include "AnimNextRigVMAsset.h"
+#include "AssetToolsModule.h"
 #include "ContentBrowserModule.h"
 #include "IContentBrowserSingleton.h"
 #include "SAssetDropTarget.h"
@@ -17,7 +18,8 @@
 #include "IWorkspaceEditor.h"
 #include "IWorkspaceEditorModule.h"
 #include "WorkspaceAssetRegistryInfo.h"
-#include "Outliner/VariablesOutlinerAssetItem.h"
+#include "Common/Outliner/OutlinerColumns.h"
+#include "Common/Outliner/SCommonOutliner.h"
 #include "Entries/AnimNextSharedVariablesEntry.h"
 
 #define LOCTEXT_NAMESPACE "SVariablesView"
@@ -26,199 +28,6 @@ namespace UE::UAF::Editor
 {
 
 const FLazyName VariablesTabName("VariablesTab");
-
-void SVariablesOutliner::RegisterAssetDelegates(const UAnimNextRigVMAsset* InAsset)
-{	
-	// Bind for any modification callbacks
-	UAnimNextRigVMAssetEditorData* EditorData = UncookedOnly::FUtils::GetEditorData<UAnimNextRigVMAssetEditorData>(InAsset);
-	if(EditorData != nullptr)
-	{
-		EditorData->ModifiedDelegate.AddSP(this, &SVariablesOutliner::OnEditorDataModified);
-
-		for (const TObjectPtr<UAnimNextRigVMAssetEntry>& Entry : EditorData->Entries)
-		{
-			if (UAnimNextSharedVariablesEntry* SharedVariablesEntry = Cast<UAnimNextSharedVariablesEntry>(Entry.Get()))
-			{
-				if (TObjectPtr<const UAnimNextSharedVariables> SharedVariablesAsset = SharedVariablesEntry->GetAsset())
-				{
-					RegisterAssetDelegates(SharedVariablesAsset);
-					Assets.AddUnique(Cast<UAnimNextRigVMAsset>(const_cast<UAnimNextSharedVariables*>(SharedVariablesAsset.Get())));					
-				}
-			}
-		}
-	}
-}
-
-void SVariablesOutliner::UnregisterAssetDelegates(const UAnimNextRigVMAsset* InAsset)
-{
-	UAnimNextRigVMAssetEditorData* EditorData = UncookedOnly::FUtils::GetEditorData<UAnimNextRigVMAssetEditorData>(InAsset);
-	if(EditorData != nullptr)
-	{
-		EditorData->ModifiedDelegate.RemoveAll(this);
-
-		for (const TObjectPtr<UAnimNextRigVMAssetEntry>& Entry : EditorData->Entries)
-		{
-			if (UAnimNextSharedVariablesEntry* SharedVariablesEntry = Cast<UAnimNextSharedVariablesEntry>(Entry.Get()))
-			{
-				if (TObjectPtr<const UAnimNextSharedVariables> SharedVariablesAsset = SharedVariablesEntry->GetAsset())
-				{
-					UnregisterAssetDelegates(SharedVariablesAsset);					
-				}
-			}
-		}
-	}
-}
-
-void SVariablesOutliner::SetAssets(TConstArrayView<TSoftObjectPtr<UAnimNextRigVMAsset>> InAssets)
-{
-	for(const TSoftObjectPtr<UAnimNextRigVMAsset>& CurrentSoftAsset : Assets)
-	{
-		if(UAnimNextRigVMAsset* CurrentAsset = CurrentSoftAsset.Get())
-		{
-			UnregisterAssetDelegates(CurrentAsset);
-		}
-	}
-
-	Assets = InAssets;
-
-	for(const TSoftObjectPtr<UAnimNextRigVMAsset>& NewSoftAsset : InAssets)
-	{
-		if(UAnimNextRigVMAsset* NewAsset = NewSoftAsset.Get())
-		{
-			RegisterAssetDelegates(NewAsset);
-		}
-	}
-
-	FullRefresh();
-}
-
-void SVariablesOutliner::UpdateAssets()
-{
-	TArray<FSoftObjectPath> AssetsToAsyncLoad;
-	TArray<TSoftObjectPtr<UAnimNextRigVMAsset>> ExportAssets;
-
-	auto HandleAssetPath = [this, &AssetsToAsyncLoad, &ExportAssets](const FSoftObjectPath& Path)
-	{
-		ExportAssets.Add(TSoftObjectPtr<UAnimNextRigVMAsset>(Path));
-
-		if (Path.ResolveObject() == nullptr)
-		{
-			AssetsToAsyncLoad.Add(Path);
-		}
-	};
-	
-	if (FWorkspaceOutlinerAssetReferenceItemData::IsAssetReference(Export)
-		|| FWorkspaceOutlinerGroupItemData::IsGroupItem(Export)
-		|| FAnimNextCollapseGraphsOutlinerDataBase::IsCollapsedGraphBase(Export))
-	{
-		// Reference item
-		TArray<FSoftObjectPath> AssetPaths;
-		Export.GetAssetPaths(AssetPaths);
-
-		for (const FSoftObjectPath& Path : AssetPaths)
-		{
-			HandleAssetPath(Path);
-		}
-	}
-	else
-	{
-		const FSoftObjectPath FirstAssetPath = Export.GetFirstAssetPath();
-		if(FirstAssetPath.IsAsset())
-		{
-			HandleAssetPath(FirstAssetPath);
-		}
-		else if(Export.HasData() && Export.GetData().GetScriptStruct()->IsChildOf(FAnimNextAssetEntryOutlinerData::StaticStruct()))
-		{
-			const FAnimNextAssetEntryOutlinerData& EntryData = Export.GetData().Get<FAnimNextAssetEntryOutlinerData>();
-			if(EntryData.SoftEntryPtr.IsValid())
-			{
-				if(UAnimNextRigVMAsset* Asset = EntryData.GetEntry()->GetTypedOuter<UAnimNextRigVMAsset>())
-				{
-					HandleAssetPath(Asset);
-				}
-			}
-		}
-	}
-
-	SetAssets(ExportAssets);
-
-	// Try async load any missing assets
-	for(const FSoftObjectPath& AssetPath : AssetsToAsyncLoad)
-	{
-		TWeakPtr<SVariablesOutliner> WeakOutliner = StaticCastSharedRef<SVariablesOutliner>(AsShared());
-		
-		AssetPath.LoadAsync(FLoadSoftObjectPathAsyncDelegate::CreateLambda([WeakOutliner](const FSoftObjectPath& InSoftObjectPath, UObject* InObject)
-		{
-			UAnimNextRigVMAsset* Asset = Cast<UAnimNextRigVMAsset>(InObject);
-			if(Asset == nullptr)
-			{
-				return;
-			}
-
-			TSharedPtr<SVariablesOutliner> PinnedVariablesOutliner = WeakOutliner.Pin();
-			if(!PinnedVariablesOutliner.IsValid())
-			{
-				return;
-			}
-
-			PinnedVariablesOutliner->HandleAssetLoaded(InSoftObjectPath, Asset);
-		}));
-	}
-}
-
-void SVariablesOutliner::SetExport(const FWorkspaceOutlinerItemExport& InExport)
-{
-	Export = InExport;
-
-	UpdateAssets();
-}
-
-void SVariablesOutliner::HandleAssetLoaded(const FSoftObjectPath& InSoftObjectPath, UAnimNextRigVMAsset* InAsset)
-{	
-	if(InAsset && Assets.Contains(TSoftObjectPtr<UAnimNextRigVMAsset>(InAsset)))
-	{
-		RegisterAssetDelegates(InAsset);
-
-		FullRefresh();
-	}
-}
-
-void SVariablesOutliner::OnEditorDataModified(UAnimNextRigVMAssetEditorData* InEditorData, EAnimNextEditorDataNotifType InType, UObject* InSubject)
-{	
-	ensure(Assets.Contains(TSoftObjectPtr<UAnimNextRigVMAsset>(UncookedOnly::FUtils::GetAsset<UAnimNextRigVMAsset>(InEditorData))));
-
-	switch(InType)
-	{
-	case EAnimNextEditorDataNotifType::UndoRedo:
-	case EAnimNextEditorDataNotifType::EntryAdded:
-	case EAnimNextEditorDataNotifType::VariableCategoryChanged:
-	case EAnimNextEditorDataNotifType::CategoryAdded:
-	case EAnimNextEditorDataNotifType::CategoryChanged:
-		FullRefresh();
-		break;
-	case EAnimNextEditorDataNotifType::EntryRemoved:
-		UpdateAssets();
-		break;
-	default:
-		break;
-	}
-}
-
-void SVariablesOutliner::SetHighlightedItem(FSceneOutlinerTreeItemPtr Item) const
-{
-	GetTreeView()->Private_SetItemHighlighted(Item, true);
-}
-
-void SVariablesOutliner::ClearHighlightedItem(FSceneOutlinerTreeItemPtr Item) const
-{
-	GetTreeView()->Private_SetItemHighlighted(Item, false);
-}
-
-bool SVariablesOutliner::HasAssets() const
-{
-	return Assets.Num() > 0;
-}
-
 void SVariablesView::Construct(const FArguments& InArgs, TSharedRef<UE::Workspace::IWorkspaceEditor> InWorkspaceEditor)
 {
 	FSceneOutlinerInitializationOptions InitOptions;
@@ -229,11 +38,27 @@ void SVariablesView::Construct(const FArguments& InArgs, TSharedRef<UE::Workspac
 		TOptional<float>(), TAttribute<FText>(), EHeaderComboVisibility::Ghosted));
 	InitOptions.ColumnMap.Add(FVariablesOutlinerValueColumn::GetID(), FSceneOutlinerColumnInfo(ESceneOutlinerColumnVisibility::Visible, 20, FCreateSceneOutlinerColumn::CreateLambda([](ISceneOutliner& InSceneOutliner) { return MakeShared<FVariablesOutlinerValueColumn>(InSceneOutliner); }), true, 0.5f));
 	InitOptions.ColumnMap.Add(FSceneOutlinerSourceControlColumn::GetID(), FSceneOutlinerColumnInfo(ESceneOutlinerColumnVisibility::Visible, 30, FCreateSceneOutlinerColumn::CreateLambda([](ISceneOutliner& InSceneOutliner) { return MakeShared<FSceneOutlinerSourceControlColumn>(InSceneOutliner); }), true));
-	InitOptions.ColumnMap.Add(FVariablesOutlinerAccessSpecifierColumn::GetID(), FSceneOutlinerColumnInfo(ESceneOutlinerColumnVisibility::Visible, 40, FCreateSceneOutlinerColumn::CreateLambda([](ISceneOutliner& InSceneOutliner) { return MakeShared<FVariablesOutlinerAccessSpecifierColumn>(InSceneOutliner); }), false));
-	InitOptions.ModeFactory = FCreateSceneOutlinerMode::CreateLambda([this, WeakWorkspaceEditor=InWorkspaceEditor.ToWeakPtr()](SSceneOutliner* InOutliner) { return new FVariablesOutlinerMode(static_cast<SVariablesOutliner*>(InOutliner), WeakWorkspaceEditor.Pin().ToSharedRef()); });
+	InitOptions.ColumnMap.Add(FOutlinerAccessSpecifierColumn::GetID(), FSceneOutlinerColumnInfo(ESceneOutlinerColumnVisibility::Visible, 40, FCreateSceneOutlinerColumn::CreateLambda([](ISceneOutliner& InSceneOutliner) { return MakeShared<FOutlinerAccessSpecifierColumn>(InSceneOutliner); }), false));
+	InitOptions.ModeFactory = FCreateSceneOutlinerMode::CreateLambda([this, WeakWorkspaceEditor=InWorkspaceEditor.ToWeakPtr()](SSceneOutliner* InOutliner) { check(WeakWorkspaceEditor.IsValid()); return new FVariablesOutlinerMode(static_cast<SCommonOutliner*>(InOutliner), WeakWorkspaceEditor.Pin().ToSharedRef()); });
 
-	VariablesOutliner = SNew(SVariablesOutliner, InitOptions);
-	VariablesOutliner->SetEnabled(MakeAttributeSP(VariablesOutliner.Get(), &SVariablesOutliner::HasAssets));
+	VariablesOutliner = SNew(UE::UAF::Editor::SCommonOutliner, InitOptions)
+	.BindAssetDelegate_Lambda([WeakThis=StaticCastWeakPtr<SVariablesView>(AsWeak())](UAnimNextRigVMAssetEditorData* EditorData)
+	{
+		if (WeakThis.IsValid())
+		{
+			TSharedPtr<SVariablesView> SharedView = WeakThis.Pin();
+			EditorData->ModifiedDelegate.AddSP(SharedView.Get(), &SVariablesView::OnEditorDataModified);
+		}
+	})
+	.UnbindAssetDelegate_Lambda([WeakThis=StaticCastWeakPtr<SVariablesView>(AsWeak())](UAnimNextRigVMAssetEditorData* EditorData)
+	{
+		if (WeakThis.IsValid())
+		{
+			TSharedPtr<SVariablesView> SharedView = WeakThis.Pin();
+			EditorData->ModifiedDelegate.RemoveAll(SharedView.Get());
+		}
+	});
+	VariablesOutliner->SetEnabled(MakeAttributeSP(VariablesOutliner.Get(), &SCommonOutliner::HasAssets));
 
 	InWorkspaceEditor->OnFocusedDocumentChanged().AddSP(this, &SVariablesView::HandleFocusedDocumentChanged);
 	const Workspace::FWorkspaceDocument& Document = InWorkspaceEditor->GetFocusedWorkspaceDocument();
@@ -259,6 +84,26 @@ void SVariablesView::HandleFocusedDocumentChanged(const UE::Workspace::FWorkspac
 	VariablesOutliner->SetExport(Export);
 }
 
+void SVariablesView::OnEditorDataModified(UAnimNextRigVMAssetEditorData* InEditorData, EAnimNextEditorDataNotifType InType, UObject* InSubject) const
+{
+	switch(InType)
+	{
+	case EAnimNextEditorDataNotifType::UndoRedo:
+	case EAnimNextEditorDataNotifType::EntryAdded:
+	case EAnimNextEditorDataNotifType::VariableCategoryChanged:
+	case EAnimNextEditorDataNotifType::CategoryAdded:
+	case EAnimNextEditorDataNotifType::CategoryChanged:
+	case EAnimNextEditorDataNotifType::CategoryRemoved:
+		VariablesOutliner->FullRefresh();
+		break;
+	case EAnimNextEditorDataNotifType::EntryRemoved:
+		VariablesOutliner->UpdateAssets();
+		break;
+	default:
+		break;
+	}
+}
+
 FAnimNextVariablesTabSummoner::FAnimNextVariablesTabSummoner(TSharedPtr<UE::Workspace::IWorkspaceEditor> InHostingApp)
 	: FWorkflowTabFactory(VariablesTabName, StaticCastSharedPtr<FAssetEditorToolkit>(InHostingApp))
 {
@@ -286,4 +131,4 @@ FText FAnimNextVariablesTabSummoner::GetTabToolTipText(const FWorkflowTabSpawnIn
 
 }
 
-#undef LOCTEXT_NAMESPACE
\ No newline at end of file
+#undef LOCTEXT_NAMESPACE
diff --git a/Engine/Plugins/Experimental/UAF/UAF/Source/UAFEditor/Private/Variables/SVariablesView.h b/Engine/Plugins/Experimental/UAF/UAF/Source/UAFEditor/Private/Variables/SVariablesView.h
--- a/Engine/Plugins/Experimental/UAF/UAF/Source/UAFEditor/Private/Variables/SVariablesView.h
+++ b/Engine/Plugins/Experimental/UAF/UAF/Source/UAFEditor/Private/Variables/SVariablesView.h
@@ -23,41 +23,10 @@ namespace UE::Workspace
 
 namespace UE::UAF::Editor
 {
-	class SVariablesOutliner;
-}
-
-namespace UE::UAF::Editor
-{
+class SCommonOutliner;
 
 extern const FLazyName VariablesTabName;
 
-class SVariablesOutliner : public SSceneOutliner
-{
-public:
-	void SetExport(const FWorkspaceOutlinerItemExport& InExport);
-	void OnEditorDataModified(UAnimNextRigVMAssetEditorData* InEditorData, EAnimNextEditorDataNotifType InType, UObject* InSubject);
-
-	void SetHighlightedItem(FSceneOutlinerTreeItemPtr Item) const;
-	void ClearHighlightedItem(FSceneOutlinerTreeItemPtr Item) const; 
-	
-	bool HasAssets() const;
-protected:
-	
-	void SetAssets(TConstArrayView<TSoftObjectPtr<UAnimNextRigVMAsset>> InAssets);
-	void UpdateAssets();
-
-	void HandleAssetLoaded(const FSoftObjectPath& InSoftObjectPath, UAnimNextRigVMAsset* InAsset);
-	
-	void RegisterAssetDelegates(const UAnimNextRigVMAsset* InAsset);
-	void UnregisterAssetDelegates(const UAnimNextRigVMAsset* InAsset);
-private:
-	friend class FVariablesOutlinerMode;
-	friend class FVariablesOutlinerHierarchy;
-
-	TArray<TSoftObjectPtr<UAnimNextRigVMAsset>> Assets;
-	FWorkspaceOutlinerItemExport Export;
-};
-
 class SVariablesView : public SCompoundWidget
 {
 	SLATE_BEGIN_ARGS(SVariablesView) {}
@@ -67,11 +36,12 @@ class SVariablesView : public SCompoundWidget
 
 	void SetExportDirectly(const FWorkspaceOutlinerItemExport& InExport) const;
 private:
-	void HandleFocusedDocumentChanged(const UE::Workspace::FWorkspaceDocument& InDocument) const;
+	void HandleFocusedDocumentChanged(const UE::Workspace::FWorkspaceDocument& InDocument) const;	
+	void OnEditorDataModified(UAnimNextRigVMAssetEditorData* InEditorData, EAnimNextEditorDataNotifType InType, UObject* InSubject) const;
 
 private:
 	friend struct FAnimNextVariablesTabSummoner;
-	TSharedPtr<SVariablesOutliner> VariablesOutliner;
+	TSharedPtr<SCommonOutliner> VariablesOutliner;
 };
 
 struct FAnimNextVariablesTabSummoner : public FWorkflowTabFactory
diff --git a/Engine/Plugins/Experimental/UAF/UAF/Source/UAFEditor/UAFEditor.Build.cs b/Engine/Plugins/Experimental/UAF/UAF/Source/UAFEditor/UAFEditor.Build.cs
--- a/Engine/Plugins/Experimental/UAF/UAF/Source/UAFEditor/UAFEditor.Build.cs
+++ b/Engine/Plugins/Experimental/UAF/UAF/Source/UAFEditor/UAFEditor.Build.cs
@@ -61,6 +61,8 @@ namespace UnrealBuildTool.Rules
 					"AdvancedPreviewScene"
 				}
 			);
+			
+			bAllowUETypesInNamespaces = true; 
 		}
 	}
 }
\ No newline at end of file
diff --git a/Engine/Plugins/Experimental/UAF/UAF/Source/UAFUncookedOnly/Internal/AnimNextRigVMAssetEditorData.h b/Engine/Plugins/Experimental/UAF/UAF/Source/UAFUncookedOnly/Internal/AnimNextRigVMAssetEditorData.h
--- a/Engine/Plugins/Experimental/UAF/UAF/Source/UAFUncookedOnly/Internal/AnimNextRigVMAssetEditorData.h
+++ b/Engine/Plugins/Experimental/UAF/UAF/Source/UAFUncookedOnly/Internal/AnimNextRigVMAssetEditorData.h
@@ -53,16 +53,22 @@ namespace UE::UAF::Editor
 	struct FVariablesOutlinerEntryItem;
 	class FVariablesOutlinerMode;
 	class FVariablesOutlinerHierarchy;
+	class FOutlinerHierarchy;
 	class SVariablesOutlinerValue;
-	class SVariablesOutliner;
+	class SFunctionsView;
+	class SVariablesView;
+	class FFunctionsOutlinerHierarchy;
 	class SAddVariablesDialog;
 	class FVariableProxyCustomization;
 	class FAnimNextAnimGraphEditorModule;
 	class FCallFunctionSharedDataDetails;
 	class FAssetCompilationHandler;
 	class FVariableBindingPropertyCustomization;
+	struct FFunctionsOutlinerEntryItem;
 	struct FVariablesOutlinerCategoryItem;
-	struct FVariablesOutlinerAssetItem;
+	struct FFunctionsOutlinerCategoryItem;
+	class FFunctionsOutlinerMode;
+	class SCommonOutliner;
 }
 
 namespace UE::UAF::Tests
@@ -90,6 +96,7 @@ enum class EAnimNextEditorDataNotifType : uint8
 	VariableCategoryChanged, // A variable entry its category changed (Subject == UAnimNextVariableEntry)
 	CategoryAdded, // A category has been added (Subject == UAnimNextRigVMAssetEditorData)
 	CategoryChanged, // A category has been changed (Subject == UAnimNextRigVMAssetEditorData)
+	CategoryRemoved, // A category has been removed (Subject == UAnimNextRigVMAssetEditorData)
 	VariablesReordered,  // Asset entries have been reordered (Subject == UAnimNextRigVMAssetEditorData)
 };
 
@@ -152,6 +159,10 @@ public:
 	/** Renames a category within an AnimNext asset */
 	UFUNCTION(BlueprintCallable, Category = "UAF|Entries", meta=(ScriptMethod))
 	static UAFUNCOOKEDONLY_API bool RenameCategory(UAnimNextRigVMAsset* InAsset, const FString& CategoryName, const FString& NewCategoryName, bool bSetupUndoRedo = true, bool bPrintPythonCommand = true);
+	
+	/** Removes a category from within an AnimNext asset */
+	UFUNCTION(BlueprintCallable, Category = "UAF|Entries", meta=(ScriptMethod))
+	static UAFUNCOOKEDONLY_API bool RemoveCategory(UAnimNextRigVMAsset* InAsset, const FString& CategoryName, bool bSetupUndoRedo = true, bool bPrintPythonCommand = true);
 };
 
 /* Base class for all AnimNext editor data objects that use RigVM */
@@ -184,6 +195,9 @@ public:
 
 	/** Renames a category using provided new name */
 	UE_API bool RenameCategory(const FString& OldName, const FString& NewName, bool bInSetupUndoRedo = true, bool bPrintPythonCommand = true);
+	
+	/** Removes category with provided name */
+	UE_API bool RemoveCategory(const FString CategoryName, bool bInSetupUndoRedo = true, bool bPrintPythonCommand = true);
 
 	/** Reorders category entry to appear before provided BeforeCategoryName */
 	UE_API void ReorderCategory(const FString& CategoryName, const FString& BeforeCategoryName, bool bInSetupUndoRedo = true, bool bPrintPythonCommand = true);
@@ -229,11 +243,13 @@ protected:
 	friend class UE::UAF::Editor::FAnimNextFunctionItemDetails;
 	friend struct UE::UAF::Editor::FVariablesOutlinerEntryItem;
 	friend struct UE::UAF::Editor::FVariablesOutlinerCategoryItem;
-	friend struct UE::UAF::Editor::FVariablesOutlinerAssetItem;
+	friend struct UE::UAF::Editor::FFunctionsOutlinerCategoryItem;
 	friend class UE::UAF::Editor::FVariablesOutlinerMode;
 	friend class UE::UAF::Editor::FVariablesOutlinerHierarchy;
+	friend class UE::UAF::Editor::FFunctionsOutlinerHierarchy;	
+	friend class UE::UAF::Editor::FOutlinerHierarchy;
 	friend class UE::UAF::Editor::SVariablesOutlinerValue;
-	friend class UE::UAF::Editor::SVariablesOutliner;
+	friend class UE::UAF::Editor::FFunctionsOutlinerMode;	
 	friend class UAnimNextModuleWorkspaceAssetUserData;
 	friend class UE::UAF::Editor::SAddVariablesDialog;
 	friend class UE::UAF::Editor::FVariableProxyCustomization;
@@ -251,6 +267,10 @@ protected:
 	friend struct FAnimNextModuleTests;
 	friend struct FAnimNextAnimationGraphTests;
 	friend FAnimNextRigVMAssetCompileContext;
+	friend UE::UAF::Editor::FFunctionsOutlinerEntryItem;
+	friend UE::UAF::Editor::SFunctionsView;
+	friend UE::UAF::Editor::SCommonOutliner;
+	friend UE::UAF::Editor::SVariablesView;
 
 	// UObject interface
 	UE_API virtual void Serialize(FArchive& Ar) override;
diff --git a/Engine/Plugins/Experimental/UAF/UAF/Source/UAFUncookedOnly/Internal/Entries/AnimNextSharedVariablesEntry.h b/Engine/Plugins/Experimental/UAF/UAF/Source/UAFUncookedOnly/Internal/Entries/AnimNextSharedVariablesEntry.h
--- a/Engine/Plugins/Experimental/UAF/UAF/Source/UAFUncookedOnly/Internal/Entries/AnimNextSharedVariablesEntry.h
+++ b/Engine/Plugins/Experimental/UAF/UAF/Source/UAFUncookedOnly/Internal/Entries/AnimNextSharedVariablesEntry.h
@@ -21,8 +21,9 @@ namespace UE::UAF::Editor
 	class SAddVariablesDialog;
 	class SVariableOverride;
 	class FVariablesOutlinerMode;
-	class SVariablesOutliner;
-	struct FVariablesOutlinerAssetItem;
+	class FFunctionsOutlinerMode;
+	class FOutlinerHierarchy;
+	class SCommonOutliner;
 }
 
 namespace UE::UAF::UncookedOnly
@@ -46,6 +47,7 @@ class UAnimNextSharedVariablesEntry : public UAnimNextRigVMAssetEntry
 	friend class UAnimNextRigVMAssetEditorData;
 	friend class UAssetDefinition_AnimNextSharedVariablesEntry;
 	friend class UE::UAF::Editor::FVariablesOutlinerHierarchy;
+	friend class UE::UAF::Editor::FOutlinerHierarchy;
 	friend struct UE::UAF::Editor::FVariablesOutlinerStructSharedVariablesItem;
 	friend struct UE::UAF::UncookedOnly::FUtils;
 	friend class UE::UAF::Editor::SVariablesOutlinerStructSharedVariablesLabel;
@@ -53,9 +55,9 @@ class UAnimNextSharedVariablesEntry : public UAnimNextRigVMAssetEntry
 	friend class UE::UAF::Editor::FVariableProxyCustomization;
 	friend class UE::UAF::Editor::SAddVariablesDialog;
 	friend class UE::UAF::Editor::SVariableOverride;
-	friend class UE::UAF::Editor::SVariablesOutliner;
-	friend struct UE::UAF::Editor::FVariablesOutlinerAssetItem;
 	friend UE::UAF::Editor::FVariablesOutlinerMode;
+	friend UE::UAF::Editor::SCommonOutliner;
+	friend UE::UAF::Editor::FFunctionsOutlinerMode;
 	friend class UAnimNextAssetFindReplaceVariables;
 
 	// UAnimNextRigVMAssetEntry interface
diff --git a/Engine/Plugins/Experimental/UAF/UAF/Source/UAFUncookedOnly/Internal/Entries/AnimNextVariableEntry.h b/Engine/Plugins/Experimental/UAF/UAF/Source/UAFUncookedOnly/Internal/Entries/AnimNextVariableEntry.h
--- a/Engine/Plugins/Experimental/UAF/UAF/Source/UAFUncookedOnly/Internal/Entries/AnimNextVariableEntry.h
+++ b/Engine/Plugins/Experimental/UAF/UAF/Source/UAFUncookedOnly/Internal/Entries/AnimNextVariableEntry.h
@@ -33,6 +33,9 @@ namespace UE::UAF::Editor
 	class SVariablesOutlinerEntryLabel;
 	class SVariablesOutlinerAssetLabel;
 	class SVariablesOutlinerCategoryLabel;
+	struct FVariablesOutlinerCategoryItem;
+	struct FVariablesOutlinerEntryItem;
+	class FVariablesOutlinerMode;
 }
 
 namespace UE::UAF::Tests
@@ -79,6 +82,10 @@ class UAnimNextVariableEntry : public UAnimNextRigVMAssetEntry, public IAnimNext
 	friend class UE::UAF::Editor::SVariablesOutlinerAssetLabel;
 	friend class UE::UAF::Editor::SVariablesOutlinerCategoryLabel;
 	friend class UAnimNextAssetFindReplaceVariables;
+	friend struct UE::UAF::Editor::FVariablesOutlinerCategoryItem;
+	friend struct UE::UAF::Editor::FVariablesOutlinerEntryItem;
+	friend class UE::UAF::Editor::FVariablesOutlinerMode;
+
 
 	UAnimNextVariableEntry();
 
diff --git a/Engine/Plugins/Experimental/UAF/UAF/Source/UAFUncookedOnly/Internal/UncookedOnlyUtils.h b/Engine/Plugins/Experimental/UAF/UAF/Source/UAFUncookedOnly/Internal/UncookedOnlyUtils.h
--- a/Engine/Plugins/Experimental/UAF/UAF/Source/UAFUncookedOnly/Internal/UncookedOnlyUtils.h
+++ b/Engine/Plugins/Experimental/UAF/UAF/Source/UAFUncookedOnly/Internal/UncookedOnlyUtils.h
@@ -210,7 +210,7 @@ struct FUtils
 	static UE_API void RenameVariable(UAnimNextVariableEntry* VariableEntry, const FName NewName, bool bSetupUndoRedo = true, bool bPrintPythonCommands = true);
 
 	// Moves the provided variable entry from its editor data to the NewOuter, replacing any references across the project its assets to reference it by its new outer asset
-	static UE_API void MoveVariableToAsset(UAnimNextVariableEntry* VariableEntry, UAnimNextRigVMAsset* NewOuter, bool bSetupUndoRedo = true, bool bPrintPythonCommands = true);
+	static UE_API UAnimNextVariableEntry* MoveVariableToAsset(UAnimNextVariableEntry* VariableEntry, UAnimNextRigVMAsset* NewOuter, bool bSetupUndoRedo = true, bool bPrintPythonCommands = true);
 
 	// Changes the provided variable entry its data type, replacing any references across the project to reference using its new data type
 	static UE_API void SetVariableType(UAnimNextVariableEntry* VariableEntry, const FAnimNextParamType NewType, bool bSetupUndoRedo = true, bool bPrintPythonCommands = true);
@@ -219,6 +219,19 @@ struct FUtils
 	static UE_API void ReplaceVariableReferences(UAnimNextRigVMAssetEditorData* InEditorData, const FAnimNextSoftVariableReference& ReferenceToFind, const FAnimNextSoftVariableReference& ReferenceToReplaceWith, bool bSetupUndoRedo = true, bool bPrintPythonCommands = true);
 	// Replaces the provided variable reference with ReferenceToReplaceWith within the context of the project
 	static UE_API void ReplaceVariableReferencesAcrossProject(const FAnimNextSoftVariableReference& ReferenceToFind, const FAnimNextSoftVariableReference& ReferenceToReplaceWith, bool bSetupUndoRedo = true, bool bPrintPythonCommands = true);
+
+	// Returns whether or not it is valid to add InSharedVariables to InEditorData, and if not OutErrorMessage will contain the reasoning
+	static UE_API bool CanAddSharedVariablesReference(UAnimNextRigVMAssetEditorData* InEditorData, const UAnimNextSharedVariables* InSharedVariables, FString* OutErrorMessage = nullptr);
+	
+	// Removes a SharedVariable reference from the provided EditorData, and ensures that any references to variables within the SharedVariables are replaced
+	static UE_API bool RemoveSharedVariablesReference(UAnimNextRigVMAssetEditorData* EditorData, const FSoftObjectPath& SharedVariablesPath);
+	
+	// Generate a (function) workspace export for the provided URigVMLibraryNode
+	static UE_API FWorkspaceOutlinerItemExport MakeFunctionExport(const URigVMLibraryNode* FunctionNode, const FWorkspaceOutlinerItemExport& ParentExport);
+
+	// Generate a (asset reference) workspace export for the provided asset path
+	static UE_API FWorkspaceOutlinerItemExport MakeAssetReferenceExport(const FSoftObjectPath& ReferencedAssetPath, const FWorkspaceOutlinerItemExport& ParentExport, bool bShouldExpand = true);
+
 private:
 	static UE_API void CompileVariableBindingsInternal(const FRigVMCompileSettings& InSettings, UAnimNextRigVMAsset* InAsset, TArray<URigVMGraph*>& OutGraphs, bool bInThreadSafe);
 	
diff --git a/Engine/Plugins/Experimental/UAF/UAF/Source/UAFUncookedOnly/Private/AnimNextRigVMAssetEditorData.cpp b/Engine/Plugins/Experimental/UAF/UAF/Source/UAFUncookedOnly/Private/AnimNextRigVMAssetEditorData.cpp
--- a/Engine/Plugins/Experimental/UAF/UAF/Source/UAFUncookedOnly/Private/AnimNextRigVMAssetEditorData.cpp
+++ b/Engine/Plugins/Experimental/UAF/UAF/Source/UAFUncookedOnly/Private/AnimNextRigVMAssetEditorData.cpp
@@ -1604,7 +1604,6 @@ bool UAnimNextRigVMAssetEditorData::AddCategory(const FString& CategoryName, boo
 		return false;
     }
 
-
 	const int32 CategoryIndex = VariableAndFunctionCategories.IndexOfByKey(CategoryName);
 	if(CategoryIndex != INDEX_NONE)
 	{
@@ -1670,8 +1669,40 @@ bool UAnimNextRigVMAssetEditorData::RenameCategory(const FString& OldName, const
 	return true;
 }
 
-void UAnimNextRigVMAssetEditorData::ReorderCategory(const FString& CategoryName, const FString& BeforeCategoryName, bool bInSetupUndoRedo,
-	bool bPrintPythonCommand)
+bool UAnimNextRigVMAssetEditorData::RemoveCategory(const FString CategoryName, bool bInSetupUndoRedo, bool bPrintPythonCommand)
+{
+	if (CategoryName.IsEmpty())
+	{
+		ReportError(*FString::Printf(TEXT("UAnimNextRigVMAssetEditorData::RemoveCategory: Invalid empty category name provided.")));
+		return false;
+	}
+
+	const int32 CategoryIndex = VariableAndFunctionCategories.IndexOfByKey(CategoryName);
+	if(CategoryIndex == INDEX_NONE)
+	{
+		ReportError(*FString::Printf(TEXT("UAnimNextRigVMAssetEditorData::RemoveCategory: Could not find category with name %s."), *CategoryName));
+		return false;
+	}
+
+	if (bInSetupUndoRedo)
+	{
+		Modify();
+	}
+
+	VariableAndFunctionCategories.Remove(CategoryName);
+	BroadcastModified(EAnimNextEditorDataNotifType::CategoryRemoved, this);
+
+	if (bPrintPythonCommand)
+	{
+		RigVMPythonUtils::Print(GetName(),
+				FString::Printf(TEXT("asset.remove_category('%s')"),
+				 *CategoryName));
+	}
+
+	return true;
+}
+
+void UAnimNextRigVMAssetEditorData::ReorderCategory(const FString& CategoryName, const FString& BeforeCategoryName, bool bInSetupUndoRedo, bool bPrintPythonCommand)
 {
 	const int32 CategoryIndex = VariableAndFunctionCategories.IndexOfByKey(CategoryName);
 	if (CategoryIndex == INDEX_NONE)
@@ -2435,30 +2466,7 @@ UAnimNextSharedVariablesEntry* UAnimNextRigVMAssetEditorData::AddSharedVariables
 	}
 
 	// Check for circularity
-	auto CheckForCircularity = [this](UAnimNextSharedVariables_EditorData* InEditorData, auto& InCheckForCircularity)
-	{
-		if(InEditorData == this)
-		{
-			return true;
-		}
-
-		for(UAnimNextRigVMAssetEntry* Entry : InEditorData->Entries)
-		{
-			if(UAnimNextSharedVariablesEntry* SharedVariablesEntry = Cast<UAnimNextSharedVariablesEntry>(Entry))
-			{
-				const UAnimNextSharedVariables* SharedVariables = SharedVariablesEntry->GetAsset();
-				UAnimNextSharedVariables_EditorData* EditorData = UE::UAF::UncookedOnly::FUtils::GetEditorData<UAnimNextSharedVariables_EditorData>(SharedVariables);
-				if(InCheckForCircularity(EditorData, InCheckForCircularity))
-				{
-					return true;
-				}
-			}
-		}
-
-		return false;
-	};
-
-	if(CheckForCircularity(EditorData, CheckForCircularity))
+	if(!FUtils::CanAddSharedVariablesReference(this, InSharedVariables))
 	{
 		ReportError(TEXT("UAnimNextRigVMAssetEditorData::AddSharedVariables: Circular reference detected."));
 		return nullptr;
@@ -2477,11 +2485,14 @@ UAnimNextSharedVariablesEntry* UAnimNextRigVMAssetEditorData::AddSharedVariables
 			}
 			else if(UAnimNextSharedVariablesEntry* SharedVariablesEntry = Cast<UAnimNextSharedVariablesEntry>(Entry))
 			{
-				const UAnimNextSharedVariables* SharedVariables = SharedVariablesEntry->GetAsset();
-				UAnimNextSharedVariables_EditorData* EditorData = UE::UAF::UncookedOnly::FUtils::GetEditorData<UAnimNextSharedVariables_EditorData>(SharedVariables);
-				if(InCheckForPublicMembers(EditorData, InCheckForPublicMembers))
+				// SharedVariables could be nullptr as the entry is a Struct based source, or the underlying asset (UAnimNextSharedVariablesEntry::ObjectPath) has been deleted
+				if (const UAnimNextSharedVariables* SharedVariables = SharedVariablesEntry->GetAsset())
 				{
-					return true;
+					UAnimNextSharedVariables_EditorData* EditorData = UE::UAF::UncookedOnly::FUtils::GetEditorData<UAnimNextSharedVariables_EditorData>(SharedVariables);
+					if(InCheckForPublicMembers(EditorData, InCheckForPublicMembers))
+					{
+						return true;
+					}
 				}
 			}
 		}
@@ -2489,12 +2500,6 @@ UAnimNextSharedVariablesEntry* UAnimNextRigVMAssetEditorData::AddSharedVariables
 		return false;
 	};
 
-	if(!CheckForPublicMembers(EditorData, CheckForPublicMembers))
-	{
-		ReportError(TEXT("UAnimNextRigVMAssetEditorData::AddSharedVariables: No public variables found."));
-		return nullptr;
-	}
-
 	// Check for duplicate entry
 	auto DuplicatePredicate = [InSharedVariables](const UAnimNextRigVMAssetEntry* InEntry)
 	{
@@ -2624,6 +2629,12 @@ bool UAnimNextRigVMAssetLibrary::RenameCategory(UAnimNextRigVMAsset* InAsset, co
 	return UE::UAF::UncookedOnly::FUtils::GetEditorData(InAsset)->RenameCategory(CategoryName, NewCategoryName, bSetupUndoRedo, bPrintPythonCommand);
 }
 
+bool UAnimNextRigVMAssetLibrary::RemoveCategory(UAnimNextRigVMAsset* InAsset, const FString& CategoryName, bool bSetupUndoRedo,
+	bool bPrintPythonCommand)
+{
+	return UE::UAF::UncookedOnly::FUtils::GetEditorData(InAsset)->RemoveCategory(CategoryName, bSetupUndoRedo, bPrintPythonCommand);
+}
+
 URigVMLibraryNode* UAnimNextRigVMAssetEditorData::AddFunction(FName InFunctionName, bool bInMutable, bool bInSetupUndoRedo, bool bPrintPythonCommand)
 {
 	URigVMController* Controller = RigVMClient.GetOrCreateController(GetLocalFunctionLibrary());
@@ -3131,7 +3142,7 @@ FInstancedPropertyBag UAnimNextRigVMAssetEditorData::GenerateCombinedPropertyBag
 					{
 						if (const FPropertyBagPropertyDesc* ExistingDesc = CombinedPropertyBag.FindPropertyDescByName(Desc.Name))
 						{
-							InCompileContext.GetAssetCompileContext().Error(Asset, LOCTEXT("DuplicateVariableEntries", "Variable with overlapping Name {0} found in {1}"), FText::FromName(Desc.Name), FText::FromString(SourceAsset->GetPathName()));
+							InCompileContext.GetAssetCompileContext().Error(Asset, LOCTEXT("DuplicateVariableEntries", "Variable(s) with overlapping Name {0} found in {1}"), FText::FromName(Desc.Name), FText::FromString(SourceAsset->GetPathName()));
 
 							bInvalidProperties = true;
 						}
diff --git a/Engine/Plugins/Experimental/UAF/UAF/Source/UAFUncookedOnly/Private/UncookedOnlyUtils.cpp b/Engine/Plugins/Experimental/UAF/UAF/Source/UAFUncookedOnly/Private/UncookedOnlyUtils.cpp
--- a/Engine/Plugins/Experimental/UAF/UAF/Source/UAFUncookedOnly/Private/UncookedOnlyUtils.cpp
+++ b/Engine/Plugins/Experimental/UAF/UAF/Source/UAFUncookedOnly/Private/UncookedOnlyUtils.cpp
@@ -38,6 +38,11 @@
 #include "ScopedTransaction.h"
 #include "Misc/UObjectToken.h"
 #include "Styling/SlateIconFinder.h"
+#include "Framework/Notifications/NotificationManager.h"
+#include "ScopedTransaction.h"
+#include "Misc/UObjectToken.h"
+#include "RigVMModel/Nodes/RigVMAggregateNode.h"
+#include "Widgets/Notifications/SNotificationList.h"
 
 #define LOCTEXT_NAMESPACE "AnimNextUncookedOnlyUtils"
 
@@ -233,13 +238,13 @@ void FUtils::CompileVariables(const FRigVMCompileSettings& InSettings, UAnimNext
 					const UAnimNextSharedVariables* SharedVariablesAsset = SharedVariablesEntry->GetAsset();
 					if (SharedVariablesAsset == nullptr)
 					{
-						SharedAssetEntryError(FText::Format(LOCTEXT("InvalidSharedVariableAssetFound", "Invalid shared variables entry found ({0})"), FText::FromString(SharedVariablesEntry->GetObjectPath().ToString())));						
+						SharedAssetEntryError(FText::Format(LOCTEXT("InvalidSharedVariableAssetFound", "Invalid shared variables asset found ({0})"), FText::FromString(SharedVariablesEntry->GetObjectPath().ToString())));
 						continue;
 					}
 
 					if (SharedVariableAssets.Contains(SharedVariablesAsset))
 					{
-						SharedAssetEntryError(FText::Format(LOCTEXT("DuplicateSharedVariableAssetFound", "Duplicate shared variables entry found ({0})"), FText::FromString(SharedVariablesEntry->GetObjectPath().ToString())));						
+						SharedAssetEntryError(FText::Format(LOCTEXT("DuplicateSharedVariableAssetFound", "Duplicate shared variables asset found ({0})"), FText::FromString(SharedVariablesEntry->GetObjectPath().ToString())));
 						continue;
 					}
 
@@ -257,7 +262,7 @@ void FUtils::CompileVariables(const FRigVMCompileSettings& InSettings, UAnimNext
 
 					if (SharedVariableStructs.Contains(SharedVariablesStruct))
 					{
-						SharedAssetEntryError(FText::Format(LOCTEXT("DuplicateSharedVariableStructFound", "Duplicate shared variables struct found ({0})"), FText::FromString(SharedVariablesEntry->GetObjectPath().ToString())));						
+						SharedAssetEntryError(FText::Format(LOCTEXT("DuplicateSharedVariableStructFound", "Duplicate shared variables struct found ({0})"), FText::FromString(SharedVariablesEntry->GetObjectPath().ToString())));
 						continue;
 					}
 
@@ -403,7 +408,7 @@ void FUtils::CompileVariableBindingsInternal(const FRigVMCompileSettings& InSett
 			return;
 		}
 		URigVMPin* CopyProxyVariablesExecPin = CopyProxyVariablesNode->FindPin(FRigVMStruct::ExecuteContextName.ToString());
-		if(ExecPin == nullptr)
+		if(CopyProxyVariablesExecPin == nullptr)
 		{
 			InSettings.ReportError(TEXT("Could not find execute pin on Copy System Proxy Variables node"));
 			return;
@@ -434,6 +439,151 @@ void FUtils::CompileVariableBindingsInternal(const FRigVMCompileSettings& InSett
 	OutGraphs.Add(BindingGraph);
 }
 
+bool FUtils::CanAddSharedVariablesReference(UAnimNextRigVMAssetEditorData* InEditorData, const UAnimNextSharedVariables* InSharedVariables, FString* OutErrorMessage)
+{
+	auto CheckForCircularity = [InEditorData](UAnimNextSharedVariables_EditorData* SourceEditorData, auto& InCheckForCircularity, TArray<const UAnimNextSharedVariables_EditorData*>& InOutReferenceChain)
+	{
+		TArray<const UAnimNextSharedVariables_EditorData*> ReferenceChain = InOutReferenceChain;
+		ReferenceChain.Add(SourceEditorData);
+		if(SourceEditorData == InEditorData)
+		{
+			InOutReferenceChain = ReferenceChain;
+			return true;
+		}
+
+		for(UAnimNextRigVMAssetEntry* Entry : SourceEditorData->Entries)
+		{
+			if(UAnimNextSharedVariablesEntry* SharedVariablesEntry = Cast<UAnimNextSharedVariablesEntry>(Entry))
+			{
+				// SharedVariables could be nullptr as the entry is a Struct based source, or the underlying asset (UAnimNextSharedVariablesEntry::ObjectPath) has been deleted
+				if (const UAnimNextSharedVariables* SharedVariables = SharedVariablesEntry->GetAsset())
+				{
+					UAnimNextSharedVariables_EditorData* EditorData = GetEditorData<UAnimNextSharedVariables_EditorData>(SharedVariables);
+					
+					if(InCheckForCircularity(EditorData, InCheckForCircularity, ReferenceChain))
+					{
+						InOutReferenceChain = ReferenceChain;
+						return true;
+					}
+				}
+			}
+		}
+
+		return false;
+	};
+	
+	if (UAnimNextSharedVariables_EditorData* SharedVariablesEditorData = GetEditorData<UAnimNextSharedVariables_EditorData>(InSharedVariables))
+	{
+		TArray<const UAnimNextSharedVariables_EditorData*> EditorDataReferenceChain;
+		EditorDataReferenceChain.Add(Cast<UAnimNextSharedVariables_EditorData>(InEditorData));
+		const bool bResult = CheckForCircularity(SharedVariablesEditorData, CheckForCircularity, EditorDataReferenceChain);
+
+		if (bResult && OutErrorMessage)
+		{
+			FString DependencyString;
+			for (const UAnimNextSharedVariables_EditorData* EditorData : EditorDataReferenceChain)
+			{
+				if (!DependencyString.IsEmpty())
+				{
+					DependencyString += TEXT("\n\t");
+				}
+				DependencyString += GetAsset<UAnimNextSharedVariables>(EditorData)->GetPackage()->GetFName().ToString();
+			}
+			*OutErrorMessage = FString::Printf(TEXT("Failed to add Shared Variables dependency as this would introduce a circular reference:\n\t%s"), *DependencyString);
+		}
+
+		return !bResult;
+	}
+
+	return false;	
+}
+
+bool FUtils::RemoveSharedVariablesReference(UAnimNextRigVMAssetEditorData* InEditorData, const FSoftObjectPath& SharedVariablesPath)
+{
+	const FText AssetNameText = FText::FromName(InEditorData->GetTypedOuter<UAnimNextRigVMAsset>()->GetFName());
+	if (UAnimNextSharedVariablesEntry* SharedVariableEntryToRemove = Cast<UAnimNextSharedVariablesEntry>(InEditorData->FindEntry(SharedVariablesPath.GetAssetFName())))
+	{
+		const FText SharedVariablesAssetNameText = FText::FromName(SharedVariablesPath.GetAssetFName());
+		const FText FormattedMessage = FText::Format(LOCTEXT("RemovingSharedAssetEntryFormat", "Removing {0} from {1}"), SharedVariablesAssetNameText, AssetNameText);
+		FScopedTransaction Transaction(FormattedMessage);
+		
+		if (const UAnimNextSharedVariables* SharedVariablesToRemove = SharedVariableEntryToRemove->GetAsset())
+		{
+			// Replace all referenced variables from the to-be-removed SharedVariables from the referencing Asset
+			FAnimNextAssetRegistryExports VariableExports;
+			GetExportedVariablesForAsset(SharedVariablesToRemove, VariableExports);
+	
+			VariableExports.ForEachExportOfType<FAnimNextVariableDeclarationData>([InEditorData, SharedVariablesPath](const FName& Identifier, const FAnimNextVariableDeclarationData& VariableDeclaration) -> bool
+			{
+				FAnimNextSoftVariableReference VariableReference(Identifier, SharedVariablesPath);
+				FAnimNextSoftVariableReference EmptyReference;
+				ReplaceVariableReferences(InEditorData, VariableReference, EmptyReference);
+			
+				return true;
+			});
+		}
+		else if (const UScriptStruct* Struct = SharedVariableEntryToRemove->GetStruct())
+		{
+			// Replace all referenced variables from the to-be-removed Native SharedVariables source from the referencing Asset
+			for (TFieldIterator<FProperty> It(Struct); It; ++It)
+			{
+				const FProperty* Property = *It;
+				if (Property->HasAnyPropertyFlags(CPF_NativeAccessSpecifierPublic))
+				{	
+					FAnimNextSoftVariableReference VariableReference(Property->GetFName(), SharedVariablesPath);
+					FAnimNextSoftVariableReference EmptyReference;
+					ReplaceVariableReferences(InEditorData, VariableReference, EmptyReference);
+				}
+			}
+		}
+	
+		// Remove the entry itself
+		return InEditorData->RemoveEntry(SharedVariableEntryToRemove);
+	}
+	
+	return false;
+}
+
+FWorkspaceOutlinerItemExport FUtils::MakeFunctionExport(const URigVMLibraryNode* FunctionNode, const FWorkspaceOutlinerItemExport& ParentExport)
+{
+	FWorkspaceOutlinerItemExport Export;
+	
+	check(FunctionNode);
+	if (const URigVMGraph* ContainedModelGraph = FunctionNode->GetContainedGraph())
+	{
+		if (const UAnimNextRigVMAssetEditorData* EditorData = FunctionNode->GetTypedOuter<UAnimNextRigVMAssetEditorData>())
+		{
+			if (URigVMEdGraph* EditorObject = Cast<URigVMEdGraph>(EditorData->GetEditorObjectForRigVMGraph(ContainedModelGraph)))
+			{
+				const FSoftObjectPath EditorObjectSoftObjPath = EditorObject;
+				ensureMsgf(EditorObjectSoftObjPath.IsSubobject(), TEXT("EditorObject for RigVMFunctionReferenceNode Graph was not a subobject as expected."));
+
+				const FName Identifier = EditorObjectSoftObjPath.IsSubobject() ? *EditorObjectSoftObjPath.GetSubPathUtf8String() : FunctionNode->GetFName();
+				Export = FWorkspaceOutlinerItemExport(Identifier, ParentExport);
+
+				Export.GetData().InitializeAsScriptStruct(FAnimNextGraphFunctionOutlinerData::StaticStruct());
+				FAnimNextGraphFunctionOutlinerData& FnGraphData = Export.GetData().GetMutable<FAnimNextGraphFunctionOutlinerData>();
+				FnGraphData.SoftEditorObject = EditorObject;
+			}
+		} 
+	}
+	
+	return Export;
+}
+
+FWorkspaceOutlinerItemExport FUtils::MakeAssetReferenceExport(const FSoftObjectPath& ReferencedAssetPath, const FWorkspaceOutlinerItemExport& ParentExport, bool bShouldExpand)
+{
+	check(ReferencedAssetPath.IsValid());
+	
+	FWorkspaceOutlinerItemExport Export = FWorkspaceOutlinerItemExport(FName(ReferencedAssetPath.ToString()), ParentExport);
+	Export.GetData().InitializeAs<FWorkspaceOutlinerAssetReferenceItemData>();
+	FWorkspaceOutlinerAssetReferenceItemData& Data = Export.GetData().GetMutable<FWorkspaceOutlinerAssetReferenceItemData>();
+	Data.ReferredObjectPath = ReferencedAssetPath;
+	Data.bShouldExpandReference = bShouldExpand;
+	
+	return Export;
+}
+
 UAnimNextRigVMAsset* FUtils::GetAsset(UAnimNextRigVMAssetEditorData* InEditorData)
 {
 	check(InEditorData);
@@ -1006,32 +1156,30 @@ void FUtils::GetAssetWorkspaceExports(const UAnimNextRigVMAssetEditorData* Edito
 				AddSharedVariablesGroupEntry();
 			}
 
-			switch(SharedVariablesEntry->Type)
+			const FSoftObjectPath ObjectPath = [SharedVariablesEntry]() -> FSoftObjectPath
 			{
-			case EAnimNextSharedVariablesType::Asset:
-				{
-					const UAnimNextRigVMAsset* Asset = SharedVariablesEntry->Asset;
-					const FSoftObjectPath ObjectPath(Asset);
-					const FWorkspaceOutlinerItemExport& ParentExport = OutExports.Exports[SharedVariablesGroupIndex];
-					FWorkspaceOutlinerItemExport& ReferenceExport = OutExports.Exports.Add_GetRef(FWorkspaceOutlinerItemExport(FName(ObjectPath.ToString()), ParentExport));
-					ReferenceExport.GetData().InitializeAs<FWorkspaceOutlinerAssetReferenceItemData>();
-					FWorkspaceOutlinerAssetReferenceItemData& Data = ReferenceExport.GetData().GetMutable<FWorkspaceOutlinerAssetReferenceItemData>();
-					Data.ReferredObjectPath = ObjectPath;
-					Data.bShouldExpandReference = false;
-				}
-				break;
-			case EAnimNextSharedVariablesType::Struct:
+				switch(SharedVariablesEntry->Type)
 				{
-					const UScriptStruct* Struct = SharedVariablesEntry->Struct;
-					const FSoftObjectPath ObjectPath(Struct);
-					const FWorkspaceOutlinerItemExport& ParentExport = OutExports.Exports[SharedVariablesGroupIndex];
-					FWorkspaceOutlinerItemExport& ReferenceExport = OutExports.Exports.Add_GetRef(FWorkspaceOutlinerItemExport(FName(ObjectPath.ToString()), ParentExport));
-					ReferenceExport.GetData().InitializeAs<FWorkspaceOutlinerAssetReferenceItemData>();
-					FWorkspaceOutlinerAssetReferenceItemData& Data = ReferenceExport.GetData().GetMutable<FWorkspaceOutlinerAssetReferenceItemData>();
-					Data.ReferredObjectPath = ObjectPath;
-					Data.bShouldExpandReference = false;
+				case EAnimNextSharedVariablesType::Asset:
+					{
+						const UAnimNextRigVMAsset* Asset = SharedVariablesEntry->Asset;
+						return Asset;
+					}
+					break;
+				case EAnimNextSharedVariablesType::Struct:
+					{
+						const UScriptStruct* Struct = SharedVariablesEntry->Struct;
+						return Struct;
+					}
+					break;
 				}
-				break;
+				
+				return nullptr;
+			}();
+
+			if (ObjectPath.IsValid())
+			{
+				OutExports.Exports.Add(MakeAssetReferenceExport(ObjectPath, OutExports.Exports[SharedVariablesGroupIndex], false));
 			}
 		}
 		return true;
@@ -1081,7 +1229,7 @@ void FUtils::GetAssetWorkspaceExports(const UAnimNextRigVMAssetEditorData* Edito
 		}
 		return true;
 	});	
-
+	
 	GetFunctionLibraryWorkspaceExportsRecursive(EditorData, OutExports, AssetIdentifier, INDEX_NONE, EditorData->GetRigVMGraphFunctionStore()->PublicFunctions, EditorData->GetRigVMGraphFunctionStore()->PrivateFunctions);
 }
 
@@ -1094,9 +1242,7 @@ void FUtils::ProcessPinAssetReferences(const URigVMPin* InPin, FWorkspaceOutline
 		if (InSoftObjectPath.ResolveObject() || IAssetRegistry::GetChecked().TryGetAssetByObjectPath(InSoftObjectPath, ReferenceAssetData) == UE::AssetRegistry::EExists::Exists)
 		{
 			const FWorkspaceOutlinerItemExport& ParentExport = ParentExportIndex == INDEX_NONE ? RootExport : OutExports.Exports[ParentExportIndex];
-			FWorkspaceOutlinerItemExport& ReferenceExport = OutExports.Exports.Add_GetRef(FWorkspaceOutlinerItemExport(FName(InSoftObjectPath.ToString()), ParentExport));
-			ReferenceExport.GetData().InitializeAs<FWorkspaceOutlinerAssetReferenceItemData>();
-			ReferenceExport.GetData().GetMutable<FWorkspaceOutlinerAssetReferenceItemData>().ReferredObjectPath = InSoftObjectPath;
+			OutExports.Exports.Add(MakeAssetReferenceExport(InSoftObjectPath, ParentExport, true));
 		}
 	};
 	
@@ -1268,57 +1414,22 @@ void FUtils::GetSubGraphWorkspaceExportsRecursive(const UAnimNextRigVMAssetEdito
 		{
 			if(ensure(EditorObject->GetModel()))
 			{
-				const URigVMCollapseNode* CollapseNode = CastChecked<URigVMCollapseNode>(EditorObject->GetModel()->GetOuter());
-				const FSoftObjectPath EditorObjectSoftObjPath = EditorObject;
-				ensureMsgf(EditorObjectSoftObjPath.IsSubobject(), TEXT("EditorObject for RigVMCollapseNode Graph was not a subobject as expected."));
-
-				const FWorkspaceOutlinerItemExport& ParentExport = ParentExportIndex == INDEX_NONE ? RootExport : OutExports.Exports[ParentExportIndex];
-				const FName Identifier = EditorObjectSoftObjPath.IsSubobject() ? *EditorObjectSoftObjPath.GetSubPathUtf8String() : CollapseNode->GetFName();
-				FWorkspaceOutlinerItemExport& Export = OutExports.Exports.Add_GetRef(FWorkspaceOutlinerItemExport(Identifier, ParentExport));
-				Export.GetData().InitializeAsScriptStruct(FAnimNextCollapseGraphOutlinerData::StaticStruct());
-			
-				FAnimNextCollapseGraphOutlinerData& FnGraphData = Export.GetData().GetMutable<FAnimNextCollapseGraphOutlinerData>();
-				FnGraphData.SoftEditorObject = EditorObject;
-
-				int32 ExportIndex = OutExports.Exports.Num() - 1;
-				GetSubGraphWorkspaceExportsRecursive(EditorData, OutExports, RootExport, ExportIndex, EditorObject, Context);
-			}
-		}
-	}
-
-	// ---- Function References ---
-	TArray<URigVMEdGraphNode*> EdNodes;
-	RigVMEdGraph->GetNodesOfClass(EdNodes);
-
-	for (URigVMEdGraphNode* EdNode : EdNodes)
-	{
-		if (URigVMFunctionReferenceNode* FunctionReferenceNode = Cast<URigVMFunctionReferenceNode>(EdNode->GetModelNode()))
-		{
-			// Only export referenced functions which are part of same outer
-			const FRigVMGraphFunctionIdentifier FunctionIdentifier = FunctionReferenceNode->GetFunctionIdentifier();
-			if (FunctionIdentifier.HostObject == FSoftObjectPath(EditorData))
-			{
-				const URigVMLibraryNode* FunctionNode = EditorData->RigVMClient.GetFunctionLibrary()->FindFunction(FunctionIdentifier.GetFunctionFName());
-				if (const URigVMGraph* ContainedGraph = FunctionNode->GetContainedGraph())
+				// Ignore aggregate nodes
+				const URigVMCollapseNode* CollapseNode = ExactCast<URigVMCollapseNode>(EditorObject->GetModel()->GetOuter());
+				if (CollapseNode)
 				{
-					URigVMEdGraph* ContainedEdGraph = Cast<URigVMEdGraph>(EditorData->GetEditorObjectForRigVMGraph(ContainedGraph));
-					const FSoftObjectPath EditorObjectSoftObjPath = ContainedEdGraph;
-					ensureMsgf(EditorObjectSoftObjPath.IsSubobject(), TEXT("EditorObject for RigVMFunctionReferenceNode Graph was not a subobject as expected."));
+					const FSoftObjectPath EditorObjectSoftObjPath = EditorObject;
+					ensureMsgf(EditorObjectSoftObjPath.IsSubobject(), TEXT("EditorObject for RigVMCollapseNode Graph was not a subobject as expected."));
 
 					const FWorkspaceOutlinerItemExport& ParentExport = ParentExportIndex == INDEX_NONE ? RootExport : OutExports.Exports[ParentExportIndex];
-					const FName Identifier = EditorObjectSoftObjPath.IsSubobject() ? *EditorObjectSoftObjPath.GetSubPathUtf8String() : FunctionNode->GetFName();
+					const FName Identifier = EditorObjectSoftObjPath.IsSubobject() ? *EditorObjectSoftObjPath.GetSubPathUtf8String() : CollapseNode->GetFName();
 					FWorkspaceOutlinerItemExport& Export = OutExports.Exports.Add_GetRef(FWorkspaceOutlinerItemExport(Identifier, ParentExport));
-					Export.GetData().InitializeAsScriptStruct(FAnimNextGraphFunctionOutlinerData::StaticStruct());
-					FAnimNextGraphFunctionOutlinerData& FnGraphData = Export.GetData().GetMutable<FAnimNextGraphFunctionOutlinerData>();
-
-					if (ContainedEdGraph)
-					{
-						FnGraphData.SoftEditorObject = ContainedEdGraph;
-						FnGraphData.SoftEdGraphNode = EdNode;
+					Export.GetData().InitializeAsScriptStruct(FAnimNextCollapseGraphOutlinerData::StaticStruct());
+					FAnimNextCollapseGraphOutlinerData& FnGraphData = Export.GetData().GetMutable<FAnimNextCollapseGraphOutlinerData>();
+					FnGraphData.SoftEditorObject = EditorObject;
 
-						int32 ExportIndex = OutExports.Exports.Num() - 1;
-						GetSubGraphWorkspaceExportsRecursive(EditorData, OutExports, RootExport, ExportIndex, ContainedEdGraph, Context);
-					}
+					int32 ExportIndex = OutExports.Exports.Num() - 1;
+					GetSubGraphWorkspaceExportsRecursive(EditorData, OutExports, RootExport, ExportIndex, EditorObject, Context);
 				}
 			}
 		}
@@ -1345,23 +1456,13 @@ void FUtils::GetFunctionsWorkspaceExportsRecursive(const UAnimNextRigVMAssetEdit
 {
 	for (const FRigVMGraphFunctionData& FunctionData : Functions)
 	{
+		const FWorkspaceOutlinerItemExport& ParentExport = ParentExportIndex == INDEX_NONE ? RootExport : OutExports.Exports[ParentExportIndex];
 		if (const URigVMLibraryNode* FunctionNode = EditorData->RigVMClient.GetFunctionLibrary()->FindFunction(FunctionData.Header.LibraryPointer.GetFunctionFName()))
 		{
-			if (const URigVMGraph* ContainedModelGraph = FunctionNode->GetContainedGraph())
+			const FWorkspaceOutlinerItemExport Export = MakeFunctionExport(FunctionNode, ParentExport);
+			if (Export.GetIdentifier() != NAME_None)
 			{
-				if (URigVMEdGraph* EditorObject = Cast<URigVMEdGraph>(EditorData->GetEditorObjectForRigVMGraph(ContainedModelGraph)))
-				{
-					const FSoftObjectPath EditorObjectSoftObjPath = EditorObject;
-					ensureMsgf(EditorObjectSoftObjPath.IsSubobject(), TEXT("EditorObject for RigVMFunctionReferenceNode Graph was not a subobject as expected."));
-
-					const FWorkspaceOutlinerItemExport& ParentExport = ParentExportIndex == INDEX_NONE ? RootExport : OutExports.Exports[ParentExportIndex];
-					const FName Identifier = EditorObjectSoftObjPath.IsSubobject() ? *EditorObjectSoftObjPath.GetSubPathUtf8String() : FunctionNode->GetFName();
-					FWorkspaceOutlinerItemExport& Export = OutExports.Exports.Add_GetRef(FWorkspaceOutlinerItemExport(Identifier, ParentExport));
-
-					Export.GetData().InitializeAsScriptStruct(FAnimNextGraphFunctionOutlinerData::StaticStruct());
-					FAnimNextGraphFunctionOutlinerData& FnGraphData = Export.GetData().GetMutable<FAnimNextGraphFunctionOutlinerData>();
-					FnGraphData.SoftEditorObject = EditorObject;
-				}
+				OutExports.Exports.Add(Export);
 			}
 		}
 	}
@@ -1568,7 +1669,7 @@ void FUtils::RenameVariable(UAnimNextVariableEntry* VariableEntry, const FName N
 	}	
 }
 
-void FUtils::MoveVariableToAsset(UAnimNextVariableEntry* VariableEntry, UAnimNextRigVMAsset* NewOuter, bool bSetupUndoRedo, bool bPrintPythonCommands)
+UAnimNextVariableEntry* FUtils::MoveVariableToAsset(UAnimNextVariableEntry* VariableEntry, UAnimNextRigVMAsset* NewOuter, bool bSetupUndoRedo, bool bPrintPythonCommands)
 {
 	if (VariableEntry && NewOuter)
 	{
@@ -1582,7 +1683,8 @@ void FUtils::MoveVariableToAsset(UAnimNextVariableEntry* VariableEntry, UAnimNex
 			FString DefaultValueString;
 			VariableEntry->GetDefaultValueString(DefaultValueString);
 
-			check(NewEditorData->AddVariable(VariableEntry->GetVariableName(), VariableEntry->GetType(), DefaultValueString, bSetupUndoRedo, bPrintPythonCommands));
+			UAnimNextVariableEntry* NewEntry = NewEditorData->AddVariable(VariableEntry->GetVariableName(), VariableEntry->GetType(), DefaultValueString, bSetupUndoRedo, bPrintPythonCommands);
+			check(NewEntry);
 
 			const FAnimNextSoftVariableReference FindReference(VariableEntry->GetEntryName(), CurrentOuter);
 			const FAnimNextSoftVariableReference ReplaceReference(VariableEntry->GetEntryName(), NewOuter);
@@ -1600,8 +1702,12 @@ void FUtils::MoveVariableToAsset(UAnimNextVariableEntry* VariableEntry, UAnimNex
 			}
 
 			CurrentEditorData->RemoveEntry(VariableEntry, bSetupUndoRedo, bPrintPythonCommands);
+			
+			return NewEntry;
 		}
 	}
+	
+	return nullptr;
 }
 
 void FUtils::SetVariableType(UAnimNextVariableEntry* VariableEntry, const FAnimNextParamType NewType, bool bSetupUndoRedo, bool bPrintPythonCommands)
@@ -1671,8 +1777,9 @@ void FUtils::ReplaceVariableReferences(UAnimNextRigVMAssetEditorData* InEditorDa
 	const bool bVariableTypeChanged = FindReferenceType != ReplaceReferenceType;
 	const bool bVariableNameChanged = ReferenceToFind.GetName() != ReferenceToReplaceWith.GetName();
 
-	// Ensure that SharedVariables have been included for replace variable
-	if (const UAnimNextSharedVariables* SharedVariablesReplaceOuter = Cast<UAnimNextSharedVariables>(ReplaceOuterAsset))
+	// Ensure that SharedVariables have been included for replace variable (if from another asset)
+	const UAnimNextSharedVariables* SharedVariablesReplaceOuter = Cast<UAnimNextSharedVariables>(ReplaceOuterAsset);
+	if (SharedVariablesReplaceOuter && ThisOuterAsset != SharedVariablesReplaceOuter)
 	{
 		bool bHasEntry = false;
 		InEditorData->ForEachEntryOfType<UAnimNextSharedVariablesEntry>([&bHasEntry, SharedVariablesReplaceOuter](const UAnimNextSharedVariablesEntry* Entry)
@@ -1691,9 +1798,21 @@ void FUtils::ReplaceVariableReferences(UAnimNextRigVMAssetEditorData* InEditorDa
 
 		if (!bHasEntry)
 		{
-			UE::UAF::UncookedOnly::FScopedCompileJob CompilerResults(LOCTEXT("AddSharedAssetJobName", "Add Shared Asset dependency"), { ThisOuterAsset });
-			InEditorData->AddSharedVariables(SharedVariablesReplaceOuter, bSetupUndoRedo, bPrintPythonCommands);
-			InEditorData->RecompileVM();
+			FString ErrorMessage;
+			if (CanAddSharedVariablesReference(InEditorData, SharedVariablesReplaceOuter, &ErrorMessage))
+			{
+				UE::UAF::UncookedOnly::FScopedCompileJob CompilerResults(LOCTEXT("AddSharedAssetJobName", "Add Shared Asset dependency"), { ThisOuterAsset });
+				InEditorData->AddSharedVariables(SharedVariablesReplaceOuter, bSetupUndoRedo, bPrintPythonCommands);
+				InEditorData->RecompileVM();
+			}
+			else
+			{
+				FText NotificationText = FText::FromString(ErrorMessage);
+				FNotificationInfo Notification(NotificationText);
+				Notification.ExpireDuration = 3.0f;
+				FSlateNotificationManager::Get().AddNotification(Notification);
+				UE_LOGFMT(LogAnimation, Error, "{0}", ErrorMessage);
+			}
 		}
 	}			
 	
diff --git a/Engine/Plugins/Experimental/Workspace/Source/WorkspaceEditor/Private/WorkspaceEditor.cpp b/Engine/Plugins/Experimental/Workspace/Source/WorkspaceEditor/Private/WorkspaceEditor.cpp
--- a/Engine/Plugins/Experimental/Workspace/Source/WorkspaceEditor/Private/WorkspaceEditor.cpp
+++ b/Engine/Plugins/Experimental/Workspace/Source/WorkspaceEditor/Private/WorkspaceEditor.cpp
@@ -141,6 +141,7 @@ void FWorkspaceEditor::CreateWidgets()
 					->SetHideTabWell(false)
 					->AddTab(WorkspaceTabs::TopLeftDocumentArea, ETabState::ClosedTab)
 					->AddTab(FBaseAssetToolkit::ViewportTabID, ETabState::OpenedTab)
+					->SetExtensionId(WorkspaceTabs::TopLeftDocumentArea)
 				)
 				->Split
 				(
@@ -148,6 +149,7 @@ void FWorkspaceEditor::CreateWidgets()
 					->SetSizeCoefficient(0.5f)
 					->SetHideTabWell(false)
 					->AddTab(WorkspaceTabs::BottomLeftDocumentArea, ETabState::ClosedTab)
+					->SetExtensionId(WorkspaceTabs::BottomLeftDocumentArea)
 				)
 			)
 			->Split
@@ -161,6 +163,7 @@ void FWorkspaceEditor::CreateWidgets()
 					->SetSizeCoefficient(0.75f)
 					->SetHideTabWell(false)
 					->AddTab(WorkspaceTabs::TopMiddleDocumentArea, ETabState::ClosedTab)
+					->SetExtensionId(WorkspaceTabs::TopMiddleDocumentArea)
 				)
 				->Split
 				(
@@ -168,6 +171,7 @@ void FWorkspaceEditor::CreateWidgets()
 					->SetSizeCoefficient(0.25f)
 					->SetHideTabWell(false)
 					->AddTab(WorkspaceTabs::BottomMiddleDocumentArea, ETabState::ClosedTab)
+					->SetExtensionId(WorkspaceTabs::BottomMiddleDocumentArea)
 				)
 			)
 			->Split
@@ -182,6 +186,7 @@ void FWorkspaceEditor::CreateWidgets()
 					->SetHideTabWell(false)
 					->AddTab(WorkspaceTabs::WorkspaceView, ETabState::OpenedTab)
 					->AddTab(WorkspaceTabs::TopRightDocumentArea, ETabState::ClosedTab)
+					->SetExtensionId(WorkspaceTabs::TopRightDocumentArea)
 				)
 				->Split
 				(
@@ -190,6 +195,7 @@ void FWorkspaceEditor::CreateWidgets()
 					->SetHideTabWell(false)
 					->AddTab(WorkspaceTabs::BottomRightDocumentArea, ETabState::ClosedTab)
 					->AddTab(FBaseAssetToolkit::DetailsTabID, ETabState::OpenedTab)
+					->SetExtensionId(WorkspaceTabs::BottomRightDocumentArea)
 				)
 			)
 		)
diff --git a/Engine/Plugins/Experimental/Workspace/Source/WorkspaceEditor/Public/WorkspaceAssetRegistryInfo.h b/Engine/Plugins/Experimental/Workspace/Source/WorkspaceEditor/Public/WorkspaceAssetRegistryInfo.h
--- a/Engine/Plugins/Experimental/Workspace/Source/WorkspaceEditor/Public/WorkspaceAssetRegistryInfo.h
+++ b/Engine/Plugins/Experimental/Workspace/Source/WorkspaceEditor/Public/WorkspaceAssetRegistryInfo.h
@@ -329,6 +329,29 @@ public:
 			}
 		}
 	}
+			
+	// Returns relative export to the provided InAssetPath
+	//  - "SoftObjectPath" - "Foo" - "SoftObjectPathTwo" - "Bar" will return "SoftObjectPath" - "Foo" when passed "SoftObjectPathTwo" or an empty export when passed "SoftObjectPath"
+	FWorkspaceOutlinerItemExport GetRelativeAssetExport(const FSoftObjectPath& InAssetPath) const
+	{
+		FWorkspaceOutlinerItemExport Export;
+		
+		if(Path.PathSegments.Num() > 0)
+		{
+			const int32 NumSegments = Path.PathSegments.Num();
+			for (int32 SegmentIndex = 0; SegmentIndex < NumSegments; SegmentIndex++)
+			{
+				FSoftObjectPath ObjectPath = Path.PathSegments[SegmentIndex].ToString();
+				if (ObjectPath.IsValid() && ObjectPath == InAssetPath && SegmentIndex != 0)
+				{
+					Export.Path.PathSegments.Append(&Path.PathSegments[0], SegmentIndex);
+					break;
+				}
+			}
+		}
+		
+		return Export;
+	}
 
 	FString GetFullPath() const
 	{
