[Backout] - CL46305822
#rnx
Original CL Desc
-----------------------------------------------------------------
Added common outliner implementation to be shared between variables/functions sharing:
- Hierarchy generation from assets
- Category/asset/entry items
- Common functionality (copy/paste/rename etc.)
- Drag-drop operations (registered with OutlinerMode rather than through item widgets)

Variables Outliner:
- Copy/cut/paste behavior
- Categories are now allowed to be empty, so now also able to them manually

Functions Outliner:
- Display and allows for editing RigVM functions
- Hierarchy matches Variables Outliner
- Currently limited to rename, duplicate, delete operations (rest to follow)
- Categories for functions are fixed, as they are tied to the Node objects themselves

Misc:
- Only show user generated Collapse node in Workspace outliner

diff --git a/Engine/Plugins/Experimental/UAF/UAF/Source/UAFEditor/Internal/Common/GraphEditorSchemaActions.h b/Engine/Plugins/Experimental/UAF/UAF/Source/UAFEditor/Internal/Common/GraphEditorSchemaActions.h
--- a/Engine/Plugins/Experimental/UAF/UAF/Source/UAFEditor/Internal/Common/GraphEditorSchemaActions.h
+++ b/Engine/Plugins/Experimental/UAF/UAF/Source/UAFEditor/Internal/Common/GraphEditorSchemaActions.h
@@ -136,8 +136,6 @@ struct FAnimNextSchemaAction_Function : public FAnimNextSchemaAction
 
 	FAnimNextSchemaAction_Function(const FRigVMGraphFunctionHeader& InReferencedPublicFunctionHeader, const FText& InNodeCategory, const FText& InMenuDesc, const FText& InToolTip, const FText& InKeywords = FText::GetEmpty());
 	FAnimNextSchemaAction_Function(const URigVMLibraryNode* InFunctionLibraryNode, const FText& InNodeCategory, const FText& InMenuDesc, const FText& InToolTip, const FText& InKeywords = FText::GetEmpty());
-	
-	FAnimNextSchemaAction_Function(const URigVMLibraryNode* InFunctionLibraryNode);
 
 	virtual const FSlateBrush* GetIconBrush() const override;
 	
diff --git a/Engine/Plugins/Experimental/UAF/UAF/Source/UAFEditor/Private/AnimNextEditorModule.cpp b/Engine/Plugins/Experimental/UAF/UAF/Source/UAFEditor/Private/AnimNextEditorModule.cpp
--- a/Engine/Plugins/Experimental/UAF/UAF/Source/UAFEditor/Private/AnimNextEditorModule.cpp
+++ b/Engine/Plugins/Experimental/UAF/UAF/Source/UAFEditor/Private/AnimNextEditorModule.cpp
@@ -73,8 +73,6 @@
 #include "AnimNextSharedVariableNode.h"
 #include "UAFStyle.h"
 #include "Common/AnimNextAssetFindReplaceVariables.h"
-#include "Common/Outliner/OutlinerCommands.h"
-#include "Functions/Outliner/FunctionsOutlinerView.h"
 #include "Toolkits/AssetEditorToolkitMenuContext.h"
 #include "UAF/Viewport/SystemViewportController.h"
 #include "Variables/Outliner/VariablesOutlinerCommands.h"
@@ -129,7 +127,6 @@ void FAnimNextEditorModule::StartupModule()
 	FRigVMCommands::Register();
 	FAnimNextRigVMAssetCommands::Register();
 	FVariablesOutlinerCommands::Register();
-	FOutlinerCommands::Register();
 	FUAFStyle::Get();
 
 	// Register settings for user editing
@@ -209,10 +206,6 @@ void FAnimNextEditorModule::StartupModule()
 			TabFactories.RegisterFactory(VariablesTabSummoner);
 			VariablesTabSummoner->RegisterTabSpawner(InTabManager, nullptr);
 
-			TSharedRef<FAnimNextFunctionsTabSummoner> FunctionsTabSummoner = MakeShared<FAnimNextFunctionsTabSummoner>(InEditorPtr);
-			TabFactories.RegisterFactory(FunctionsTabSummoner);
-			FunctionsTabSummoner->RegisterTabSpawner(InTabManager, nullptr);
-
 			TSharedRef<FAnimNextFindAndReplaceTabSummoner> FindAndReplaceTabSummoner = MakeShared<FAnimNextFindAndReplaceTabSummoner>(InEditorPtr);
 			TabFactories.RegisterFactory(FindAndReplaceTabSummoner);
 			FindAndReplaceTabSummoner->RegisterTabSpawner(InTabManager, nullptr);
@@ -230,10 +223,7 @@ void FAnimNextEditorModule::StartupModule()
 		InLayoutExtender.ExtendLayout(FTabId(Workspace::WorkspaceTabs::BottomMiddleDocumentArea), ELayoutExtensionPosition::After, FindAndReplaceTab);
 
 		FTabManager::FTab VariablesTab(FTabId(VariablesTabName), ETabState::OpenedTab);
-		InLayoutExtender.ExtendStack(Workspace::WorkspaceTabs::BottomLeftDocumentArea, ELayoutExtensionPosition::After, VariablesTab, true);
-
-		FTabManager::FTab FunctionsTab(FTabId(FunctionsTabName), ETabState::OpenedTab);
-		InLayoutExtender.ExtendStack(Workspace::WorkspaceTabs::BottomLeftDocumentArea, ELayoutExtensionPosition::After, FunctionsTab);
+		InLayoutExtender.ExtendLayout(FTabId(Workspace::WorkspaceTabs::BottomLeftDocumentArea), ELayoutExtensionPosition::After, VariablesTab);
 	});
 
 	UToolMenu* AnimNextEdGraphContextMenu = UToolMenus::Get()->ExtendMenu( "GraphEditor.GraphNodeContextMenu.AnimNextEdGraphNode");
@@ -312,11 +302,11 @@ void FAnimNextEditorModule::StartupModule()
 											}
 										};
 
-										InSubmenu->AddMenuEntry(NAME_None, FToolMenuEntry::InitMenuEntry("FindGlobalReferences", FOutlinerCommands::Get().FindReferences->GetLabel(), FOutlinerCommands::Get().FindReferences->GetDescription(), UAnimNextAssetFindReplaceVariables::GetIconFromSearchScope(ESearchScope::Global), FToolUIAction(FToolMenuExecuteAction::CreateLambda(InvokeFindVariables, ESearchScope::Global))));
+										InSubmenu->AddMenuEntry(NAME_None, FToolMenuEntry::InitMenuEntry("FindGlobalReferences", FVariablesOutlinerCommands::Get().FindReferences->GetLabel(), FVariablesOutlinerCommands::Get().FindReferences->GetDescription(), UAnimNextAssetFindReplaceVariables::GetIconFromSearchScope(ESearchScope::Global), FToolUIAction(FToolMenuExecuteAction::CreateLambda(InvokeFindVariables, ESearchScope::Global))));
 
-										InSubmenu->AddMenuEntry(NAME_None, FToolMenuEntry::InitMenuEntry("FindWorkspaceReferences", FOutlinerCommands::Get().FindReferencesInWorkspace->GetLabel(), FOutlinerCommands::Get().FindReferencesInWorkspace->GetDescription(), UAnimNextAssetFindReplaceVariables::GetIconFromSearchScope(ESearchScope::Workspace), FToolUIAction(FToolMenuExecuteAction::CreateLambda(InvokeFindVariables, ESearchScope::Workspace))));
+										InSubmenu->AddMenuEntry(NAME_None, FToolMenuEntry::InitMenuEntry("FindWorkspaceReferences", FVariablesOutlinerCommands::Get().FindReferencesInWorkspace->GetLabel(), FVariablesOutlinerCommands::Get().FindReferencesInWorkspace->GetDescription(), UAnimNextAssetFindReplaceVariables::GetIconFromSearchScope(ESearchScope::Workspace), FToolUIAction(FToolMenuExecuteAction::CreateLambda(InvokeFindVariables, ESearchScope::Workspace))));
 
-										InSubmenu->AddMenuEntry(NAME_None, FToolMenuEntry::InitMenuEntry("FindAssetReferences", FOutlinerCommands::Get().FindReferencesInAsset->GetLabel(), FOutlinerCommands::Get().FindReferencesInAsset->GetDescription(), UAnimNextAssetFindReplaceVariables::GetIconFromSearchScope(ESearchScope::Asset), FToolUIAction(FToolMenuExecuteAction::CreateLambda(InvokeFindVariables, ESearchScope::Asset))));
+										InSubmenu->AddMenuEntry(NAME_None, FToolMenuEntry::InitMenuEntry("FindAssetReferences", FVariablesOutlinerCommands::Get().FindReferencesInAsset->GetLabel(), FVariablesOutlinerCommands::Get().FindReferencesInAsset->GetDescription(), UAnimNextAssetFindReplaceVariables::GetIconFromSearchScope(ESearchScope::Asset), FToolUIAction(FToolMenuExecuteAction::CreateLambda(InvokeFindVariables, ESearchScope::Asset))));
 									}),
 									FToolUIActionChoice(),
 									EUserInterfaceActionType::Button,
diff --git a/Engine/Plugins/Experimental/UAF/UAF/Source/UAFEditor/Private/Common/AnimNextFunctionItemDetails.cpp b/Engine/Plugins/Experimental/UAF/UAF/Source/UAFEditor/Private/Common/AnimNextFunctionItemDetails.cpp
--- a/Engine/Plugins/Experimental/UAF/UAF/Source/UAFEditor/Private/Common/AnimNextFunctionItemDetails.cpp
+++ b/Engine/Plugins/Experimental/UAF/UAF/Source/UAFEditor/Private/Common/AnimNextFunctionItemDetails.cpp
@@ -82,6 +82,7 @@ bool FAnimNextFunctionItemDetails::CanDelete(const FWorkspaceOutlinerItemExport&
 
 void FAnimNextFunctionItemDetails::Delete(TConstArrayView<FWorkspaceOutlinerItemExport> Exports) const
 {
+	TMap<UAnimNextRigVMAssetEditorData*, TArray<UAnimNextRigVMAssetEntry*>> EntriesToDelete;
 	for (const FWorkspaceOutlinerItemExport& Export : Exports)
 	{
 		const TInstancedStruct<FWorkspaceOutlinerItemData>& Data = Export.GetData();
@@ -169,7 +170,10 @@ void FAnimNextFunctionItemDetails::Rename(const FWorkspaceOutlinerItemExport& Ex
 
 					FScopedTransaction Transaction(LOCTEXT("RenameFunctionInOutliner", "Rename Function"));
 					FRigVMControllerCompileBracketScope CompileScope(EdGraph->GetController());
-					GraphSchema->TryRenameGraph(EdGraph, *InName.ToString());
+					if (GraphSchema->TryRenameGraph(EdGraph, *InName.ToString()))
+					{
+						return;
+					}
 				}
 			}
 		}
diff --git a/Engine/Plugins/Experimental/UAF/UAF/Source/UAFEditor/Private/Common/GraphEditorSchemaActions.cpp b/Engine/Plugins/Experimental/UAF/UAF/Source/UAFEditor/Private/Common/GraphEditorSchemaActions.cpp
--- a/Engine/Plugins/Experimental/UAF/UAF/Source/UAFEditor/Private/Common/GraphEditorSchemaActions.cpp
+++ b/Engine/Plugins/Experimental/UAF/UAF/Source/UAFEditor/Private/Common/GraphEditorSchemaActions.cpp
@@ -422,17 +422,6 @@ FAnimNextSchemaAction_Function::FAnimNextSchemaAction_Function(const URigVMLibra
 	bIsLocalFunction = true;
 }
 
-FAnimNextSchemaAction_Function::FAnimNextSchemaAction_Function(const URigVMLibraryNode* InFunctionLibraryNode)
-	: FAnimNextSchemaAction(InFunctionLibraryNode ? FText::FromString(InFunctionLibraryNode->GetNodeCategory()) : FText::GetEmpty(), InFunctionLibraryNode ? FText::FromString(InFunctionLibraryNode->GetNodeDescription()) : FText::GetEmpty(), InFunctionLibraryNode ? InFunctionLibraryNode->GetToolTipText() : FText::GetEmpty(), InFunctionLibraryNode ? FText::FromString(InFunctionLibraryNode->GetNodeKeywords()) : FText::GetEmpty())
-{
-	if (InFunctionLibraryNode)
-	{
-		ReferencedPublicFunctionHeader = InFunctionLibraryNode->GetFunctionHeader();
-	}
-	NodeClass = UAnimNextEdGraphNode::StaticClass();
-	bIsLocalFunction = true;
-}
-
 const FSlateBrush* FAnimNextSchemaAction_Function::GetIconBrush() const
 {
 	return FAppStyle::GetBrush("GraphEditor.Function_16x");
diff --git a/Engine/Plugins/Experimental/UAF/UAF/Source/UAFEditor/Private/Common/Outliner/CommonOutlinerMode.cpp b/Engine/Plugins/Experimental/UAF/UAF/Source/UAFEditor/Private/Common/Outliner/CommonOutlinerMode.cpp
deleted file mode 100644
diff --git a/Engine/Plugins/Experimental/UAF/UAF/Source/UAFEditor/Private/Common/Outliner/CommonOutlinerMode.h b/Engine/Plugins/Experimental/UAF/UAF/Source/UAFEditor/Private/Common/Outliner/CommonOutlinerMode.h
deleted file mode 100644
diff --git a/Engine/Plugins/Experimental/UAF/UAF/Source/UAFEditor/Private/Common/Outliner/OutlinerAssetItem.cpp b/Engine/Plugins/Experimental/UAF/UAF/Source/UAFEditor/Private/Common/Outliner/OutlinerAssetItem.cpp
deleted file mode 100644
diff --git a/Engine/Plugins/Experimental/UAF/UAF/Source/UAFEditor/Private/Common/Outliner/OutlinerAssetItem.h b/Engine/Plugins/Experimental/UAF/UAF/Source/UAFEditor/Private/Common/Outliner/OutlinerAssetItem.h
deleted file mode 100644
diff --git a/Engine/Plugins/Experimental/UAF/UAF/Source/UAFEditor/Private/Common/Outliner/OutlinerCategoryItem.cpp b/Engine/Plugins/Experimental/UAF/UAF/Source/UAFEditor/Private/Common/Outliner/OutlinerCategoryItem.cpp
deleted file mode 100644
diff --git a/Engine/Plugins/Experimental/UAF/UAF/Source/UAFEditor/Private/Common/Outliner/OutlinerCategoryItem.h b/Engine/Plugins/Experimental/UAF/UAF/Source/UAFEditor/Private/Common/Outliner/OutlinerCategoryItem.h
deleted file mode 100644
diff --git a/Engine/Plugins/Experimental/UAF/UAF/Source/UAFEditor/Private/Common/Outliner/OutlinerColumns.cpp b/Engine/Plugins/Experimental/UAF/UAF/Source/UAFEditor/Private/Common/Outliner/OutlinerColumns.cpp
deleted file mode 100644
diff --git a/Engine/Plugins/Experimental/UAF/UAF/Source/UAFEditor/Private/Common/Outliner/OutlinerColumns.h b/Engine/Plugins/Experimental/UAF/UAF/Source/UAFEditor/Private/Common/Outliner/OutlinerColumns.h
deleted file mode 100644
diff --git a/Engine/Plugins/Experimental/UAF/UAF/Source/UAFEditor/Private/Common/Outliner/OutlinerCommands.cpp b/Engine/Plugins/Experimental/UAF/UAF/Source/UAFEditor/Private/Common/Outliner/OutlinerCommands.cpp
deleted file mode 100644
diff --git a/Engine/Plugins/Experimental/UAF/UAF/Source/UAFEditor/Private/Common/Outliner/OutlinerCommands.h b/Engine/Plugins/Experimental/UAF/UAF/Source/UAFEditor/Private/Common/Outliner/OutlinerCommands.h
deleted file mode 100644
diff --git a/Engine/Plugins/Experimental/UAF/UAF/Source/UAFEditor/Private/Common/Outliner/OutlinerDragDropOps.cpp b/Engine/Plugins/Experimental/UAF/UAF/Source/UAFEditor/Private/Common/Outliner/OutlinerDragDropOps.cpp
deleted file mode 100644
diff --git a/Engine/Plugins/Experimental/UAF/UAF/Source/UAFEditor/Private/Common/Outliner/OutlinerDragDropOps.h b/Engine/Plugins/Experimental/UAF/UAF/Source/UAFEditor/Private/Common/Outliner/OutlinerDragDropOps.h
deleted file mode 100644
diff --git a/Engine/Plugins/Experimental/UAF/UAF/Source/UAFEditor/Private/Common/Outliner/OutlinerEntryItem.cpp b/Engine/Plugins/Experimental/UAF/UAF/Source/UAFEditor/Private/Common/Outliner/OutlinerEntryItem.cpp
deleted file mode 100644
diff --git a/Engine/Plugins/Experimental/UAF/UAF/Source/UAFEditor/Private/Common/Outliner/OutlinerEntryItem.h b/Engine/Plugins/Experimental/UAF/UAF/Source/UAFEditor/Private/Common/Outliner/OutlinerEntryItem.h
deleted file mode 100644
diff --git a/Engine/Plugins/Experimental/UAF/UAF/Source/UAFEditor/Private/Common/Outliner/OutlinerEntryItemWidget.cpp b/Engine/Plugins/Experimental/UAF/UAF/Source/UAFEditor/Private/Common/Outliner/OutlinerEntryItemWidget.cpp
deleted file mode 100644
diff --git a/Engine/Plugins/Experimental/UAF/UAF/Source/UAFEditor/Private/Common/Outliner/OutlinerEntryItemWidget.h b/Engine/Plugins/Experimental/UAF/UAF/Source/UAFEditor/Private/Common/Outliner/OutlinerEntryItemWidget.h
deleted file mode 100644
diff --git a/Engine/Plugins/Experimental/UAF/UAF/Source/UAFEditor/Private/Common/Outliner/OutlinerHierarchy.cpp b/Engine/Plugins/Experimental/UAF/UAF/Source/UAFEditor/Private/Common/Outliner/OutlinerHierarchy.cpp
deleted file mode 100644
diff --git a/Engine/Plugins/Experimental/UAF/UAF/Source/UAFEditor/Private/Common/Outliner/OutlinerHierarchy.h b/Engine/Plugins/Experimental/UAF/UAF/Source/UAFEditor/Private/Common/Outliner/OutlinerHierarchy.h
deleted file mode 100644
diff --git a/Engine/Plugins/Experimental/UAF/UAF/Source/UAFEditor/Private/Common/Outliner/OutlinerItem.cpp b/Engine/Plugins/Experimental/UAF/UAF/Source/UAFEditor/Private/Common/Outliner/OutlinerItem.cpp
deleted file mode 100644
diff --git a/Engine/Plugins/Experimental/UAF/UAF/Source/UAFEditor/Private/Common/Outliner/OutlinerItem.h b/Engine/Plugins/Experimental/UAF/UAF/Source/UAFEditor/Private/Common/Outliner/OutlinerItem.h
deleted file mode 100644
diff --git a/Engine/Plugins/Experimental/UAF/UAF/Source/UAFEditor/Private/Common/Outliner/OutlinerItemMenuContexts.h b/Engine/Plugins/Experimental/UAF/UAF/Source/UAFEditor/Private/Common/Outliner/OutlinerItemMenuContexts.h
deleted file mode 100644
diff --git a/Engine/Plugins/Experimental/UAF/UAF/Source/UAFEditor/Private/Common/Outliner/SCommonOutliner.cpp b/Engine/Plugins/Experimental/UAF/UAF/Source/UAFEditor/Private/Common/Outliner/SCommonOutliner.cpp
deleted file mode 100644
diff --git a/Engine/Plugins/Experimental/UAF/UAF/Source/UAFEditor/Private/Common/Outliner/SCommonOutliner.h b/Engine/Plugins/Experimental/UAF/UAF/Source/UAFEditor/Private/Common/Outliner/SCommonOutliner.h
deleted file mode 100644
diff --git a/Engine/Plugins/Experimental/UAF/UAF/Source/UAFEditor/Private/Functions/Outliner/FunctionOutlinerEntry.cpp b/Engine/Plugins/Experimental/UAF/UAF/Source/UAFEditor/Private/Functions/Outliner/FunctionOutlinerEntry.cpp
deleted file mode 100644
diff --git a/Engine/Plugins/Experimental/UAF/UAF/Source/UAFEditor/Private/Functions/Outliner/FunctionOutlinerEntry.h b/Engine/Plugins/Experimental/UAF/UAF/Source/UAFEditor/Private/Functions/Outliner/FunctionOutlinerEntry.h
deleted file mode 100644
diff --git a/Engine/Plugins/Experimental/UAF/UAF/Source/UAFEditor/Private/Functions/Outliner/FunctionsOutlinerCategoryItem.cpp b/Engine/Plugins/Experimental/UAF/UAF/Source/UAFEditor/Private/Functions/Outliner/FunctionsOutlinerCategoryItem.cpp
deleted file mode 100644
diff --git a/Engine/Plugins/Experimental/UAF/UAF/Source/UAFEditor/Private/Functions/Outliner/FunctionsOutlinerCategoryItem.h b/Engine/Plugins/Experimental/UAF/UAF/Source/UAFEditor/Private/Functions/Outliner/FunctionsOutlinerCategoryItem.h
deleted file mode 100644
diff --git a/Engine/Plugins/Experimental/UAF/UAF/Source/UAFEditor/Private/Functions/Outliner/FunctionsOutlinerColumns.cpp b/Engine/Plugins/Experimental/UAF/UAF/Source/UAFEditor/Private/Functions/Outliner/FunctionsOutlinerColumns.cpp
deleted file mode 100644
diff --git a/Engine/Plugins/Experimental/UAF/UAF/Source/UAFEditor/Private/Functions/Outliner/FunctionsOutlinerColumns.h b/Engine/Plugins/Experimental/UAF/UAF/Source/UAFEditor/Private/Functions/Outliner/FunctionsOutlinerColumns.h
deleted file mode 100644
diff --git a/Engine/Plugins/Experimental/UAF/UAF/Source/UAFEditor/Private/Functions/Outliner/FunctionsOutlinerHierarchy.cpp b/Engine/Plugins/Experimental/UAF/UAF/Source/UAFEditor/Private/Functions/Outliner/FunctionsOutlinerHierarchy.cpp
deleted file mode 100644
diff --git a/Engine/Plugins/Experimental/UAF/UAF/Source/UAFEditor/Private/Functions/Outliner/FunctionsOutlinerHierarchy.h b/Engine/Plugins/Experimental/UAF/UAF/Source/UAFEditor/Private/Functions/Outliner/FunctionsOutlinerHierarchy.h
deleted file mode 100644
diff --git a/Engine/Plugins/Experimental/UAF/UAF/Source/UAFEditor/Private/Functions/Outliner/FunctionsOutlinerMode.cpp b/Engine/Plugins/Experimental/UAF/UAF/Source/UAFEditor/Private/Functions/Outliner/FunctionsOutlinerMode.cpp
deleted file mode 100644
diff --git a/Engine/Plugins/Experimental/UAF/UAF/Source/UAFEditor/Private/Functions/Outliner/FunctionsOutlinerMode.h b/Engine/Plugins/Experimental/UAF/UAF/Source/UAFEditor/Private/Functions/Outliner/FunctionsOutlinerMode.h
deleted file mode 100644
diff --git a/Engine/Plugins/Experimental/UAF/UAF/Source/UAFEditor/Private/Functions/Outliner/FunctionsOutlinerView.cpp b/Engine/Plugins/Experimental/UAF/UAF/Source/UAFEditor/Private/Functions/Outliner/FunctionsOutlinerView.cpp
deleted file mode 100644
diff --git a/Engine/Plugins/Experimental/UAF/UAF/Source/UAFEditor/Private/Functions/Outliner/FunctionsOutlinerView.h b/Engine/Plugins/Experimental/UAF/UAF/Source/UAFEditor/Private/Functions/Outliner/FunctionsOutlinerView.h
deleted file mode 100644
diff --git a/Engine/Plugins/Experimental/UAF/UAF/Source/UAFEditor/Private/Variables/AnimNextVariableItemMenuContext.h b/Engine/Plugins/Experimental/UAF/UAF/Source/UAFEditor/Private/Variables/AnimNextVariableItemMenuContext.h
new file mode 100644
--- /dev/null
+++ b/Engine/Plugins/Experimental/UAF/UAF/Source/UAFEditor/Private/Variables/AnimNextVariableItemMenuContext.h
@@ -0,0 +1,40 @@
+﻿// Copyright Epic Games, Inc. All Rights Reserved.
+
+#pragma once
+
+#include "AnimNextVariableItemMenuContext.generated.h"
+
+class FUICommandList;
+class UAnimNextRigVMAssetEditorData;
+class UAnimNextRigVMAssetEntry;
+
+namespace UE::Workspace
+{
+	class IWorkspaceEditor;
+}
+
+namespace UE::UAF::Editor
+{
+	class SVariablesOutliner;
+}
+
+UCLASS()
+class UAnimNextVariableItemMenuContext : public UObject
+{
+	GENERATED_BODY()
+
+public:
+	// Currently selected asset's editor data
+	UPROPERTY()
+	TArray<TWeakObjectPtr<UAnimNextRigVMAssetEditorData>> WeakEditorDatas;
+
+	// Currently selected entries
+	UPROPERTY()
+	TArray<TWeakObjectPtr<UAnimNextRigVMAssetEntry>> WeakEntries;
+
+	TWeakPtr<UE::UAF::Editor::SVariablesOutliner> WeakOutliner;
+
+	TWeakPtr<UE::Workspace::IWorkspaceEditor> WeakWorkspaceEditor;
+
+	TWeakPtr<FUICommandList> WeakCommandList;
+};
diff --git a/Engine/Plugins/Experimental/UAF/UAF/Source/UAFEditor/Private/Variables/Outliner/VariablesOutlinerAssetItem.cpp b/Engine/Plugins/Experimental/UAF/UAF/Source/UAFEditor/Private/Variables/Outliner/VariablesOutlinerAssetItem.cpp
new file mode 100644
--- /dev/null
+++ b/Engine/Plugins/Experimental/UAF/UAF/Source/UAFEditor/Private/Variables/Outliner/VariablesOutlinerAssetItem.cpp
@@ -0,0 +1,742 @@
+// Copyright Epic Games, Inc. All Rights Reserved.
+
+#include "VariablesOutlinerAssetItem.h"
+
+#include "ISceneOutliner.h"
+#include "Styling/SlateColor.h"
+#include "Misc/Attribute.h"
+#include "Misc/Optional.h"
+#include "SceneOutlinerStandaloneTypes.h"
+#include "UAFStyle.h"
+#include "AnimNextRigVMAsset.h"
+#include "AnimNextRigVMAssetEditorData.h"
+#include "AssetToolsModule.h"
+#include "AssetViewUtils.h"
+#include "IAssetTools.h"
+#include "ScopedTransaction.h"
+#include "UncookedOnlyUtils.h"
+#include "VariablesOutlinerDragDrop.h"
+#include "VariablesOutlinerMode.h"
+#include "WorkspaceDragDropOperation.h"
+#include "Entries/AnimNextSharedVariablesEntry.h"
+#include "Styling/SlateIconFinder.h"
+#include "Styling/StarshipCoreStyle.h"
+#include "UObject/Package.h"
+#include "Variables/AnimNextVariableItemMenuContext.h"
+#include "Variables/SVariablesView.h"
+#include "Widgets/SCompoundWidget.h"
+#include "Widgets/DeclarativeSyntaxSupport.h"
+#include "Widgets/Images/SImage.h"
+#include "Widgets/Images/SThrobber.h"
+#include "Widgets/Text/SInlineEditableTextBlock.h"
+#include "Widgets/Text/STextBlock.h"
+#include "Widgets/Views/STableRow.h"
+#include "Widgets/Layout/SSpacer.h"
+
+#define LOCTEXT_NAMESPACE "VariablesOutlinerTreeItem"
+
+namespace UE::UAF::Editor
+{
+
+const FSceneOutlinerTreeItemType FVariablesOutlinerAssetItem::Type;
+
+class SVariablesOutlinerAssetLabel : FSceneOutlinerCommonLabelData, public SCompoundWidget
+{
+	SLATE_BEGIN_ARGS(SVariablesOutlinerAssetLabel) {}
+	SLATE_END_ARGS()
+
+	void Construct(const FArguments& InArgs, FVariablesOutlinerAssetItem& InTreeItem, ISceneOutliner& SceneOutliner, const STableRow<FSceneOutlinerTreeItemPtr>& InRow)
+	{
+		WeakSceneOutliner = StaticCastSharedRef<ISceneOutliner>(SceneOutliner.AsShared());
+		TreeItem = StaticCastSharedRef<FVariablesOutlinerAssetItem>(InTreeItem.AsShared());
+		
+		TWeakPtr<SVariablesOutlinerAssetLabel> WeakLabel = StaticCastSharedRef<SVariablesOutlinerAssetLabel>(AsShared());
+		ChildSlot
+		[
+			SNew(SHorizontalBox)
+			+SHorizontalBox::Slot()
+			.AutoWidth()
+			.VAlign(VAlign_Center)
+			.Padding(2.0f, 2.0f)
+			[
+				SNew(SImage)
+				.Image(this, &SVariablesOutlinerAssetLabel::GetAssetIcon)
+			]
+			+SHorizontalBox::Slot()
+			.AutoWidth()
+			.VAlign(VAlign_Center)
+			.Padding(2.0f, 2.0f)
+			[
+				SAssignNew(TextBlock, SInlineEditableTextBlock)
+				.Font(FStyleFonts::Get().NormalBold)
+				.Text(this, &SVariablesOutlinerAssetLabel::GetDisplayText)
+				.HighlightText(SceneOutliner.GetFilterHighlightText())
+				.ColorAndOpacity(this, &SVariablesOutlinerAssetLabel::GetForegroundColor)
+				.OnTextCommitted(this, &SVariablesOutlinerAssetLabel::OnTextCommited)
+				.OnVerifyTextChanged(this, &SVariablesOutlinerAssetLabel::OnVerifyTextChanged)
+				.IsReadOnly(this, &SVariablesOutlinerAssetLabel::CanRename)
+			]
+			+SHorizontalBox::Slot()
+			.VAlign(VAlign_Center)
+			.Padding(2.0f, 0.0f, 2.0f, 3.0f)
+			.AutoWidth()
+			[
+				SNew(SImage)
+				.Visibility(this, &SVariablesOutlinerAssetLabel::GetDirtyImageVisibility)
+				.ToolTipText(this, &SVariablesOutlinerAssetLabel::GetDirtyTooltipText)
+				.Image(this, &SVariablesOutlinerAssetLabel::GetDirtyImageBrush)
+			]
+			+SHorizontalBox::Slot()
+			.VAlign(VAlign_Center)
+			.Padding(2.0f, 0.0f)
+			.AutoWidth()
+			[
+				SNew(SBox)
+				.Visibility(this, &SVariablesOutlinerAssetLabel::GetAddVariableButtonVisibility)
+				[
+					SNew(SComboButton)
+					.ContentPadding(FMargin(0, 0))
+					.ButtonStyle(FAppStyle::Get(), "SimpleButton")
+					.HasDownArrow(false)
+					.ButtonContent()
+					[
+						SNew(SImage)
+						.Image(FAppStyle::Get().GetBrush("Icons.PlusCircle"))
+						.ColorAndOpacity(FSlateColor::UseForeground())
+					]
+					.OnGetMenuContent_Lambda([WeakOutliner = SceneOutliner.AsShared().ToWeakPtr(), WeakItem = InTreeItem.AsShared().ToWeakPtr()]() -> TSharedRef<SWidget>
+					{
+						TSharedPtr<SVariablesOutliner> SharedOutliner = StaticCastSharedPtr<SVariablesOutliner>(WeakOutliner.Pin());
+						ensure(SharedOutliner.IsValid());
+
+						if (const FVariablesOutlinerMode* VariablesMode = static_cast<const FVariablesOutlinerMode*>(SharedOutliner->GetMode()))
+						{
+							FToolMenuContext Context;
+							UAnimNextVariableItemMenuContext* MenuContext = NewObject<UAnimNextVariableItemMenuContext>();
+							MenuContext->WeakOutliner = SharedOutliner;	
+
+							TSharedPtr<FVariablesOutlinerAssetItem> SharedAssetItem = StaticCastSharedPtr<FVariablesOutlinerAssetItem>(WeakItem.Pin());
+							if (SharedAssetItem.IsValid())
+							{
+								if (UAnimNextRigVMAsset* Asset = SharedAssetItem->SoftAsset.Get())
+								{
+									if (UAnimNextRigVMAssetEditorData* EditorData = UncookedOnly::FUtils::GetEditorData(Asset))
+									{
+										MenuContext->WeakEditorDatas.Add(EditorData);
+									}
+								}
+							}
+
+							Context.AddObject(MenuContext);
+							static const FName MenuName("VariablesOutliner.AddVariablesMenu");
+							return UToolMenus::Get()->GenerateWidget(MenuName, Context);
+						}
+
+						return SNullWidget::NullWidget;
+					})
+				]
+			]
+			+SHorizontalBox::Slot()
+			.VAlign(VAlign_Center)
+			.Padding(0.0f, 0.0f)
+			.AutoWidth()
+			[
+				SAssignNew(ImplementedSharedVariablesHBox, SHorizontalBox)
+				.Clipping(EWidgetClipping::OnDemand)
+			]
+			
+			+SHorizontalBox::Slot()
+			.VAlign(VAlign_Center)
+			.Padding(2.0f, 0.0f)
+			.AutoWidth()
+			[
+				SNew(SThrobber)
+				.Visibility(this, &SVariablesOutlinerAssetLabel::GetLoadingIndicatorVisibility)
+				.ToolTipText(LOCTEXT("LoadingTooltip", "Asset is loading..."))
+			]
+		];
+		
+		if (const TSharedPtr<FVariablesOutlinerAssetItem> Item = TreeItem.Pin())
+		{
+			for (const FSoftObjectPath& SharedVariablesPath : Item->SharedVariableSourcePaths)
+			{
+				const FString SharedVariablesAssetName = SharedVariablesPath.GetAssetName();
+				const FLinearColor BackgroundColor = FLinearColor::MakeRandomSeededColor(GetTypeHash(SharedVariablesAssetName)).Desaturate(0.4f);
+
+				const FButtonStyle& CloseButtonStyle = FUAFStyle::Get().GetWidgetStyle<FButtonStyle>("VariablesOutliner.SharedVariablesPill.CloseButton");
+				TSharedPtr<SBorder> SharedVariablesBorderWidget;
+				ImplementedSharedVariablesHBox->AddSlot()
+				.VAlign(VAlign_Center)
+				.Padding(2.f, 0.f)
+				.AutoWidth()
+				[
+					SAssignNew(SharedVariablesBorderWidget, SBorder)
+					.BorderImage(FUAFStyle::Get().GetBrush("VariablesOutliner.SharedVariablesPill.Border"))
+					.BorderBackgroundColor(BackgroundColor)
+					.ToolTipText( FText::FromString(SharedVariablesAssetName))
+					.VAlign(VAlign_Center)
+					.Padding(FMargin(6.0f, 1.0f, 4.0f, 1.0f))
+					[
+						SNew(SHorizontalBox)
+						+SHorizontalBox::Slot()
+						.AutoWidth()
+						.HAlign(HAlign_Center)
+						.VAlign(VAlign_Center)
+						[
+							SNew(STextBlock)
+							.TextStyle(FUAFStyle::Get(), "VariablesOutliner.SharedVariablesPill.Text")
+							.Text(FText::FromString(SharedVariablesAssetName))
+							.HighlightText(SceneOutliner.GetFilterHighlightText())
+							.OverflowPolicy(ETextOverflowPolicy::MiddleEllipsis)
+						]
+						+SHorizontalBox::Slot()
+						.HAlign(HAlign_Right)
+						.VAlign(VAlign_Center)
+						.AutoWidth()
+						[
+							SNew(SButton)
+							.ButtonStyle(&CloseButtonStyle)
+							.OnClicked( Item.ToSharedRef(), &FVariablesOutlinerAssetItem::OnRemoveSharedVariable, SharedVariablesPath )
+							.ContentPadding(FMargin(0.0))
+							.ToolTipText(LOCTEXT("RemoveReference_ToolTip", "Removes the SharedVariables reference from this Asset."))
+							[
+								SNew(SSpacer)
+								.Size(CloseButtonStyle.Normal.ImageSize)
+							]
+						]
+					]
+					.OnMouseButtonUp(Item.ToSharedRef(), &FVariablesOutlinerAssetItem::OnSharedVariableWidgetMouseUp, SharedVariablesPath)
+				];
+
+				SharedVariablesBorderWidget->SetOnMouseMove(FPointerEventHandler::CreateLambda([SharedVariablesPath, WeakOutliner = WeakSceneOutliner](const FGeometry&, const FPointerEvent& Event) -> FReply
+				{
+					const bool bShouldHighlight = Event.GetModifierKeys().IsControlDown();
+					{
+						if (TSharedPtr<ISceneOutliner> SharedOutliner = WeakOutliner.Pin())
+						{
+							if (const FVariablesOutlinerMode* Mode = static_cast<const FVariablesOutlinerMode*>(SharedOutliner->GetMode()))
+							{
+								if (bShouldHighlight)
+								{
+									Mode->SetHighlightedItem(GetTypeHash(SharedVariablesPath));
+								}
+								else
+								{
+									Mode->ClearHighlightedItem(GetTypeHash(SharedVariablesPath));
+								}
+							}
+						}
+					}
+
+					return FReply::Handled();
+				}));
+
+				SharedVariablesBorderWidget->SetOnMouseLeave(FSimpleNoReplyPointerEventHandler::CreateLambda([SharedVariablesPath, WeakOutliner = WeakSceneOutliner](const FPointerEvent&)
+				{
+					if (TSharedPtr<ISceneOutliner> SharedOutliner = WeakOutliner.Pin())
+					{
+						if (const FVariablesOutlinerMode* Mode = static_cast<const FVariablesOutlinerMode*>(SharedOutliner->GetMode()))
+						{
+							Mode->ClearHighlightedItem(GetTypeHash(SharedVariablesPath));					
+						}
+					}
+				}));
+			}
+		}
+	}
+
+	bool CanRename() const
+	{
+		if (const TSharedPtr<FVariablesOutlinerAssetItem> Item = TreeItem.Pin())
+		{
+			return Item->SoftAsset.Get() != nullptr;
+		}
+		return false;
+	}
+
+	FText GetDirtyTooltipText() const
+	{
+		if (const TSharedPtr<FVariablesOutlinerAssetItem> Item = TreeItem.Pin())
+		{
+			FTextBuilder TextBuilder;
+			TextBuilder.AppendLine(LOCTEXT("ModifiedTooltip", "Modified"));
+
+			if(UAnimNextRigVMAsset* Asset = Item->SoftAsset.Get())
+			{
+				const UPackage* Package = Asset->GetPackage();
+				check(Package);
+				if(Package->IsDirty())
+				{
+					TextBuilder.AppendLine(FText::FromName(Package->GetFName()));
+				}
+			}
+
+			return TextBuilder.ToText();
+		}
+		return FText::GetEmpty();
+	}
+
+	const FSlateBrush* GetDirtyImageBrush() const
+	{
+		if (const TSharedPtr<FVariablesOutlinerAssetItem> Item = TreeItem.Pin())
+		{
+			if(UAnimNextRigVMAsset* Asset = Item->SoftAsset.Get())
+			{
+				bool bIsDirty = false;
+				const UPackage* Package = Asset->GetPackage();
+				check(Package);
+				if(Package->IsDirty())
+				{
+					bIsDirty = true;
+				}
+
+				return bIsDirty ? FAppStyle::GetBrush("Icons.DirtyBadge") : nullptr;
+			}
+		}
+		return nullptr;
+	}
+
+	EVisibility GetDirtyImageVisibility() const
+	{
+		if (const TSharedPtr<FVariablesOutlinerAssetItem> Item = TreeItem.Pin())
+		{
+			if(UAnimNextRigVMAsset* Asset = Item->SoftAsset.Get())
+			{
+				bool bIsDirty = false;
+				const UPackage* Package = Asset->GetPackage();
+				check(Package);
+				if(Package->IsDirty())
+				{
+					bIsDirty = true;
+				}
+
+				return bIsDirty ? EVisibility::Visible : EVisibility::Collapsed;
+			}
+		}
+		return EVisibility::Collapsed;
+	}
+
+	EVisibility GetLoadingIndicatorVisibility() const
+	{
+		if (const TSharedPtr<FVariablesOutlinerAssetItem> Item = TreeItem.Pin())
+		{
+			UAnimNextRigVMAsset* Asset = Item->SoftAsset.Get();
+			return Asset != nullptr ? EVisibility::Collapsed : EVisibility::Visible;
+		}
+		return EVisibility::Collapsed;
+	}
+	
+	FText GetDisplayText() const
+	{
+		if (const TSharedPtr<FVariablesOutlinerAssetItem> Item = TreeItem.Pin())
+		{
+			return FText::FromString(Item->GetDisplayString());
+		}
+		return FText();
+	}
+
+	const FSlateBrush* GetAssetIcon() const
+	{
+		if (const TSharedPtr<FVariablesOutlinerAssetItem> Item = TreeItem.Pin())
+		{
+			UAnimNextRigVMAsset* Asset = Item->SoftAsset.Get();
+			return Asset != nullptr ? FSlateIconFinder::FindIconBrushForClass(Asset->GetClass()) : FAppStyle::Get().GetBrush("ClassIcon.Object");
+		}
+		return FAppStyle::Get().GetBrush("ClassIcon.Object");
+	}
+
+	void OnTextCommited(const FText& InLabel, ETextCommit::Type InCommitInfo) const
+	{
+		if(InCommitInfo == ETextCommit::OnEnter)
+		{
+			if (const TSharedPtr<FVariablesOutlinerAssetItem> Item = TreeItem.Pin())
+			{
+				Item->Rename(InLabel); 
+			}
+		}
+	}
+
+	bool OnVerifyTextChanged(const FText& InLabel, FText& OutErrorMessage) const
+	{
+		if (const TSharedPtr<FVariablesOutlinerAssetItem> Item = TreeItem.Pin())
+		{
+			return Item->ValidateName(InLabel, OutErrorMessage); 
+		}
+		return false;
+	}
+
+	virtual FSlateColor GetForegroundColor() const override
+	{
+		const TOptional<FLinearColor> BaseColor = FSceneOutlinerCommonLabelData::GetForegroundColor(*TreeItem.Pin());
+		return BaseColor.IsSet() ? BaseColor.GetValue() : FSlateColor::UseForeground();
+	}
+	
+	EVisibility GetAddVariableButtonVisibility() const
+	{
+		if (const TSharedPtr<FVariablesOutlinerAssetItem> Item = TreeItem.Pin())
+		{
+			if(Item->SoftAsset.Get())
+			{
+				return EVisibility::Visible;
+			}
+		}
+		return EVisibility::Collapsed;
+	}
+
+	virtual FReply OnDrop(const FGeometry& MyGeometry, const FDragDropEvent& DragDropEvent ) override
+	{
+		TSharedPtr<FVariableDragDropOp> GraphDropOp = DragDropEvent.GetOperationAs<FVariableDragDropOp>();
+		if (GraphDropOp.IsValid())
+		{
+			if (UAnimNextRigVMAsset* Asset = TreeItem.Pin()->SoftAsset.Get())
+			{
+				TSharedPtr<FVariablesOutlinerEntryItem> EntryItem = GraphDropOp->WeakItem.Pin();
+				if (EntryItem.IsValid())
+				{
+					if (UAnimNextVariableEntry* VariableEntry = EntryItem->WeakEntry.Get())
+					{
+						const UAnimNextRigVMAsset* VariableOuter = VariableEntry->GetTypedOuter<UAnimNextRigVMAsset>();
+						if (VariableOuter != Asset)
+						{
+							const FText VariableNameText = FText::FromName(VariableEntry->GetVariableName());
+							const FText AssetNameText = FText::FromName(Asset->GetFName());
+							
+							const FText FormattedMessage = FText::Format(LOCTEXT("MoveVariableToAssetFormatTransaction", "Moving {0} to {1}"), VariableNameText, AssetNameText);
+							FScopedTransaction Transaction(FormattedMessage);
+							UncookedOnly::FUtils::MoveVariableToAsset(VariableEntry, Asset, true, true);
+
+							return FReply::Handled();
+						}
+					}
+				}
+			}
+		}
+
+		TSharedPtr<FWorkspaceDragDropOp> WorkspaceDragDrop  = DragDropEvent.GetOperationAs<FWorkspaceDragDropOp>();
+		if (WorkspaceDragDrop.IsValid())
+		{
+			if (UAnimNextRigVMAsset* Asset = TreeItem.Pin()->SoftAsset.Get())
+			{
+				const FText AssetNameText = FText::FromName(Asset->GetFName());
+
+				UAnimNextRigVMAssetEditorData* EditorData = UncookedOnly::FUtils::GetEditorData<UAnimNextRigVMAssetEditorData>(Asset);
+				TArray<const UAnimNextSharedVariables*> SharedVariableAssetsToAdd;				
+				for (const FAssetData& AssetData : WorkspaceDragDrop->AssetDatas)
+				{					
+					const UClass* AssetClass = AssetData.GetClass(EResolveClass::Yes);
+					if (AssetClass && AssetClass->IsChildOf<UAnimNextSharedVariables>())
+					{						
+						if (!EditorData->FindEntry(AssetData.AssetName))
+						{
+							if (const UAnimNextSharedVariables* SharedVariables = Cast<UAnimNextSharedVariables>(AssetData.GetAsset()))
+							{
+								SharedVariableAssetsToAdd.Add(SharedVariables);
+							}
+						}
+					}
+				}
+
+				if (SharedVariableAssetsToAdd.Num())
+				{
+					const FText FormattedMessage = FText::Format(LOCTEXT("AddSharedVariablesToAssetTransaction", "Adding Shared Variables to {0}"), AssetNameText);
+					FScopedTransaction Transaction(FormattedMessage);
+					for (const UAnimNextSharedVariables* SharedVariables : SharedVariableAssetsToAdd)
+					{
+						EditorData->AddSharedVariables(SharedVariables);
+					}
+
+					return FReply::Handled();
+				}
+			}
+		}
+		
+		return FReply::Unhandled();
+	}
+
+	virtual void OnDragEnter( const FGeometry& MyGeometry, const FDragDropEvent& DragDropEvent ) override
+	{
+		TSharedPtr<FVariableDragDropOp> GraphDropOp = DragDropEvent.GetOperationAs<FVariableDragDropOp>();
+		if (GraphDropOp.IsValid())
+		{			
+			if (const UAnimNextRigVMAsset* Asset = TreeItem.Pin()->SoftAsset.Get())
+			{
+				TSharedPtr<FVariablesOutlinerEntryItem> EntryItem = GraphDropOp->WeakItem.Pin();
+				if (EntryItem.IsValid())
+				{
+					if (const UAnimNextVariableEntry* VariableEntry = EntryItem->WeakEntry.Get())
+					{
+						const FText VariableNameText = FText::FromName(VariableEntry->GetVariableName());
+						const FText AssetNameText = FText::FromName(Asset->GetFName());
+						
+						const UAnimNextRigVMAsset* VariableOuter = VariableEntry->GetTypedOuter<UAnimNextRigVMAsset>();
+						if (VariableOuter != Asset)
+						{
+							// [TODO] should we only allow move variables "up" the asset chain? If so validate that according to the workspace export chain, that VariableOuter is a child of Asset
+							const FText FormattedMessage = FText::Format(LOCTEXT("MoveVariableToAssetFormat", "Move {0} to {1}"), VariableNameText, AssetNameText);
+							GraphDropOp->SetSimpleFeedbackMessage(nullptr, FSlateColor(GetForegroundColor()), FormattedMessage);				
+						}
+						else
+						{
+							const FSlateBrush* ErrorSymbol = FAppStyle::GetBrush(TEXT("Graph.ConnectorFeedback.Error"));
+							const FText FormattedMessage = FText::Format(LOCTEXT("VariableAlreadyPartOfAssetFormat", "{0} is already part of {1}"), VariableNameText, AssetNameText);
+							GraphDropOp->SetSimpleFeedbackMessage(ErrorSymbol, FSlateColor(GetForegroundColor()), FormattedMessage);		
+						}
+					}
+				}
+			}
+		}
+
+		TSharedPtr<FWorkspaceDragDropOp> WorkspaceDragDrop  = DragDropEvent.GetOperationAs<FWorkspaceDragDropOp>();
+		if (WorkspaceDragDrop.IsValid())
+		{
+			if (const UAnimNextRigVMAsset* Asset = TreeItem.Pin()->SoftAsset.Get())
+			{
+				const FText AssetNameText = FText::FromName(Asset->GetFName());
+
+				TArray<FName> SharedVariableAssetNamesToAdd;
+				TArray<FName> PreexistingSharedVariableAssetNames;
+				
+				for (const FAssetData& AssetData : WorkspaceDragDrop->AssetDatas)
+				{					
+					const UClass* AssetClass = AssetData.GetClass(EResolveClass::Yes);
+					if (AssetClass && AssetClass->IsChildOf<UAnimNextSharedVariables>())
+					{						
+						const UAnimNextRigVMAssetEditorData* EditorData = UncookedOnly::FUtils::GetEditorData<const UAnimNextRigVMAssetEditorData>(Asset);
+						if (!EditorData->FindEntry(AssetData.AssetName))
+						{
+							SharedVariableAssetNamesToAdd.Add(AssetData.AssetName);
+						}
+						else
+						{
+
+							PreexistingSharedVariableAssetNames.Add(AssetData.AssetName);
+						}
+					}
+				}
+
+				const bool bAddingSharedVariables = SharedVariableAssetNamesToAdd.Num() > 0;
+				const bool bPreexistingSharedVariables = PreexistingSharedVariableAssetNames.Num() > 0;
+
+				if (bPreexistingSharedVariables && !bAddingSharedVariables)
+				{
+					const bool bMultipleAssets = PreexistingSharedVariableAssetNames.Num() > 1;
+					
+					const FSlateBrush* ErrorSymbol = FAppStyle::GetBrush(TEXT("Graph.ConnectorFeedback.Error"));
+					const FText SingleFormattedMessage = FText::Format(LOCTEXT("ExistingSharedVariablesInAsset", "{0} already references {1}"), AssetNameText, FText::FromName(PreexistingSharedVariableAssetNames[0]));
+					const FText MultiFormattedMessage = FText::Format(LOCTEXT("ExistingSharedVariablesInAssets", "{0} already references assets"), AssetNameText);
+
+					WorkspaceDragDrop->SetToolTip(bMultipleAssets ? MultiFormattedMessage : SingleFormattedMessage, ErrorSymbol);
+				}
+				else if (bAddingSharedVariables)
+				{
+					if (SharedVariableAssetNamesToAdd.Num() > 1)
+					{
+						FString SharedVariablesNameListString;
+						for (const FName& AssetName : SharedVariableAssetNamesToAdd)
+						{
+							if (!SharedVariablesNameListString.IsEmpty())
+							{
+								SharedVariablesNameListString.Append(TEXT(", "));
+							}
+
+							SharedVariablesNameListString.Appendf(TEXT("%s"), *AssetName.ToString());
+						}
+						
+						const FText SharedVariablesNameListText = FText::FromString(SharedVariablesNameListString);
+						const FText MultiFormattedMessage = FText::Format(LOCTEXT("AddMultipleSharedVariablesToAsset", "Add {0} to {1}"), SharedVariablesNameListText, AssetNameText);
+
+						WorkspaceDragDrop->SetToolTip(MultiFormattedMessage, FSlateIconFinder::FindIconBrushForClass(UAnimNextSharedVariables::StaticClass()));
+					}
+					else
+					{
+						const FText SharedVariablesNameText = FText::FromName(SharedVariableAssetNamesToAdd[0]);
+						const FText SingleFormattedMessage = FText::Format(LOCTEXT("AddSharedVariablesToAsset", "Add {0} to {1}"), SharedVariablesNameText, AssetNameText);
+						WorkspaceDragDrop->SetToolTip(SingleFormattedMessage, FSlateIconFinder::FindIconBrushForClass(UAnimNextSharedVariables::StaticClass()));
+					}				
+				}
+			}
+		}
+	}
+
+	virtual void OnDragLeave( const FDragDropEvent& DragDropEvent ) override
+	{
+		TSharedPtr<FVariableDragDropOp> GraphDropOp = DragDropEvent.GetOperationAs<FVariableDragDropOp>();
+		if (GraphDropOp.IsValid())
+		{
+			GraphDropOp->HoverTargetChanged();
+		}
+	}
+
+	TWeakPtr<FVariablesOutlinerAssetItem> TreeItem;
+	TSharedPtr<SInlineEditableTextBlock> TextBlock;
+	TSharedPtr<SHorizontalBox> ImplementedSharedVariablesHBox;
+};
+
+FVariablesOutlinerAssetItem::FVariablesOutlinerAssetItem(const FItemData& InItemData)
+	: ISceneOutlinerTreeItem(FVariablesOutlinerAssetItem::Type)
+	, SoftAsset(InItemData.Asset)
+	, SortValue(InItemData.SortValue)
+	, SharedVariableSourcePaths(InItemData.ImplementedSharedVariablesPaths)
+{
+}
+
+bool FVariablesOutlinerAssetItem::IsValid() const
+{
+	return !SoftAsset.IsNull();
+}
+
+FSceneOutlinerTreeItemID FVariablesOutlinerAssetItem::GetID() const
+{
+	return GetTypeHash(SoftAsset.ToSoftObjectPath());
+}
+
+FString FVariablesOutlinerAssetItem::GetDisplayString() const
+{
+	return SoftAsset.GetAssetName();
+}
+
+TSharedRef<SWidget> FVariablesOutlinerAssetItem::GenerateLabelWidget(ISceneOutliner& Outliner, const STableRow<FSceneOutlinerTreeItemPtr>& InRow)
+{
+	SAssignNew(AssetLabel, SVariablesOutlinerAssetLabel, *this, Outliner, InRow);
+	RenameRequestEvent.BindSP(AssetLabel->TextBlock.Get(), &SInlineEditableTextBlock::EnterEditingMode);
+	return AssetLabel->AsShared();
+}
+
+FString FVariablesOutlinerAssetItem::GetPackageName() const
+{
+	return SoftAsset.GetLongPackageName();
+}
+
+void FVariablesOutlinerAssetItem::Rename(const FText& InNewName) const
+{
+	UAnimNextRigVMAsset* Asset = SoftAsset.Get();
+	if(Asset == nullptr)
+	{
+		return;
+	}
+
+	FAssetToolsModule& AssetToolsModule = FModuleManager::LoadModuleChecked<FAssetToolsModule>("AssetTools");
+	const FString CurrentAssetPath = FPackageName::GetLongPackagePath(Asset->GetPackage()->GetName());
+	TArray<FAssetRenameData> AssetsToRename = { FAssetRenameData(Asset, CurrentAssetPath, InNewName.ToString()) };
+	AssetToolsModule.Get().RenameAssets(AssetsToRename);
+}
+
+bool FVariablesOutlinerAssetItem::ValidateName(const FText& InNewName, FText& OutErrorMessage) const
+{
+	UAnimNextRigVMAsset* Asset = SoftAsset.Get();
+	if(Asset == nullptr)
+	{
+		OutErrorMessage = LOCTEXT("InvalidAssetError", "Asset is invalid");
+		return false;
+	}
+
+	const FString NewName = InNewName.ToString();
+	if (NewName.Len() >= NAME_SIZE || Asset->GetFName() != FName(*NewName)) // Deliberately ignore case here to allow case-only renames of existing assets
+	{
+		const FString PackageName = Asset->GetPackage()->GetPathName() / NewName;
+		const FString ObjectPath = FString::Printf(TEXT("%s.%s"), *PackageName, *NewName);
+
+		FText ValidationErrorMsg;
+		if (!AssetViewUtils::IsValidObjectPathForCreate(ObjectPath, ValidationErrorMsg))
+		{
+			OutErrorMessage = ValidationErrorMsg;
+			return false;
+		}
+	}
+
+	FAssetToolsModule& AssetToolsModule = FAssetToolsModule::GetModule();
+	TWeakPtr<IAssetTypeActions> WeakAssetTypeActions = AssetToolsModule.Get().GetAssetTypeActionsForClass(UAnimNextRigVMAsset::StaticClass());
+	if (TSharedPtr<IAssetTypeActions> AssetTypeActions = WeakAssetTypeActions.Pin())
+	{
+		if (!AssetTypeActions->CanRename(FAssetData(Asset), &OutErrorMessage))
+		{
+			return false;
+		}
+	}
+
+	return true;
+}
+
+FReply FVariablesOutlinerAssetItem::OnSharedVariableWidgetMouseUp(const FGeometry& Geometry, const FPointerEvent& PointerEvent, const FSoftObjectPath ClickedSharedVariablesPath) const
+{
+	if (PointerEvent.GetEffectingButton() == EKeys::RightMouseButton)
+	{
+		const FWidgetPath WidgetPath = (PointerEvent.GetEventPath() != nullptr) ? *PointerEvent.GetEventPath() : FWidgetPath();
+
+		const bool bCloseAfterSelection = true;
+		FMenuBuilder MenuBuilder(bCloseAfterSelection, nullptr);
+		MenuBuilder.BeginSection("", LOCTEXT("SharedVariablesContextMenu", "Shared Variables") );
+		{
+			MenuBuilder.AddMenuEntry(
+					LOCTEXT("RemoveReference", "Remove"),
+					LOCTEXT("RemoveReference_ToolTip", "Removes the SharedVariables reference from this Asset."),
+					FSlateIcon(),
+					FUIAction(FExecuteAction::CreateLambda([this, ClickedSharedVariablesPath]() { OnRemoveSharedVariable(ClickedSharedVariablesPath); }))
+			);
+		}
+		MenuBuilder.EndSection();
+
+		FSlateApplication::Get().PushMenu(AssetLabel.ToSharedRef(), WidgetPath, MenuBuilder.MakeWidget(), FSlateApplication::Get().GetCursorPos(), FPopupTransitionEffect(FPopupTransitionEffect::ContextMenu));
+		
+		return FReply::Handled();
+	}
+
+	return FReply::Unhandled();
+}
+
+FReply FVariablesOutlinerAssetItem::OnRemoveSharedVariable(const FSoftObjectPath ClickedSharedVariablesPath) const
+{
+	if (UAnimNextRigVMAsset* Asset = SoftAsset.Get())
+	{
+		const FText AssetNameText = FText::FromName(Asset->GetFName());
+
+		UAnimNextRigVMAssetEditorData* EditorData = UncookedOnly::FUtils::GetEditorData<UAnimNextRigVMAssetEditorData>(Asset);
+		if (UAnimNextSharedVariablesEntry* SharedVariableEntryToRemove = Cast<UAnimNextSharedVariablesEntry>(EditorData->FindEntry(ClickedSharedVariablesPath.GetAssetFName())))
+		{
+			const FText SharedVariablesAssetNameText = FText::FromName(ClickedSharedVariablesPath.GetAssetFName());
+			const FText FormattedMessage = FText::Format(LOCTEXT("RemovingSharedAssetEntryFormat", "Removing {0} from {1}"), SharedVariablesAssetNameText, AssetNameText);
+			FScopedTransaction Transaction(FormattedMessage);
+			
+			if (const UAnimNextSharedVariables* SharedVariablesToRemove = SharedVariableEntryToRemove->GetAsset())
+			{
+				// Replace all referenced variables from the to-be-removed SharedVariables from the referencing Asset
+				FAnimNextAssetRegistryExports VariableExports;
+				UncookedOnly::FUtils::GetExportedVariablesForAsset(SharedVariablesToRemove, VariableExports);
+
+				VariableExports.ForEachExportOfType<FAnimNextVariableDeclarationData>([EditorData, ClickedSharedVariablesPath](const FName& Identifier, const FAnimNextVariableDeclarationData& VariableDeclaration) -> bool
+				{
+					FAnimNextSoftVariableReference VariableReference(Identifier, ClickedSharedVariablesPath);
+					FAnimNextSoftVariableReference EmptyReference;
+					UncookedOnly::FUtils::ReplaceVariableReferences(EditorData, VariableReference, EmptyReference);
+				
+					return true;
+				});
+			}
+			else if (const UScriptStruct* Struct = SharedVariableEntryToRemove->GetStruct())
+			{
+				// Replace all referenced variables from the to-be-removed Native SharedVariables source from the referencing Asset
+				for (TFieldIterator<FProperty> It(Struct); It; ++It)
+				{
+					const FProperty* Property = *It;
+					if (Property->HasAnyPropertyFlags(CPF_NativeAccessSpecifierPublic))
+					{
+
+						FAnimNextSoftVariableReference VariableReference(Property->GetFName(), ClickedSharedVariablesPath);
+						FAnimNextSoftVariableReference EmptyReference;
+						UncookedOnly::FUtils::ReplaceVariableReferences(EditorData, VariableReference, EmptyReference);
+					}
+				}
+			}
+
+			// Remove the entry itself
+			EditorData->RemoveEntry(SharedVariableEntryToRemove);
+
+			return FReply::Handled();
+		}
+	}
+
+	return FReply::Unhandled();
+}
+}
+
+#undef LOCTEXT_NAMESPACE // "VariablesOutlinerTreeItem"
\ No newline at end of file
diff --git a/Engine/Plugins/Experimental/UAF/UAF/Source/UAFEditor/Private/Variables/Outliner/VariablesOutlinerAssetItem.h b/Engine/Plugins/Experimental/UAF/UAF/Source/UAFEditor/Private/Variables/Outliner/VariablesOutlinerAssetItem.h
new file mode 100644
--- /dev/null
+++ b/Engine/Plugins/Experimental/UAF/UAF/Source/UAFEditor/Private/Variables/Outliner/VariablesOutlinerAssetItem.h
@@ -0,0 +1,57 @@
+﻿// Copyright Epic Games, Inc. All Rights Reserved.
+
+#pragma once
+
+#include "ISceneOutlinerTreeItem.h"
+#include "UObject/SoftObjectPtr.h"
+
+class UAnimNextRigVMAsset;
+
+namespace UE::UAF::Editor
+{
+
+class IWorkspaceOutlinerItemDetails;
+
+struct FVariablesOutlinerAssetItem : ISceneOutlinerTreeItem
+{
+	static const FSceneOutlinerTreeItemType Type;
+	
+	struct FItemData
+	{
+		const TSoftObjectPtr<UAnimNextRigVMAsset>& Asset;
+		const int32 SortValue;
+		const TArrayView<const FSoftObjectPath> ImplementedSharedVariablesPaths;
+	};
+
+	FVariablesOutlinerAssetItem(const FItemData& InItemData);
+	virtual ~FVariablesOutlinerAssetItem() override = default;
+
+	// Begin ISceneOutlinerTreeItem overrides
+	virtual bool IsValid() const override;
+	virtual FSceneOutlinerTreeItemID GetID() const override;
+	virtual FString GetDisplayString() const override;
+	virtual bool CanInteract() const override { return true; }
+	virtual TSharedRef<SWidget> GenerateLabelWidget(ISceneOutliner& Outliner, const STableRow<FSceneOutlinerTreeItemPtr>& InRow) override;
+	virtual FString GetPackageName() const override;
+	// End ISceneOutlinerTreeItem overrides
+
+	// Renames the item to the specified name
+	void Rename(const FText& InNewName) const;
+
+	// Validates the new item name
+	bool ValidateName(const FText& InNewName, FText& OutErrorMessage) const;
+
+	FReply OnSharedVariableWidgetMouseUp(const FGeometry&, const FPointerEvent& PointerEvent, const FSoftObjectPath ClickedSharedVariablesPath) const;
+	FReply OnRemoveSharedVariable(const FSoftObjectPath ClickedSharedVariablesPath) const;
+
+	// Soft ptr to the underlying asset, which may not be loaded yet
+	TSoftObjectPtr<UAnimNextRigVMAsset> SoftAsset;
+
+	int32 SortValue;
+
+	TArray<const FSoftObjectPath> SharedVariableSourcePaths;
+
+	TSharedPtr<class SVariablesOutlinerAssetLabel> AssetLabel;
+};
+
+}
diff --git a/Engine/Plugins/Experimental/UAF/UAF/Source/UAFEditor/Private/Variables/Outliner/VariablesOutlinerCategoryItem.cpp b/Engine/Plugins/Experimental/UAF/UAF/Source/UAFEditor/Private/Variables/Outliner/VariablesOutlinerCategoryItem.cpp
--- a/Engine/Plugins/Experimental/UAF/UAF/Source/UAFEditor/Private/Variables/Outliner/VariablesOutlinerCategoryItem.cpp
+++ b/Engine/Plugins/Experimental/UAF/UAF/Source/UAFEditor/Private/Variables/Outliner/VariablesOutlinerCategoryItem.cpp
@@ -7,7 +7,7 @@
 #include "ISceneOutliner.h"
 #include "ScopedTransaction.h"
 #include "UncookedOnlyUtils.h"
-#include "Common/Outliner/OutlinerDragDropOps.h"
+#include "VariablesOutlinerDragDrop.h"
 #include "Entries/AnimNextRigVMAssetEntry.h"
 #include "Styling/StarshipCoreStyle.h"
 #include "Widgets/Text/SInlineEditableTextBlock.h"
@@ -17,54 +17,298 @@
 
 namespace UE::UAF::Editor
 {
+class FVariableDragDropOp;
 
-const FSceneOutlinerTreeItemType FVariablesOutlinerCategoryItem::Type(&FOutlinerCategoryItem::Type);
+const FSceneOutlinerTreeItemType FVariablesOutlinerCategoryItem::Type;
 
-FVariablesOutlinerCategoryItem::FVariablesOutlinerCategoryItem(const FOutlinerCategoryItem::FItemData& InItemData) : FOutlinerCategoryItem(FVariablesOutlinerCategoryItem::Type, InItemData)
+class SVariablesOutlinerCategoryLabel : FSceneOutlinerCommonLabelData, public SCompoundWidget
 {
-}
+	SLATE_BEGIN_ARGS(SVariablesOutlinerCategoryLabel) {}
+	SLATE_END_ARGS()
 
-void FVariablesOutlinerCategoryItem::Rename(const FText& InNewName) const
-{
-	if(UAnimNextRigVMAsset* Owner = WeakOwner.Get())
+	void Construct(const FArguments& InArgs, FVariablesOutlinerCategoryItem& InTreeItem, ISceneOutliner& SceneOutliner)
 	{
-		if (UAnimNextRigVMAssetEditorData* EditorData = UE::UAF::UncookedOnly::FUtils::GetEditorData(Owner))
+		WeakSceneOutliner = StaticCastSharedRef<ISceneOutliner>(SceneOutliner.AsShared());
+		TreeItem = StaticCastSharedRef<FVariablesOutlinerCategoryItem>(InTreeItem.AsShared());
+		
+		ChildSlot
+		[
+			SNew(SHorizontalBox)
+			+SHorizontalBox::Slot()
+			.AutoWidth()
+			.VAlign(VAlign_Center)
+			.Padding(2.0f, 2.0f)
+			[
+				SAssignNew(TextBlock, SInlineEditableTextBlock)
+				.Font(FStyleFonts::Get().NormalBold)
+				.Text(this, &SVariablesOutlinerCategoryLabel::GetDisplayText)
+				.HighlightText(SceneOutliner.GetFilterHighlightText())
+				.ColorAndOpacity(this, &SVariablesOutlinerCategoryLabel::GetForegroundColor)
+				.OnTextCommitted(this, &SVariablesOutlinerCategoryLabel::OnTextCommited)
+				.OnVerifyTextChanged(this, &SVariablesOutlinerCategoryLabel::OnVerifyTextChanged)
+			]
+		];
+	}
+	
+	FText GetDisplayText() const
+	{
+		if (const TSharedPtr<FVariablesOutlinerCategoryItem> Item = TreeItem.Pin())
 		{
-			FScopedTransaction Transaction(LOCTEXT("RenameCategory", "Rename category"));
+			return FText::FromString(Item->GetDisplayString());
+		}
+		return FText();
+	}
+
+	virtual FSlateColor GetForegroundColor() const override
+	{
+		const TOptional<FLinearColor> BaseColor = FSceneOutlinerCommonLabelData::GetForegroundColor(*TreeItem.Pin());
+		return BaseColor.IsSet() ? BaseColor.GetValue() : FSlateColor::UseForeground();
+	}
+
+	
+	void OnTextCommited(const FText& InLabel, ETextCommit::Type InCommitInfo) const	
+	{
+		if(InCommitInfo == ETextCommit::OnEnter)
+		{
+			if (const TSharedPtr<FVariablesOutlinerCategoryItem> Item = TreeItem.Pin())
+			{
+				Item->Rename(InLabel); 
+			}
+		}
+	}
+	
+	bool OnVerifyTextChanged(const FText& InLabel, FText& OutErrorMessage) const
+	{
+		if (const TSharedPtr<FVariablesOutlinerCategoryItem> Item = TreeItem.Pin())
+		{
+			return Item->ValidateName(InLabel, OutErrorMessage); 
+		}
+		return false;
+	}
+
+	virtual FReply OnDrop( const FGeometry& MyGeometry, const FDragDropEvent& DragDropEvent ) override
+	{
+		TSharedPtr<FVariableDragDropOp> GraphDropOp = DragDropEvent.GetOperationAs<FVariableDragDropOp>();
+		if (GraphDropOp.IsValid())
+		{
+			if (const TSharedPtr<FVariablesOutlinerCategoryItem> CategoryItem = TreeItem.Pin())
+			{
+				TSharedPtr<FVariablesOutlinerEntryItem> EntryItem = GraphDropOp->WeakItem.Pin();
+				if (EntryItem.IsValid())
+				{
+					if (UAnimNextVariableEntry* VariableEntry = EntryItem->WeakEntry.Get())
+					{
+						UAnimNextRigVMAsset* VariableAsset = VariableEntry->GetTypedOuter<UAnimNextRigVMAsset>();
+						UAnimNextRigVMAsset* CategoryAsset = CategoryItem->WeakOwner.Get();
+
+						const FText OtherVariableNameText = FText::FromName(VariableEntry->GetVariableName());
+						const FText ThisVariableCategoryText = FText::FromString(CategoryItem->GetDisplayString());
+						const FText FormattedMessage = FText::Format(LOCTEXT("MoveVariableToCategoryFormat", "Moving {0} to Category {1}"), OtherVariableNameText, ThisVariableCategoryText);	
+						FScopedTransaction Transaction(FormattedMessage);
+
+						UAnimNextVariableEntry* NewVariableEntry = VariableEntry;
+						
+						if (VariableAsset != CategoryAsset)
+						{
+							const FName VariableEntryName = VariableEntry->GetVariableName();
+							UncookedOnly::FUtils::MoveVariableToAsset(VariableEntry, CategoryAsset, true, true);
+							NewVariableEntry = Cast<UAnimNextVariableEntry>(UncookedOnly::FUtils::GetEditorData(CategoryAsset)->FindEntry(VariableEntryName));
+						}
+						
+						check(NewVariableEntry);
+						
+						// Only add if not already part of category
+						if (NewVariableEntry->GetVariableCategory() != CategoryItem->CategoryPath)
+						{
+							NewVariableEntry->SetVariableCategory(CategoryItem->CategoryPath);
+						}
+					}
+				}
+			}
 			
-			FString NewCategoryPath = ParentCategoryName;
-			if (!NewCategoryPath.IsEmpty())
+			return FReply::Handled();
+		}
+		
+		TSharedPtr<FCategoryDragDropOp> CategoryDropOp = DragDropEvent.GetOperationAs<FCategoryDragDropOp>();
+		if (CategoryDropOp.IsValid())
+		{
+			if (const TSharedPtr<FVariablesOutlinerCategoryItem> CategoryItem = TreeItem.Pin())
+			{
+				TSharedPtr<FVariablesOutlinerCategoryItem> DropCategoryItem = CategoryDropOp->WeakItem.Pin();
+				if (DropCategoryItem.IsValid())
+				{
+					const FString ThisCategoryName = CategoryItem->CategoryName;
+					const FString OtherCategoryName = DropCategoryItem->CategoryName;
+
+					UAnimNextRigVMAsset* ThisOwner = CategoryItem->WeakOwner.Get();
+					const UAnimNextRigVMAsset* OtherOwner = DropCategoryItem->WeakOwner.Get();
+
+					if (ThisOwner == OtherOwner)
+					{
+						const FText OtherCategoryText = FText::FromString(DropCategoryItem->GetDisplayString());
+						const FText ThisCategoryText = FText::FromString(CategoryItem->GetDisplayString());
+						
+						const FText FormattedMessage = FText::Format(LOCTEXT("ReorderCategory", "Reorder {0} before {1}"), OtherCategoryText, ThisCategoryText);
+						
+						FScopedTransaction Transaction(FormattedMessage);
+						UncookedOnly::FUtils::GetEditorData(ThisOwner)->ReorderCategory(OtherCategoryName, ThisCategoryName);
+					}
+				}
+			}
+
+			return FReply::Handled();
+		}
+		
+		
+		return FReply::Unhandled();
+	}
+
+	virtual void OnDragEnter( const FGeometry& MyGeometry, const FDragDropEvent& DragDropEvent ) override
+	{
+		TSharedPtr<FVariableDragDropOp> GraphDropOp = DragDropEvent.GetOperationAs<FVariableDragDropOp>();
+		if (GraphDropOp.IsValid())
+		{
+			if (const TSharedPtr<FVariablesOutlinerCategoryItem> CategoryItem = TreeItem.Pin())
 			{
-				NewCategoryPath.Append(TEXT("|"));
+				TSharedPtr<FVariablesOutlinerEntryItem> EntryItem = GraphDropOp->WeakItem.Pin();
+				if (EntryItem.IsValid())
+				{
+					if (UAnimNextVariableEntry* VariableEntry = EntryItem->WeakEntry.Get())
+					{
+						const FText OtherVariableNameText = FText::FromName(VariableEntry->GetVariableName());
+						const FText ThisVariableCategoryText = FText::FromString(CategoryItem->GetDisplayString());
+
+						const UAnimNextRigVMAssetEditorData* VariableEditorData = VariableEntry->GetTypedOuter<UAnimNextRigVMAssetEditorData>();
+						const UAnimNextRigVMAssetEditorData* CategoryEditorData = UncookedOnly::FUtils::GetEditorData(CategoryItem->WeakOwner.Get());
+
+						if (VariableEditorData == CategoryEditorData)
+						{
+							if (VariableEntry->GetVariableCategory() == CategoryItem->CategoryName)
+							{
+								const FText FormattedMessage = FText::Format(LOCTEXT("VariableAlreadyPartOfCategoryFormat", "{0} is already part of Category {1}"), OtherVariableNameText, ThisVariableCategoryText);
+								
+								const FSlateBrush* ErrorSymbol = FAppStyle::GetBrush(TEXT("Graph.ConnectorFeedback.Error"));
+								GraphDropOp->SetSimpleFeedbackMessage(ErrorSymbol, FSlateColor(GetForegroundColor()), FormattedMessage);
+							}
+							else if (!VariableEntry->GetVariableCategory().IsEmpty())
+							{
+								const FText FormattedMessage = FText::Format(LOCTEXT("MoveVariableToCategoryFormat", "Moving {0} to Category {1}"), OtherVariableNameText, ThisVariableCategoryText);							
+								GraphDropOp->SetSimpleFeedbackMessage(nullptr, FSlateColor(GetForegroundColor()), FormattedMessage);
+							}
+							else
+							{
+								const FText FormattedMessage = FText::Format(LOCTEXT("AddVariableToCategoryFormat", "Add {0} to Category {1}"), OtherVariableNameText, ThisVariableCategoryText);
+								GraphDropOp->SetSimpleFeedbackMessage(nullptr, FSlateColor(GetForegroundColor()), FormattedMessage);
+							}
+						}
+						else
+						{
+							const FText ThisAssetText =  FText::FromString(CategoryItem->WeakOwner.Get()->GetName());
+							const FText FormattedMessage = FText::Format(LOCTEXT("MoveVariableToCategoryAndAssetFormat", "Moving {0} Category {1} in {2}"), OtherVariableNameText, ThisVariableCategoryText, ThisAssetText);
+							GraphDropOp->SetSimpleFeedbackMessage(nullptr, FSlateColor(GetForegroundColor()), FormattedMessage);						
+						}
+					}
+				}
+			}
+
+			return;
+		}
+		
+		TSharedPtr<FCategoryDragDropOp> CategoryDropOp = DragDropEvent.GetOperationAs<FCategoryDragDropOp>();
+		if (CategoryDropOp.IsValid())
+		{
+			if (const TSharedPtr<FVariablesOutlinerCategoryItem> CategoryItem = TreeItem.Pin())
+			{
+				TSharedPtr<FVariablesOutlinerCategoryItem> DropCategoryItem = CategoryDropOp->WeakItem.Pin();
+				if (DropCategoryItem.IsValid())
+				{
+					const UAnimNextRigVMAsset* ThisOwner = CategoryItem->WeakOwner.Get();
+					const UAnimNextRigVMAsset* OtherOwner = DropCategoryItem->WeakOwner.Get();
+
+					if (ThisOwner && ThisOwner == OtherOwner)
+					{
+						const FText OtherCategoryText = FText::FromString(DropCategoryItem->GetDisplayString());
+						const FText ThisCategoryText = FText::FromString(CategoryItem->GetDisplayString());
+						
+						const FText FormattedMessage = FText::Format(LOCTEXT("ReorderCategory", "Reorder {0} before {1}"), OtherCategoryText, ThisCategoryText);
+						CategoryDropOp->SetToolTip(FormattedMessage, nullptr);
+					}
+				}
 			}
-			NewCategoryPath.Append(InNewName.ToString());
+		}
+	}
 
-			EditorData->RenameCategory(CategoryPath, NewCategoryPath);
+	virtual void OnDragLeave( const FDragDropEvent& DragDropEvent ) override
+	{
+		TSharedPtr<FVariableDragDropOp> GraphDropOp = DragDropEvent.GetOperationAs<FVariableDragDropOp>();
+		if (GraphDropOp.IsValid())
+		{
+			GraphDropOp->HoverTargetChanged();
 		}
 	}
+
+	TWeakPtr<FVariablesOutlinerCategoryItem> TreeItem;
+	TSharedPtr<SInlineEditableTextBlock> TextBlock;
+};
+
+FVariablesOutlinerCategoryItem::FVariablesOutlinerCategoryItem(const FItemData& InItemData) : ISceneOutlinerTreeItem(FVariablesOutlinerCategoryItem::Type), WeakOwner(InItemData.InOwner), CategoryName(InItemData.InCategoryName), ParentCategoryName(InItemData.InParentCategoryName), CategoryPath(InItemData.InCategoryPath)
+{
 }
 
-bool FVariablesOutlinerCategoryItem::ValidateName(const FText& InNewName, FText& OutErrorMessage) const
+bool FVariablesOutlinerCategoryItem::IsValid() const
 {
-	if (InNewName.IsEmpty())
-	{		
-		OutErrorMessage = LOCTEXT("NameEmpty", "Category name cannot be empty");
-		return false;
+	return WeakOwner.Get() != nullptr;
+}
+
+FSceneOutlinerTreeItemID FVariablesOutlinerCategoryItem::GetID() const
+{
+	UAnimNextRigVMAsset* Owner = WeakOwner.Get();
+	const FSoftObjectPath OwnerPath = Owner;
+	return HashCombine(GetTypeHash(OwnerPath), GetTypeHash(CategoryPath));	
+}
+
+FString FVariablesOutlinerCategoryItem::GetDisplayString() const
+{
+	return CategoryName;
+}
+
+TSharedRef<SWidget> FVariablesOutlinerCategoryItem::GenerateLabelWidget(ISceneOutliner& Outliner, const STableRow<FSceneOutlinerTreeItemPtr>& InRow)
+{
+	TSharedRef<SVariablesOutlinerCategoryLabel> LabelWidget = SNew(SVariablesOutlinerCategoryLabel, *this, Outliner);
+	RenameRequestEvent.BindSP(LabelWidget->TextBlock.Get(), &SInlineEditableTextBlock::EnterEditingMode);
+
+	return LabelWidget;
+}
+
+FString FVariablesOutlinerCategoryItem::GetPackageName() const
+{
+	if (UAnimNextRigVMAsset* Owner = WeakOwner.Get())
+	{
+		return Owner->GetPackage()->GetName();
 	}
+	
+	return ISceneOutlinerTreeItem::GetPackageName();
+}
 
+void FVariablesOutlinerCategoryItem::Rename(const FText& InNewName) const
+{
 	if(UAnimNextRigVMAsset* Owner = WeakOwner.Get())
 	{
 		if (UAnimNextRigVMAssetEditorData* EditorData = UE::UAF::UncookedOnly::FUtils::GetEditorData(Owner))
 		{
-			FString NewCategoryPath = ParentCategoryName;
-			if (!NewCategoryPath.IsEmpty())
-			{
-				NewCategoryPath.Append(TEXT("|"));
-			}
-			NewCategoryPath.Append(InNewName.ToString());
-		
-			if (EditorData->VariableAndFunctionCategories.Contains(NewCategoryPath) && CategoryPath != NewCategoryPath
-			)
+			FScopedTransaction Transaction(LOCTEXT("RenameCategory", "Rename category"));
+			EditorData->RenameCategory(CategoryName, InNewName.ToString());						
+		}
+	}	
+}
+
+bool FVariablesOutlinerCategoryItem::ValidateName(const FText& InNewName, FText& OutErrorMessage) const
+{
+	if(UAnimNextRigVMAsset* Owner = WeakOwner.Get())
+	{
+		if (UAnimNextRigVMAssetEditorData* EditorData = UE::UAF::UncookedOnly::FUtils::GetEditorData(Owner))
+		{
+			if (EditorData->VariableAndFunctionCategories.Contains(InNewName.ToString()))
 			{
 				OutErrorMessage = LOCTEXT("NameExistsError", "Category name already exists in this asset");
 				return false;
diff --git a/Engine/Plugins/Experimental/UAF/UAF/Source/UAFEditor/Private/Variables/Outliner/VariablesOutlinerCategoryItem.h b/Engine/Plugins/Experimental/UAF/UAF/Source/UAFEditor/Private/Variables/Outliner/VariablesOutlinerCategoryItem.h
--- a/Engine/Plugins/Experimental/UAF/UAF/Source/UAFEditor/Private/Variables/Outliner/VariablesOutlinerCategoryItem.h
+++ b/Engine/Plugins/Experimental/UAF/UAF/Source/UAFEditor/Private/Variables/Outliner/VariablesOutlinerCategoryItem.h
@@ -3,21 +3,46 @@
 #pragma once
 
 #include "ISceneOutlinerTreeItem.h"
-#include "Common/Outliner/OutlinerCategoryItem.h"
 
 class UAnimNextRigVMAsset;
 
 namespace UE::UAF::Editor
 {
 
-struct FVariablesOutlinerCategoryItem : FOutlinerCategoryItem
+struct FVariablesOutlinerCategoryItem : ISceneOutlinerTreeItem
 {
 	static const FSceneOutlinerTreeItemType Type;
 
-	FVariablesOutlinerCategoryItem(const FOutlinerCategoryItem::FItemData& IntemData);
-
-	virtual void Rename(const FText& InNewName) const override final;
-	virtual bool ValidateName(const FText& InNewName, FText& OutErrorMessage) const override final;
+	struct FItemData
+	{
+		UAnimNextRigVMAsset* InOwner;
+		FStringView InCategoryName;
+		FStringView InParentCategoryName;
+		FStringView InCategoryPath;
+	};
+
+	FVariablesOutlinerCategoryItem(const FItemData& InItemData);
+
+	// Begin ISceneOutlinerTreeItem overrides
+	virtual bool IsValid() const override;
+	virtual FSceneOutlinerTreeItemID GetID() const override;
+	virtual FString GetDisplayString() const override;
+	virtual bool CanInteract() const override { return true; }
+	virtual TSharedRef<SWidget> GenerateLabelWidget(ISceneOutliner& Outliner, const STableRow<FSceneOutlinerTreeItemPtr>& InRow) override;
+	virtual FString GetPackageName() const override;
+	// End ISceneOutlinerTreeItem overrides
+
+	// Renames the item to the specified name
+	void Rename(const FText& InNewName) const;
+
+	// Validates the new item name
+	bool ValidateName(const FText& InNewName, FText& OutErrorMessage) const;
+	
+	TWeakObjectPtr<UAnimNextRigVMAsset> WeakOwner;
+	FString CategoryName;
+	FString ParentCategoryName;
+	FString CategoryPath;
+	int32 SortValue = 0;
 };
 
 }
diff --git a/Engine/Plugins/Experimental/UAF/UAF/Source/UAFEditor/Private/Variables/Outliner/VariablesOutlinerColumns.cpp b/Engine/Plugins/Experimental/UAF/UAF/Source/UAFEditor/Private/Variables/Outliner/VariablesOutlinerColumns.cpp
--- a/Engine/Plugins/Experimental/UAF/UAF/Source/UAFEditor/Private/Variables/Outliner/VariablesOutlinerColumns.cpp
+++ b/Engine/Plugins/Experimental/UAF/UAF/Source/UAFEditor/Private/Variables/Outliner/VariablesOutlinerColumns.cpp
@@ -13,6 +13,7 @@
 #include "ScopedTransaction.h"
 #include "SPinTypeSelector.h"
 #include "UncookedOnlyUtils.h"
+#include "VariablesOutlinerAssetItem.h"
 #include "VariablesOutlinerEntryItem.h"
 #include "Entries/AnimNextSharedVariablesEntry.h"
 #include "Entries/AnimNextRigVMAssetEntry.h"
@@ -24,7 +25,6 @@
 #include "Widgets/Input/SButton.h"
 #include "Widgets/Input/SCheckBox.h"
 #include "Variables/Outliner/VariablesOutlinerMode.h"
-#include "Common/Outliner/OutlinerAssetItem.h"
 
 #define LOCTEXT_NAMESPACE "VariablesOutlinerColumns"
 
@@ -123,7 +123,7 @@ SHeaderRow::FColumn::FArguments FVariablesOutlinerTypeColumn::ConstructHeaderRow
 
 const TSharedRef<SWidget> FVariablesOutlinerTypeColumn::ConstructRowWidget(FSceneOutlinerTreeItemRef Item, const STableRow<FSceneOutlinerTreeItemPtr>& Row)
 {
-	const FOutlinerAssetItem* AssetItem = Item->CastTo<FOutlinerAssetItem>();
+	const FVariablesOutlinerAssetItem* AssetItem = Item->CastTo<FVariablesOutlinerAssetItem>();
 	if (AssetItem)
 	{
 		return SNullWidget::NullWidget;
@@ -417,7 +417,29 @@ class SVariablesOutlinerValue : public SCompoundWidget, public FNotifyHook
 	}
 
 	// FNotifyHook interface
-	virtual void NotifyPostChange(const FPropertyChangedEvent& PropertyChangedEvent, FProperty* PropertyThatChanged) override;
+	virtual void NotifyPostChange(const FPropertyChangedEvent& PropertyChangedEvent, FProperty* PropertyThatChanged) override
+	{
+		TSharedPtr<FVariablesOutlinerEntryItem> TreeItem = WeakTreeItem.Pin();
+		if(!TreeItem.IsValid())
+		{
+			return;
+		}
+
+		UAnimNextVariableEntry* VariableEntry = Cast<UAnimNextVariableEntry>(TreeItem->WeakEntry.Get());
+		if(VariableEntry == nullptr)
+		{
+			return;
+		}
+
+		UAnimNextRigVMAssetEditorData* EditorData = VariableEntry->GetTypedOuter<UAnimNextRigVMAssetEditorData>();
+		if (EditorData == nullptr)
+		{
+			return;
+		}
+
+		VariableEntry->MarkPackageDirty();
+		VariableEntry->BroadcastModified(EAnimNextEditorDataNotifType::VariableDefaultValueChanged);
+	}
 
 	void HandleModified(UAnimNextRigVMAssetEditorData* InEditorData, EAnimNextEditorDataNotifType InType, UObject* InSubject)
 	{
@@ -449,51 +471,185 @@ class SVariablesOutlinerValue : public SCompoundWidget, public FNotifyHook
 	FInstancedPropertyBag InternalPropertyBag;
 };
 
-void SVariablesOutlinerValue::NotifyPostChange(const FPropertyChangedEvent& PropertyChangedEvent, FProperty* PropertyThatChanged)
+const TSharedRef<SWidget> FVariablesOutlinerValueColumn::ConstructRowWidget(FSceneOutlinerTreeItemRef Item, const STableRow<FSceneOutlinerTreeItemPtr>& Row)
 {
-	TSharedPtr<FVariablesOutlinerEntryItem> TreeItem = WeakTreeItem.Pin();
-	if(!TreeItem.IsValid())
+	const FVariablesOutlinerAssetItem* AssetItem = Item->CastTo<FVariablesOutlinerAssetItem>();
+	if (AssetItem)
 	{
-		return;
+		return SNullWidget::NullWidget;
 	}
 
-	UAnimNextVariableEntry* VariableEntry = Cast<UAnimNextVariableEntry>(TreeItem->WeakEntry.Get());
-	if(VariableEntry == nullptr)
+	// Return a dummy box widget, with ShouldGenerateEmptyWidgetForSpacing settings this will prevent us from skipping outside of the header, avoiding alignment issues
+	FVariablesOutlinerEntryItem* TreeItem = Item->CastTo<FVariablesOutlinerEntryItem>();
+	if (TreeItem == nullptr)
 	{
-		return;
+		return SNew(SBox);
 	}
-
-	UAnimNextRigVMAssetEditorData* EditorData = VariableEntry->GetTypedOuter<UAnimNextRigVMAssetEditorData>();
-	if (EditorData == nullptr)
+	IAnimNextRigVMVariableInterface* Variable = Cast<IAnimNextRigVMVariableInterface>(TreeItem->WeakEntry.Get());
+	if(Variable == nullptr)
 	{
-		return;
+		return SNew(SBox);
 	}
+	
+	return SNew(SVariablesOutlinerValue, *TreeItem, WeakSceneOutliner.Pin().ToSharedRef().Get(), Row);
+}
+
+FLazyName VariablesOutlinerAccessSpecifier("AccessSpecifier");
 
-	VariableEntry->MarkPackageDirty();
-	VariableEntry->BroadcastModified(EAnimNextEditorDataNotifType::VariableDefaultValueChanged);
+FName FVariablesOutlinerAccessSpecifierColumn::GetID()
+{
+	return VariablesOutlinerAccessSpecifier;
 }
 
-const TSharedRef<SWidget> FVariablesOutlinerValueColumn::ConstructRowWidget(FSceneOutlinerTreeItemRef Item, const STableRow<FSceneOutlinerTreeItemPtr>& Row)
+SHeaderRow::FColumn::FArguments FVariablesOutlinerAccessSpecifierColumn::ConstructHeaderRowColumn()
 {
-	const FOutlinerAssetItem* AssetItem = Item->CastTo<FOutlinerAssetItem>();
-	if (AssetItem)
+	return
+		SHeaderRow::Column(GetColumnID())
+		.FixedWidth(24.0f)
+		.HAlignHeader(HAlign_Center)
+		.VAlignHeader(VAlign_Center)
+		.HAlignCell(HAlign_Fill)
+		.VAlignCell(VAlign_Fill)
+		[
+			SNew(SBox) 
+			.VAlign(VAlign_Center)
+			.HAlign(HAlign_Center)
+			[
+				SNew(SImage)
+				.ColorAndOpacity(FSlateColor::UseForeground())
+				.Image(FAppStyle::GetBrush("Level.VisibleIcon16x"))
+				.ToolTipText(LOCTEXT("AccessSpecifierAccessLevelTooltip", "Access level of this entry"))
+			]
+		];
+}
+
+class SVariablesOutlinerAccessSpecifier : public SCompoundWidget, public FNotifyHook
+{
+	SLATE_BEGIN_ARGS(SVariablesOutlinerAccessSpecifier) {}
+	SLATE_END_ARGS()
+
+	void Construct(const FArguments& InArgs, FVariablesOutlinerEntryItem& InTreeItem, ISceneOutliner& SceneOutliner, const STableRow<FSceneOutlinerTreeItemPtr>& InRow)
 	{
-		return SNullWidget::NullWidget;
+		WeakTreeItem = StaticCastSharedRef<FVariablesOutlinerEntryItem>(InTreeItem.AsShared());
+
+		ChildSlot
+		[
+			SNew(SBox)
+			.IsEnabled(!InTreeItem.WeakSharedVariablesEntry.IsValid())
+			.WidthOverride(16.0f)
+			.HeightOverride(16.0f)
+			.VAlign(VAlign_Center)
+			.HAlign(HAlign_Center)
+			.IsEnabled(this, &SVariablesOutlinerAccessSpecifier::IsButtonEnabled)
+			[
+				SNew(SButton)
+				.ButtonStyle(FAppStyle::Get(), "HoverHintOnly")
+				.OnClicked(this, &SVariablesOutlinerAccessSpecifier::OnClicked)
+				.Content()
+				[
+					SNew(SImage)
+					.ColorAndOpacity(FSlateColor::UseForeground())
+					.Image(this, &SVariablesOutlinerAccessSpecifier::GetImage)
+					.ToolTipText(this, &SVariablesOutlinerAccessSpecifier::GetTooltipText)
+				]
+			]
+		];
 	}
 
-	// Return a dummy box widget, with ShouldGenerateEmptyWidgetForSpacing settings this will prevent us from skipping outside of the header, avoiding alignment issues
+	FReply OnClicked() const
+	{
+		if (TSharedPtr<FVariablesOutlinerEntryItem> TreeItem = WeakTreeItem.Pin())
+		{
+			if (TSharedPtr<ISceneOutliner> SceneOutliner = TreeItem->WeakSceneOutliner.Pin())
+			{
+				if (const FVariablesOutlinerMode* Mode = static_cast<const FVariablesOutlinerMode*>(SceneOutliner->GetMode()))
+				{
+					Mode->ToggleSelectedVariablesExportAccessSpecifiers(TreeItem);
+				}
+			}
+		}
+
+		return FReply::Unhandled();	// Fall through so we dont deselect our item
+	}
+
+	const FSlateBrush* GetImage() const
+	{
+		TSharedPtr<FVariablesOutlinerEntryItem> TreeItem = WeakTreeItem.Pin();
+		if(!TreeItem.IsValid())
+		{
+			return nullptr;
+		}
+
+		IAnimNextRigVMExportInterface* Export = Cast<IAnimNextRigVMExportInterface>(TreeItem->WeakEntry.Get());
+		if(Export == nullptr)
+		{
+			return nullptr;
+		}
+
+		return Export->GetExportAccessSpecifier() == EAnimNextExportAccessSpecifier::Public ?
+			FAppStyle::GetBrush("Level.VisibleIcon16x") :
+			FAppStyle::GetBrush("Level.NotVisibleHighlightIcon16x");
+	}
+
+	bool IsButtonEnabled() const
+	{
+		TSharedPtr<FVariablesOutlinerEntryItem> TreeItem = WeakTreeItem.Pin();
+		if(!TreeItem.IsValid())
+		{
+			return false;
+		}
+
+		if (UAnimNextVariableEntry* Entry = TreeItem->WeakEntry.Get())
+		{
+			if (UAnimNextRigVMAsset* OuterAsset = Entry->GetTypedOuter<UAnimNextRigVMAsset>())
+			{
+				// Disable toggling public/private for pure UAnimNextSharedVariables objects				
+				return ExactCast<UAnimNextSharedVariables>(OuterAsset) == nullptr;	
+			}
+		}
+		
+		return false;		
+	}	
+
+	FText GetTooltipText() const
+	{
+		TSharedPtr<FVariablesOutlinerEntryItem> TreeItem = WeakTreeItem.Pin();
+		if(!TreeItem.IsValid())
+		{
+			return FText::GetEmpty();
+		}
+
+		IAnimNextRigVMExportInterface* Export = Cast<IAnimNextRigVMExportInterface>(TreeItem->WeakEntry.Get());
+		if(Export == nullptr)
+		{
+			return FText::GetEmpty();
+		}
+
+		FText AccessSpecifier = Export->GetExportAccessSpecifier() == EAnimNextExportAccessSpecifier::Public ? LOCTEXT("PublicSpecifier", "public") : LOCTEXT("PrivateSpecifier", "private");
+		FText AccessSpecifierDesc = Export->GetExportAccessSpecifier() == EAnimNextExportAccessSpecifier::Public ?
+			LOCTEXT("PublicSpecifierDesc", "This means that the entry is usable from gameplay and from other AnimNext assets") :
+			LOCTEXT("PrivateSpecifierDesc", "This means that the entry is only usable inside this asset");
+		return FText::Format(LOCTEXT("AccessSpecifierEntryTooltip", "This entry is {0}.\n{1}"), AccessSpecifier, AccessSpecifierDesc);
+	}
+
+	TWeakPtr<FVariablesOutlinerEntryItem> WeakTreeItem;
+};
+
+const TSharedRef<SWidget> FVariablesOutlinerAccessSpecifierColumn::ConstructRowWidget(FSceneOutlinerTreeItemRef Item, const STableRow<FSceneOutlinerTreeItemPtr>& Row)
+{
 	FVariablesOutlinerEntryItem* TreeItem = Item->CastTo<FVariablesOutlinerEntryItem>();
 	if (TreeItem == nullptr)
 	{
-		return SNew(SBox);
+		return SNullWidget::NullWidget;
 	}
+
 	IAnimNextRigVMVariableInterface* Variable = Cast<IAnimNextRigVMVariableInterface>(TreeItem->WeakEntry.Get());
 	if(Variable == nullptr)
 	{
-		return SNew(SBox);
+		return SNullWidget::NullWidget;
 	}
 	
-	return SNew(SVariablesOutlinerValue, *TreeItem, WeakSceneOutliner.Pin().ToSharedRef().Get(), Row);
+	return SNew(SVariablesOutlinerAccessSpecifier, *TreeItem, WeakSceneOutliner.Pin().ToSharedRef().Get(), Row);
 }
 
 }
diff --git a/Engine/Plugins/Experimental/UAF/UAF/Source/UAFEditor/Private/Variables/Outliner/VariablesOutlinerColumns.h b/Engine/Plugins/Experimental/UAF/UAF/Source/UAFEditor/Private/Variables/Outliner/VariablesOutlinerColumns.h
--- a/Engine/Plugins/Experimental/UAF/UAF/Source/UAFEditor/Private/Variables/Outliner/VariablesOutlinerColumns.h
+++ b/Engine/Plugins/Experimental/UAF/UAF/Source/UAFEditor/Private/Variables/Outliner/VariablesOutlinerColumns.h
@@ -56,6 +56,23 @@ private:
 	TWeakPtr<ISceneOutliner> WeakSceneOutliner;
 };
 
+class FVariablesOutlinerAccessSpecifierColumn : public ISceneOutlinerColumn
+{
+public:
+	static FName GetID();
+
+	FVariablesOutlinerAccessSpecifierColumn(ISceneOutliner& SceneOutliner)
+		: WeakSceneOutliner(StaticCastSharedRef<ISceneOutliner>(SceneOutliner.AsShared()))
+	{}
 
+	// ISceneOutlinerColumn interface
+	virtual FName GetColumnID() override { return GetID(); }
+	virtual SHeaderRow::FColumn::FArguments ConstructHeaderRowColumn() override;
+	virtual const TSharedRef<SWidget> ConstructRowWidget(FSceneOutlinerTreeItemRef TreeItem, const STableRow<FSceneOutlinerTreeItemPtr>& Row) override;
+	virtual bool SupportsSorting() const override { return false; }
+
+private:
+	TWeakPtr<ISceneOutliner> WeakSceneOutliner;
+};
 
 }
diff --git a/Engine/Plugins/Experimental/UAF/UAF/Source/UAFEditor/Private/Variables/Outliner/VariablesOutlinerCommands.cpp b/Engine/Plugins/Experimental/UAF/UAF/Source/UAFEditor/Private/Variables/Outliner/VariablesOutlinerCommands.cpp
--- a/Engine/Plugins/Experimental/UAF/UAF/Source/UAFEditor/Private/Variables/Outliner/VariablesOutlinerCommands.cpp
+++ b/Engine/Plugins/Experimental/UAF/UAF/Source/UAFEditor/Private/Variables/Outliner/VariablesOutlinerCommands.cpp
@@ -10,6 +10,15 @@ void FVariablesOutlinerCommands::RegisterCommands()
 	UI_COMMAND(AddNewVariable, "Add Variable", "Adds a single new variable to the selected asset.", EUserInterfaceActionType::Button, FInputChord());
 	UI_COMMAND(AddNewVariables, "Add Variable(s)", "Adds variables to assets.\nIf multiple assets are selected, then variables will be added to each.\nIf no assets are selected and there are multiple assets, variables will be added to all assets.", EUserInterfaceActionType::Button, FInputChord());
 
+	UI_COMMAND(FindReferences, "In Project", "Finds all references to the selected variable variables across assets in the project.", EUserInterfaceActionType::Button, FInputChord());
+	UI_COMMAND(FindReferencesInWorkspace, "In Workspace", "Finds all references to the selected variable variables across assets in the current workspace.", EUserInterfaceActionType::Button, FInputChord());
+	UI_COMMAND(FindReferencesInAsset, "In Asset", "Finds all references to the selected variable variables in the currently opened asset.", EUserInterfaceActionType::Button, FInputChord());
+
+	UI_COMMAND(SaveAsset, "Save Asset", "Save the selected asset.", EUserInterfaceActionType::Button, FInputChord());
+
+	UI_COMMAND(SetVariableExportPublic, "Sets Variable Export to Public", "Sets Variable to be Public.", EUserInterfaceActionType::Button, FInputChord());
+	UI_COMMAND(SetVariableExportPrivate, "Sets Variable Export to Private", "Sets Variable to be Private.", EUserInterfaceActionType::Button, FInputChord());
+
 	UI_COMMAND(CreateSharedVariablesAssets, "Create Shared Variables from Selection", "Creates a new Shared Variables asset from the selected set of variables.", EUserInterfaceActionType::Button, FInputChord());
 }
 
diff --git a/Engine/Plugins/Experimental/UAF/UAF/Source/UAFEditor/Private/Variables/Outliner/VariablesOutlinerCommands.h b/Engine/Plugins/Experimental/UAF/UAF/Source/UAFEditor/Private/Variables/Outliner/VariablesOutlinerCommands.h
--- a/Engine/Plugins/Experimental/UAF/UAF/Source/UAFEditor/Private/Variables/Outliner/VariablesOutlinerCommands.h
+++ b/Engine/Plugins/Experimental/UAF/UAF/Source/UAFEditor/Private/Variables/Outliner/VariablesOutlinerCommands.h
@@ -19,6 +19,12 @@ public:
 
 	TSharedPtr<FUICommandInfo> AddNewVariable;
 	TSharedPtr<FUICommandInfo> AddNewVariables;
+	TSharedPtr<FUICommandInfo> FindReferences;
+	TSharedPtr<FUICommandInfo> FindReferencesInWorkspace;
+	TSharedPtr<FUICommandInfo> FindReferencesInAsset;
+	TSharedPtr<FUICommandInfo> SaveAsset;	
+	TSharedPtr<FUICommandInfo> SetVariableExportPublic;
+	TSharedPtr<FUICommandInfo> SetVariableExportPrivate;
 	TSharedPtr<FUICommandInfo> CreateSharedVariablesAssets;
 };
 
diff --git a/Engine/Plugins/Experimental/UAF/UAF/Source/UAFEditor/Private/Variables/Outliner/VariablesOutlinerDragDrop.cpp b/Engine/Plugins/Experimental/UAF/UAF/Source/UAFEditor/Private/Variables/Outliner/VariablesOutlinerDragDrop.cpp
new file mode 100644
--- /dev/null
+++ b/Engine/Plugins/Experimental/UAF/UAF/Source/UAFEditor/Private/Variables/Outliner/VariablesOutlinerDragDrop.cpp
@@ -0,0 +1,53 @@
+﻿// Copyright Epic Games, Inc. All Rights Reserved.
+
+#include "VariablesOutlinerDragDrop.h"
+
+#include "Common/GraphEditorSchemaActions.h"
+
+#define LOCTEXT_NAMESPACE "VariableDragDropOp"
+
+namespace UE::UAF::Editor
+{
+
+TSharedPtr<FVariableDragDropOp> FVariableDragDropOp::New(TSharedPtr<FAnimNextSchemaAction_Variable> InAction, TSharedPtr<FVariablesOutlinerEntryItem> InEntry)
+{
+	TSharedPtr<FVariableDragDropOp> NewOp = MakeShared<FVariableDragDropOp>();
+	NewOp->WeakItem = InEntry;
+	NewOp->SourceAction = InAction;
+	NewOp->Construct();
+	return NewOp;
+}
+
+TSharedPtr<FAnimNextSchemaAction_Variable> FVariableDragDropOp::GetAction() const
+{
+	if (SourceAction.IsValid())
+	{
+		return StaticCastSharedPtr<FAnimNextSchemaAction_Variable>(SourceAction);
+	}
+
+	return nullptr;
+}
+
+void FVariableDragDropOp::GetDefaultStatusSymbol(const FSlateBrush*& PrimaryBrushOut, FSlateColor& IconColorOut, FSlateBrush const*& SecondaryBrushOut, FSlateColor& SecondaryColorOut) const
+{
+	PrimaryBrushOut = StaticCastSharedPtr<FAnimNextSchemaAction_Variable>(SourceAction)->GetIconBrush();
+	IconColorOut = StaticCastSharedPtr<FAnimNextSchemaAction_Variable>(SourceAction)->GetIconColor();
+	SecondaryBrushOut = nullptr;
+	SecondaryColorOut = FSlateColor::UseForeground();
+}
+
+TSharedPtr<FCategoryDragDropOp> FCategoryDragDropOp::New(TSharedPtr<FVariablesOutlinerCategoryItem> InEntry)
+{
+	TSharedPtr<FCategoryDragDropOp> NewOp = MakeShared<FCategoryDragDropOp>();
+	NewOp->WeakItem = InEntry;
+
+	NewOp->CurrentHoverText = FText::FromString(InEntry->CategoryName);
+	NewOp->SetupDefaults();
+	NewOp->Construct();
+	
+	return NewOp;
+}
+
+}
+
+#undef LOCTEXT_NAMESPACE
\ No newline at end of file
diff --git a/Engine/Plugins/Experimental/UAF/UAF/Source/UAFEditor/Private/Variables/Outliner/VariablesOutlinerDragDrop.h b/Engine/Plugins/Experimental/UAF/UAF/Source/UAFEditor/Private/Variables/Outliner/VariablesOutlinerDragDrop.h
new file mode 100644
--- /dev/null
+++ b/Engine/Plugins/Experimental/UAF/UAF/Source/UAFEditor/Private/Variables/Outliner/VariablesOutlinerDragDrop.h
@@ -0,0 +1,44 @@
+﻿// Copyright Epic Games, Inc. All Rights Reserved.
+
+#pragma once
+
+#include "GraphEditorDragDropAction.h"
+#include "ISceneOutlinerTreeItem.h"
+#include "VariablesOutlinerEntryItem.h"
+#include "VariablesOutlinerCategoryItem.h"
+#include "DragAndDrop/DecoratedDragDropOp.h"
+#include "UObject/WeakInterfacePtr.h"
+
+class IAnimNextRigVMVariableInterface;
+struct FAnimNextSchemaAction_Variable;
+class UAnimNextRigVMAssetEntry;
+
+namespace UE::UAF::Editor
+{
+
+class FVariableDragDropOp : public FGraphSchemaActionDragDropAction
+{
+public:
+	DRAG_DROP_OPERATOR_TYPE(FVariableDragDropOp, FGraphSchemaActionDragDropAction)
+
+	static TSharedPtr<FVariableDragDropOp> New(TSharedPtr<FAnimNextSchemaAction_Variable> InAction, TSharedPtr<FVariablesOutlinerEntryItem> InEntry);
+
+	// FGraphSchemaActionDragDropAction interface
+	virtual void GetDefaultStatusSymbol(const FSlateBrush*& PrimaryBrushOut, FSlateColor& IconColorOut, FSlateBrush const*& SecondaryBrushOut, FSlateColor& SecondaryColorOut) const override;
+
+	TSharedPtr<FAnimNextSchemaAction_Variable> GetAction() const;
+
+	TWeakPtr<FVariablesOutlinerEntryItem> WeakItem;
+};
+
+class FCategoryDragDropOp : public FDecoratedDragDropOp
+{
+public:
+	DRAG_DROP_OPERATOR_TYPE(FCategoryDragDropOp, FDecoratedDragDropOp)
+
+	static TSharedPtr<FCategoryDragDropOp> New(TSharedPtr<FVariablesOutlinerCategoryItem> InEntry);
+
+	TWeakPtr<FVariablesOutlinerCategoryItem> WeakItem;
+};
+
+}
diff --git a/Engine/Plugins/Experimental/UAF/UAF/Source/UAFEditor/Private/Variables/Outliner/VariablesOutlinerEntryItem.cpp b/Engine/Plugins/Experimental/UAF/UAF/Source/UAFEditor/Private/Variables/Outliner/VariablesOutlinerEntryItem.cpp
--- a/Engine/Plugins/Experimental/UAF/UAF/Source/UAFEditor/Private/Variables/Outliner/VariablesOutlinerEntryItem.cpp
+++ b/Engine/Plugins/Experimental/UAF/UAF/Source/UAFEditor/Private/Variables/Outliner/VariablesOutlinerEntryItem.cpp
@@ -3,15 +3,24 @@
 #include "VariablesOutlinerEntryItem.h"
 
 #include "ISceneOutliner.h"
+#include "Styling/SlateColor.h"
+#include "Misc/Optional.h"
 #include "SceneOutlinerStandaloneTypes.h"
+#include "AnimNextEditorModule.h"
 #include "AnimNextRigVMAssetEditorData.h"
 #include "EditorUtils.h"
 #include "ScopedTransaction.h"
 #include "UncookedOnlyUtils.h"
 #include "Entries/AnimNextRigVMAssetEntry.h"
 #include "UObject/Package.h"
+#include "Widgets/SCompoundWidget.h"
+#include "Widgets/DeclarativeSyntaxSupport.h"
+#include "Widgets/Images/SImage.h"
+#include "Widgets/Text/SInlineEditableTextBlock.h"
+#include "Widgets/Text/STextBlock.h"
+#include "Widgets/Views/STableRow.h"
 #include "Widgets/Views/SListView.h"
-#include "Common/Outliner/OutlinerDragDropOps.h"
+#include "VariablesOutlinerDragDrop.h"
 #include "Entries/AnimNextVariableEntry.h"
 
 #define LOCTEXT_NAMESPACE "VariablesOutlinerTreeItem"
@@ -19,13 +28,319 @@
 namespace UE::UAF::Editor
 {
 
-const FSceneOutlinerTreeItemType FVariablesOutlinerEntryItem::Type(&FOutlinerEntryItem::Type);
+const FSceneOutlinerTreeItemType FVariablesOutlinerEntryItem::Type;
 
-FVariablesOutlinerEntryItem::FVariablesOutlinerEntryItem(const FVariablesOutlinerEntryItem::FItemData& InItemData)
-: FOutlinerEntryItem(FVariablesOutlinerEntryItem::Type, { InItemData.Entry ? InItemData.Entry ->GetTypedOuter<UAnimNextRigVMAsset>() : nullptr, InItemData.SortValue })
-	, WeakEntry(InItemData.Entry)
-	, PropertyPath(InItemData.Property)
-	, bStructOwner(InItemData.Property != nullptr)
+class SVariablesOutlinerEntryLabel : FSceneOutlinerCommonLabelData, public SCompoundWidget
+{
+	SLATE_BEGIN_ARGS(SVariablesOutlinerEntryLabel) {}
+	SLATE_END_ARGS()
+
+	void Construct(const FArguments& InArgs, FVariablesOutlinerEntryItem& InTreeItem, ISceneOutliner& SceneOutliner, const STableRow<FSceneOutlinerTreeItemPtr>& InRow)
+	{
+		WeakSceneOutliner = StaticCastSharedRef<ISceneOutliner>(SceneOutliner.AsShared());
+		TreeItem = StaticCastSharedRef<FVariablesOutlinerEntryItem>(InTreeItem.AsShared());
+
+		ChildSlot
+		[
+			SNew(SHorizontalBox)
+			+SHorizontalBox::Slot()
+			.AutoWidth()
+			.VAlign(VAlign_Center)
+			.Padding(2.0f, 2.0f)
+			[
+				SAssignNew(TextBlock, SInlineEditableTextBlock)
+				.IsReadOnly(this, &SVariablesOutlinerEntryLabel::IsReadOnly)
+				.Text(this, &SVariablesOutlinerEntryLabel::GetDisplayText)
+				.HighlightText(SceneOutliner.GetFilterHighlightText())
+				.ColorAndOpacity(this, &SVariablesOutlinerEntryLabel::GetForegroundColor)
+				.OnTextCommitted(this, &SVariablesOutlinerEntryLabel::OnTextCommited)
+				.OnVerifyTextChanged(this, &SVariablesOutlinerEntryLabel::OnVerifyTextChanged)
+			]
+			+SHorizontalBox::Slot()
+			.VAlign(VAlign_Center)
+			.Padding(2.0f, 0.0f, 2.0f, 3.0f)
+			.AutoWidth()
+			[
+				SNew(SImage)
+				.ToolTipText(this, &SVariablesOutlinerEntryLabel::GetDirtyTooltipText)
+				.Image(this, &SVariablesOutlinerEntryLabel::GetDirtyImageBrush)
+			]
+		];
+	}
+
+	FText GetDirtyTooltipText() const
+	{
+		if (const TSharedPtr<FVariablesOutlinerEntryItem> Item = TreeItem.Pin())
+		{
+			FTextBuilder TextBuilder;
+			TextBuilder.AppendLine(LOCTEXT("ModifiedTooltip", "Modified"));
+
+			if(UAnimNextVariableEntry* AssetEntry = Item->WeakEntry.Get())
+			{
+				const UPackage* ExternalPackage = AssetEntry->GetExternalPackage();
+				if(ExternalPackage && ExternalPackage->IsDirty())
+				{
+					TextBuilder.AppendLine(FText::FromName(ExternalPackage->GetFName()));
+				}
+			}
+
+			return TextBuilder.ToText();
+		}
+		return FText::GetEmpty();
+	}
+
+	const FSlateBrush* GetDirtyImageBrush() const
+	{
+		if (const TSharedPtr<FVariablesOutlinerEntryItem> Item = TreeItem.Pin())
+		{
+			if(UAnimNextVariableEntry* AssetEntry = Item->WeakEntry.Get())
+			{
+				bool bIsDirty = false;
+				const UPackage* ExternalPackage = AssetEntry->GetExternalPackage();
+				if(ExternalPackage && ExternalPackage->IsDirty())
+				{
+					bIsDirty = true;
+				}
+
+				return bIsDirty ? FAppStyle::GetBrush("Icons.DirtyBadge") : nullptr;
+			}
+		}
+		return nullptr;
+	}
+	
+	FText GetDisplayText() const
+	{
+		if (const TSharedPtr<FVariablesOutlinerEntryItem> Item = TreeItem.Pin())
+		{
+			return FText::FromString(Item->GetDisplayString());
+		}
+		return FText();
+	}
+
+	void OnTextCommited(const FText& InLabel, ETextCommit::Type InCommitInfo) const
+	{
+		if(InCommitInfo == ETextCommit::OnEnter)
+		{
+			if (const TSharedPtr<FVariablesOutlinerEntryItem> Item = TreeItem.Pin())
+			{
+				Item->Rename(InLabel); 
+			}
+		}
+	}
+
+	bool OnVerifyTextChanged(const FText& InLabel, FText& OutErrorMessage) const
+	{
+		if (const TSharedPtr<FVariablesOutlinerEntryItem> Item = TreeItem.Pin())
+		{
+			return Item->ValidateName(InLabel, OutErrorMessage); 
+		}
+		return false;
+	}
+
+	bool IsReadOnly() const
+	{
+		if (const TSharedPtr<FVariablesOutlinerEntryItem> Item = TreeItem.Pin())
+		{
+			return !Item->PropertyPath.IsPathToFieldEmpty() || Item->WeakEntry.Get() == nullptr;
+		}
+		return false;
+	}
+	
+	virtual FSlateColor GetForegroundColor() const override
+	{
+		const TOptional<FLinearColor> BaseColor = FSceneOutlinerCommonLabelData::GetForegroundColor(*TreeItem.Pin());
+		return BaseColor.IsSet() ? BaseColor.GetValue() : FSlateColor::UseForeground();
+	}
+	
+	virtual FReply OnDrop(const FGeometry& MyGeometry, const FDragDropEvent& DragDropEvent) override
+	{
+		TSharedPtr<FVariableDragDropOp> GraphDropOp = DragDropEvent.GetOperationAs<FVariableDragDropOp>();
+		if (GraphDropOp.IsValid())
+		{
+			if (const TSharedPtr<FVariablesOutlinerEntryItem> ThisEntryItem = TreeItem.Pin())
+			{
+				const UAnimNextVariableEntry* ThisEntry = ThisEntryItem->WeakEntry.Get();
+				TSharedPtr<FVariablesOutlinerEntryItem> EntryItem = GraphDropOp->WeakItem.Pin();
+				if (ThisEntryItem->HasStructOwner())
+				{
+					const FSlateBrush* ErrorSymbol = FAppStyle::GetBrush(TEXT("Graph.ConnectorFeedback.Error"));
+					GraphDropOp->SetSimpleFeedbackMessage(ErrorSymbol, FSlateColor(GetForegroundColor()), LOCTEXT("StructItemFeedback", "Cannot make changes to Struct Entries"));
+				}
+				else if (EntryItem.IsValid())
+				{
+					UAnimNextVariableEntry* OtherEntry = EntryItem->WeakEntry.Get();
+					if (ThisEntry && OtherEntry && ThisEntry != OtherEntry)
+					{
+						UAnimNextRigVMAsset* ThisAsset = ThisEntry->GetTypedOuter<UAnimNextRigVMAsset>();
+						const UAnimNextRigVMAsset* OtherAsset = OtherEntry->GetTypedOuter<UAnimNextRigVMAsset>();
+						if (ThisAsset && ThisAsset == OtherAsset)
+						{
+							if (ThisEntry->GetVariableCategory() != OtherEntry->GetVariableCategory())
+							{
+								if (OtherEntry && ThisEntry)
+								{
+									UAnimNextRigVMAssetEditorData* VariableEditorData = OtherEntry->GetTypedOuter<UAnimNextRigVMAssetEditorData>();
+									UAnimNextRigVMAssetEditorData* ThisEditorData = ThisEntry->GetTypedOuter<UAnimNextRigVMAssetEditorData>();
+
+									if (VariableEditorData && ThisEditorData)
+									{
+										if (VariableEditorData == ThisEditorData)
+										{
+											// Add or remove to category of this variable entry
+											const bool bRemovingFromCategory = ThisEntry->GetVariableCategory().IsEmpty();
+
+											const FText OtherVariableNameText = FText::FromName(OtherEntry->GetVariableName());					
+											const FText OtherVariableCategoryText = FText::FromString(OtherEntry->GetVariableCategory().GetData());
+											const FText ThisVariableNameText = FText::FromName(ThisEntry->GetVariableName());
+											const FText ThisVariableCategoryText = FText::FromString(ThisEntry->GetVariableCategory().GetData());											
+											const FText AddFormattedMessage = FText::Format(LOCTEXT("AddVariableToCategoryFormat", "Add {0} to Category {1}"), OtherVariableNameText, ThisVariableCategoryText);
+											const FText RemoveFormattedMessage = FText::Format(LOCTEXT("RemoveVariableToCategoryFormat", "Remove {0} from Category {1}"), OtherVariableNameText, OtherVariableCategoryText);
+											
+											FScopedTransaction Transaction(bRemovingFromCategory ? RemoveFormattedMessage : AddFormattedMessage);
+											OtherEntry->SetVariableCategory(ThisEntry->GetVariableCategory(), true);		
+										}
+									}
+								}
+							}
+							else if (OtherEntry)
+							{
+								UAnimNextVariableEntry* VariableEntry = EntryItem->WeakEntry.Get();
+								const UAnimNextVariableEntry* ThisVariableEntry = ThisEntryItem->WeakEntry.Get();
+								if (VariableEntry && ThisVariableEntry)
+								{
+									const FText OtherVariableNameText = FText::FromName(OtherEntry->GetVariableName());					
+									const FText ThisVariableNameText = FText::FromName(ThisEntry->GetVariableName());
+									const FText FormattedMessage = FText::Format(LOCTEXT("ReorderVariableFormat", "Reorder {0} before {1}"), OtherVariableNameText, ThisVariableNameText);
+									FScopedTransaction Transaction(FormattedMessage);
+									
+									UAnimNextRigVMAssetEditorData* VariableEditorData = VariableEntry->GetTypedOuter<UAnimNextRigVMAssetEditorData>();
+									UAnimNextRigVMAssetEditorData* ThisEditorData = ThisVariableEntry->GetTypedOuter<UAnimNextRigVMAssetEditorData>();
+									if (VariableEditorData && VariableEditorData == ThisEditorData)
+									{
+										VariableEditorData->ReorderVariable(VariableEntry, ThisVariableEntry);
+									}
+								}								
+							}
+						}
+						else if (ThisAsset)
+						{
+							if (UAnimNextVariableEntry* VariableEntry = EntryItem->WeakEntry.Get())
+							{
+								const FText VariableNameText = FText::FromName(OtherEntry->GetVariableName());
+								const FText AssetNameText = FText::FromName(ThisAsset->GetFName());
+																		
+								const FText FormattedMessage = FText::Format(LOCTEXT("MoveVariableToAssetFormatTransaction", "Moving {0} to {1}"), VariableNameText, AssetNameText);
+								FScopedTransaction Transaction(FormattedMessage);
+								UncookedOnly::FUtils::MoveVariableToAsset(VariableEntry, ThisAsset, true, true);
+							}
+						}	
+					}
+				}
+			}
+		}
+		
+		return FReply::Unhandled();
+	}
+
+	virtual void OnDragEnter(const FGeometry& MyGeometry, const FDragDropEvent& DragDropEvent) override
+	{
+		TSharedPtr<FVariableDragDropOp> GraphDropOp = DragDropEvent.GetOperationAs<FVariableDragDropOp>();
+		if (GraphDropOp.IsValid())
+		{
+			if (const TSharedPtr<FVariablesOutlinerEntryItem> ThisEntryItem = TreeItem.Pin())
+			{
+				const UAnimNextVariableEntry* ThisEntry = ThisEntryItem->WeakEntry.Get();
+				TSharedPtr<FVariablesOutlinerEntryItem> EntryItem = GraphDropOp->WeakItem.Pin();
+
+				const FSlateBrush* ErrorSymbol = FAppStyle::GetBrush(TEXT("Graph.ConnectorFeedback.Error"));
+				if (ThisEntryItem->HasStructOwner())
+				{
+					GraphDropOp->SetSimpleFeedbackMessage(ErrorSymbol, FSlateColor(GetForegroundColor()), LOCTEXT("StructItemFeedback", "Cannot make changes to Struct Entries"));
+				}
+				else if (ThisEntry && EntryItem.IsValid())
+				{
+					if (const UAnimNextVariableEntry* OtherEntry = EntryItem->WeakEntry.Get())
+					{
+						const FText OtherVariableNameText = FText::FromName(OtherEntry->GetVariableName());					
+						const FText OtherVariableCategoryText = FText::FromString(OtherEntry->GetVariableCategory().GetData());
+						const FText ThisVariableNameText = FText::FromName(ThisEntry->GetVariableName());
+						const FText ThisVariableCategoryText = FText::FromString(ThisEntry->GetVariableCategory().GetData());
+						
+						if (ThisEntry && ThisEntry != OtherEntry)
+						{
+							UAnimNextRigVMAsset* ThisAsset = ThisEntryItem->WeakEntry.Get()->GetTypedOuter<UAnimNextRigVMAsset>();
+							const UAnimNextRigVMAsset* OtherAsset = EntryItem->WeakEntry.Get()->GetTypedOuter<UAnimNextRigVMAsset>();
+							if (ThisAsset && ThisAsset == OtherAsset)
+							{
+								if (ThisEntry->GetVariableCategory() != OtherEntry->GetVariableCategory())
+								{
+									if (ThisEntry->GetVariableCategory().IsEmpty())
+									{
+										// Remove from category
+										const FText FormattedMessage = FText::Format(LOCTEXT("RemoveVariableToCategoryFormat", "Remove {0} from Category {1}"), OtherVariableNameText, OtherVariableCategoryText);
+										GraphDropOp->SetSimpleFeedbackMessage(nullptr, FSlateColor(GetForegroundColor()), FormattedMessage);
+									}
+									else
+									{
+										// Add to category
+										const FText FormattedMessage = FText::Format(LOCTEXT("AddVariableToCategoryFormat", "Add {0} to Category {1}"), OtherVariableNameText, ThisVariableCategoryText);
+										GraphDropOp->SetSimpleFeedbackMessage(nullptr, FSlateColor(GetForegroundColor()), FormattedMessage);
+									}
+								}
+								else
+								{
+									// Reorder
+									const FText FormattedMessage = FText::Format(LOCTEXT("ReorderVariableFormat", "Reorder {0} before {1}"), OtherVariableNameText, ThisVariableNameText);
+									GraphDropOp->SetSimpleFeedbackMessage(nullptr, FSlateColor(GetForegroundColor()), FormattedMessage);
+								}
+							}
+							else if (ThisAsset)
+							{
+								// Reorder
+								GraphDropOp->SetSimpleFeedbackMessage(nullptr, FSlateColor(GetForegroundColor()), LOCTEXT("ReorderDifferentOuterVariableLabel", "Cannot Reorder Variable in different Asset"));
+
+								const FText VariableNameText = FText::FromName(OtherEntry->GetVariableName());
+								const FText AssetNameText = FText::FromName(ThisAsset->GetFName());
+								
+								// [TODO] should we only allow move variables "up" the asset chain? If so validate that according to the workspace export chain, that VariableOuter is a child of Asset
+								const FText FormattedMessage = FText::Format(LOCTEXT("MoveVariableToAssetFormat", "Move {0} to {1}"), VariableNameText, AssetNameText);
+								GraphDropOp->SetSimpleFeedbackMessage(nullptr, FSlateColor(GetForegroundColor()), FormattedMessage);
+							}						
+						}
+						else
+						{
+							// Reorder
+							GraphDropOp->SetSimpleFeedbackMessage(ErrorSymbol, FSlateColor(GetForegroundColor()), LOCTEXT("ReorderSelfVariableLabel", "Cannot Reorder Variable before itself"));
+						}
+					}					
+				}
+			}
+		}
+	}
+
+	virtual void OnDragLeave( const FDragDropEvent& DragDropEvent ) override
+	{
+		TSharedPtr<FVariableDragDropOp> GraphDropOp = DragDropEvent.GetOperationAs<FVariableDragDropOp>();
+		if (GraphDropOp.IsValid())
+		{
+			GraphDropOp->HoverTargetChanged();
+		}
+	}
+
+	TWeakPtr<FVariablesOutlinerEntryItem> TreeItem;
+	TSharedPtr<SInlineEditableTextBlock> TextBlock;
+};
+
+FVariablesOutlinerEntryItem::FVariablesOutlinerEntryItem(UAnimNextVariableEntry* InEntry)
+	: ISceneOutlinerTreeItem(FVariablesOutlinerEntryItem::Type)
+	, WeakEntry(InEntry)
+{
+}
+
+FVariablesOutlinerEntryItem::FVariablesOutlinerEntryItem(const FProperty* InProperty)
+	: ISceneOutlinerTreeItem(FVariablesOutlinerEntryItem::Type)
+	, PropertyPath(InProperty)
+	, bStructOwner(true)
 {
 }
 
@@ -60,6 +375,13 @@ FString FVariablesOutlinerEntryItem::GetDisplayString() const
 	return FString();
 }
 
+TSharedRef<SWidget> FVariablesOutlinerEntryItem::GenerateLabelWidget(ISceneOutliner& Outliner, const STableRow<FSceneOutlinerTreeItemPtr>& InRow)
+{
+	TSharedRef<SVariablesOutlinerEntryLabel> LabelWidget = SNew(SVariablesOutlinerEntryLabel, *this, Outliner, InRow);
+	RenameRequestEvent.BindSP(LabelWidget->TextBlock.Get(), &SInlineEditableTextBlock::EnterEditingMode);
+	return LabelWidget;
+}
+
 FString FVariablesOutlinerEntryItem::GetPackageName() const
 {
 	if(UAnimNextRigVMAssetEntry* Entry = WeakEntry.Get())
@@ -118,58 +440,6 @@ bool FVariablesOutlinerEntryItem::ValidateName(const FText& InNewName, FText& Ou
 	return true;
 }
 
-void FVariablesOutlinerEntryItem::SetAccessSpecifier(const EAnimNextExportAccessSpecifier& InSpecifier) const
-{
-	IAnimNextRigVMExportInterface* Export = Cast<IAnimNextRigVMExportInterface>(WeakEntry.Get());
-	if(Export == nullptr)
-	{
-		return;
-	}
-
-	FScopedTransaction Transaction(LOCTEXT("SetAccessSpecifierTransaction", "Set Access Specifier"));
-	Export->SetExportAccessSpecifier(InSpecifier);
-}
-
-EAnimNextExportAccessSpecifier FVariablesOutlinerEntryItem::GetAccessSpecifier() const
-{
-	IAnimNextRigVMExportInterface* Export = Cast<IAnimNextRigVMExportInterface>(WeakEntry.Get());
-	if(Export == nullptr)
-	{
-		return EAnimNextExportAccessSpecifier::Private;
-	}
-
-	FScopedTransaction Transaction(LOCTEXT("SetAccessSpecifierTransaction", "Set Access Specifier"));
-	return Export->GetExportAccessSpecifier();
-}
-
-bool FVariablesOutlinerEntryItem::CanSetAccessSpecifier() const
-{
-	if (UAnimNextVariableEntry* Entry = WeakEntry.Get())
-	{
-		if (UAnimNextRigVMAsset* OuterAsset = Entry->GetTypedOuter<UAnimNextRigVMAsset>())
-		{
-			// Disable toggling public/private for pure UAnimNextSharedVariables objects				
-			return ExactCast<UAnimNextSharedVariables>(OuterAsset) == nullptr;	
-		}
-	}
-
-	return false;
-}
-
-FStringView FVariablesOutlinerEntryItem::GetCategoryPath() const
-{
-	if (UAnimNextVariableEntry* Entry = WeakEntry.Get())
-	{
-		return Entry->GetVariableCategory();
-	}
-	
-	return FStringView();
-}
-
-bool FVariablesOutlinerEntryItem::IsReadOnly() const
-{
-	return !PropertyPath.IsPathToFieldEmpty() || WeakEntry.Get() == nullptr;
-}
 }
 
 #undef LOCTEXT_NAMESPACE // "VariablesOutlinerTreeItem"
\ No newline at end of file
diff --git a/Engine/Plugins/Experimental/UAF/UAF/Source/UAFEditor/Private/Variables/Outliner/VariablesOutlinerEntryItem.h b/Engine/Plugins/Experimental/UAF/UAF/Source/UAFEditor/Private/Variables/Outliner/VariablesOutlinerEntryItem.h
--- a/Engine/Plugins/Experimental/UAF/UAF/Source/UAFEditor/Private/Variables/Outliner/VariablesOutlinerEntryItem.h
+++ b/Engine/Plugins/Experimental/UAF/UAF/Source/UAFEditor/Private/Variables/Outliner/VariablesOutlinerEntryItem.h
@@ -3,7 +3,6 @@
 #pragma once
 
 #include "ISceneOutlinerTreeItem.h"
-#include "Common/Outliner/OutlinerEntryItem.h"
 
 class UAnimNextSharedVariablesEntry;
 class UAnimNextVariableEntry;
@@ -11,39 +10,27 @@ class UAnimNextVariableEntry;
 namespace UE::UAF::Editor
 {
 
-struct FVariablesOutlinerEntryItem : FOutlinerEntryItem
+struct FVariablesOutlinerEntryItem : ISceneOutlinerTreeItem
 {
 	static const FSceneOutlinerTreeItemType Type;
-	
-	struct FItemData 
-	{
-		UAnimNextVariableEntry* Entry;
-		const FProperty* Property;
-		int32 SortValue;
-	};
 
-	FVariablesOutlinerEntryItem() = delete;
-	FVariablesOutlinerEntryItem(const FVariablesOutlinerEntryItem::FItemData& InItemData);
+	FVariablesOutlinerEntryItem(UAnimNextVariableEntry* InEntry);
+	FVariablesOutlinerEntryItem(const FProperty* InProperty);
 
 	// Begin ISceneOutlinerTreeItem overrides
 	virtual bool IsValid() const override;
 	virtual FSceneOutlinerTreeItemID GetID() const override;
 	virtual FString GetDisplayString() const override;
 	virtual bool CanInteract() const override { return true; }
+	virtual TSharedRef<SWidget> GenerateLabelWidget(ISceneOutliner& Outliner, const STableRow<FSceneOutlinerTreeItemPtr>& InRow) override;
 	virtual FString GetPackageName() const override;
 	// End ISceneOutlinerTreeItem overrides
 
 	// Renames the item to the specified name
-	virtual void Rename(const FText& InNewName) const override;
+	void Rename(const FText& InNewName) const;
 
 	// Validates the new item name
-	virtual bool ValidateName(const FText& InNewName, FText& OutErrorMessage) const override;
-
-	virtual void SetAccessSpecifier(const EAnimNextExportAccessSpecifier& InSpecifier) const override;
-	virtual EAnimNextExportAccessSpecifier GetAccessSpecifier() const override;
-	virtual bool CanSetAccessSpecifier() const override;
-	virtual FStringView GetCategoryPath() const override;
-	virtual bool IsReadOnly() const override;
+	bool ValidateName(const FText& InNewName, FText& OutErrorMessage) const;
 
 	bool HasStructOwner() const { return bStructOwner; }
 
@@ -57,6 +44,8 @@ struct FVariablesOutlinerEntryItem : FOutlinerEntryItem
 	TFieldPath<const FProperty> PropertyPath;
 
 	bool bStructOwner = false;
+
+	int32 SortValue = 0;
 };
 
 }
diff --git a/Engine/Plugins/Experimental/UAF/UAF/Source/UAFEditor/Private/Variables/Outliner/VariablesOutlinerHierarchy.cpp b/Engine/Plugins/Experimental/UAF/UAF/Source/UAFEditor/Private/Variables/Outliner/VariablesOutlinerHierarchy.cpp
--- a/Engine/Plugins/Experimental/UAF/UAF/Source/UAFEditor/Private/Variables/Outliner/VariablesOutlinerHierarchy.cpp
+++ b/Engine/Plugins/Experimental/UAF/UAF/Source/UAFEditor/Private/Variables/Outliner/VariablesOutlinerHierarchy.cpp
@@ -5,12 +5,11 @@
 #include "AnimNextRigVMAsset.h"
 #include "ISceneOutlinerMode.h"
 #include "AnimNextRigVMAssetEditorData.h"
+#include "VariablesOutlinerAssetItem.h"
 #include "VariablesOutlinerCategoryItem.h"
 #include "VariablesOutlinerStructSharedVariablesItem.h"
 #include "VariablesOutlinerMode.h"
 #include "VariablesOutlinerEntryItem.h"
-#include "Common/Outliner/OutlinerAssetItem.h"
-#include "Common/Outliner/SCommonOutliner.h"
 #include "Entries/AnimNextSharedVariablesEntry.h"
 #include "Entries/AnimNextRigVMAssetEntry.h"
 #include "Entries/AnimNextVariableEntry.h"
@@ -20,21 +19,42 @@
 namespace UE::UAF::Editor
 {
 
+TAutoConsoleVariable<bool> GOutputHierarchyDebugInformationCvar(TEXT("VariablesOutliner.OutputHierarchyDebugInformation"), false, TEXT("If enabled will output debug information about the hierarchy generation to the output log")); 
+
 FVariablesOutlinerHierarchy::FVariablesOutlinerHierarchy(ISceneOutlinerMode* Mode)
-	: FOutlinerHierarchy(Mode)
+	: ISceneOutlinerHierarchy(Mode)
+{
+}
+
+void FVariablesOutlinerHierarchy::CreateItems(TArray<FSceneOutlinerTreeItemPtr>& OutItems) const
 {
+	SVariablesOutliner* Outliner = static_cast<FVariablesOutlinerMode* const>(Mode)->GetOutliner();
+
+	WeakProcessedAssets.Reset();
+	WeakProcessedStructs.Reset();
+
+	if(Outliner->Assets.Num())
+	{
+		// Process first asset only for now
+		const int32 NumAssets = Outliner->Assets.Num();
+		for(int32 Index = 0; Index < NumAssets; Index++)
+		{
+			const TSoftObjectPtr<UAnimNextRigVMAsset>& SoftAsset = Outliner->Assets[Index];
+			TArray<const FSoftObjectPath> SharedVariableSourcePaths;
+			ProcessAsset(SoftAsset, (NumAssets - Index), OutItems, SharedVariableSourcePaths);
+		}
+	}
 }
 
 FSceneOutlinerTreeItemPtr FVariablesOutlinerHierarchy::FindOrCreateParentItem(const ISceneOutlinerTreeItem& Item, const TMap<FSceneOutlinerTreeItemID, FSceneOutlinerTreeItemPtr>& Items, bool bCreate)
 {
 	uint32 ParentHash = INDEX_NONE;
-	FSoftObjectPath OwnerPath;
-	FStringView CategoryPath;
 	
 	if (const FVariablesOutlinerEntryItem* EntryItem = Item.CastTo<FVariablesOutlinerEntryItem>())
 	{
 		const UAnimNextSharedVariablesEntry* SharedVariablesEntry = EntryItem->WeakSharedVariablesEntry.Get();
 		
+		FSoftObjectPath OwnerPath;
 		if(SharedVariablesEntry != nullptr)
 		{
 			OwnerPath = SharedVariablesEntry->GetObjectPath();
@@ -48,7 +68,7 @@ FSceneOutlinerTreeItemPtr FVariablesOutlinerHierarchy::FindOrCreateParentItem(co
 			}
 
 			UAnimNextRigVMAsset* Asset = Entry->GetTypedOuter<UAnimNextRigVMAsset>();
-			if (Asset == nullptr)
+			if (Entry == nullptr)
 			{
 				return nullptr;
 			}
@@ -65,23 +85,22 @@ FSceneOutlinerTreeItemPtr FVariablesOutlinerHierarchy::FindOrCreateParentItem(co
 		{
 			if (!VariableEntry->GetVariableCategory().IsEmpty())
 			{
-				CategoryPath = VariableEntry->GetVariableCategory();
-				CategoryHash = GetTypeHash(CategoryPath);
+				CategoryHash = GetTypeHash(VariableEntry->GetVariableCategory());
 			}
 		}
 
 		check(OwnerPath.IsValid());
-		ParentHash = CategoryHash != INDEX_NONE ? HashCombine(GetTypeHash(OwnerPath), CategoryHash) : GetTypeHash(OwnerPath);
+		ParentHash = CategoryHash != INDEX_NONE ? HashCombine(GetTypeHash(OwnerPath), CategoryHash) : GetTypeHash(OwnerPath);		
+		
 	}
 
 	if (const FVariablesOutlinerCategoryItem* CategoryItem = Item.CastTo<FVariablesOutlinerCategoryItem>())
 	{
 		if (UAnimNextRigVMAsset* Asset = CategoryItem->WeakOwner.Get())
 		{
-			OwnerPath = Asset;
+			const FSoftObjectPath OwnerPath = Asset;
 			const uint32 AssetHash = GetTypeHash(OwnerPath);
 			const FStringView ParentCategory = CategoryItem->ParentCategoryName;
-			CategoryPath = ParentCategory;
 			ParentHash = ParentCategory.Len() ? HashCombine(AssetHash, GetTypeHash(ParentCategory)) : AssetHash;
 		}
 	}
@@ -91,59 +110,36 @@ FSceneOutlinerTreeItemPtr FVariablesOutlinerHierarchy::FindOrCreateParentItem(co
 		return *ParentItem;
 	}
 
-	TSoftObjectPtr<UAnimNextRigVMAsset> SoftAsset = Cast<UAnimNextRigVMAsset>(OwnerPath.ResolveObject());
+	return nullptr;
+}
 
-	if (!CategoryPath.IsEmpty())
+void FVariablesOutlinerHierarchy::ProcessAsset(const TSoftObjectPtr<UAnimNextRigVMAsset>& InSoftAsset, int32 Depth, TArray<FSceneOutlinerTreeItemPtr>& OutItems, TArray<const FSoftObjectPath>& OutSharedVariableSourcePaths) const
+{
+	UAnimNextRigVMAsset* Asset = InSoftAsset.Get();
+	if(Asset == nullptr)
 	{
-		UAnimNextRigVMAssetEditorData* EditorData = UncookedOnly::FUtils::GetEditorData(SoftAsset.Get());
-		const int32 CategoryIndex = EditorData->VariableAndFunctionCategories.IndexOfByKey(CategoryPath);
-		ensure(CategoryIndex != INDEX_NONE);
-		
-		FStringView ParentCategoryName;
-		FStringView SubCategoryName;
-
-		int32 Index = INDEX_NONE;
-		if (CategoryPath.FindLastChar('|', Index))
-		{
-			ParentCategoryName = CategoryPath.Left(Index);
-			SubCategoryName = CategoryPath.RightChop(Index+1);
-		}
-		else
-		{
-			SubCategoryName = CategoryPath;
-		}
-		
-		return Mode->CreateItemFor<FVariablesOutlinerCategoryItem>(FOutlinerCategoryItem::FItemData{UncookedOnly::FUtils::GetAsset<UAnimNextRigVMAsset>(EditorData), CategoryIndex, SubCategoryName, ParentCategoryName, CategoryPath}, bCreate);
+		return;
 	}
-	else if (OwnerPath.IsValid())
-	{
-		TArray<const FSoftObjectPath> SharedVariableSourcePaths;
 		
-		const UScriptStruct* Struct = Cast<UScriptStruct>(OwnerPath.ResolveObject());
+	const UAnimNextRigVMAssetEditorData* EditorData = UncookedOnly::FUtils::GetEditorData<UAnimNextRigVMAssetEditorData>(Asset);
+	if (EditorData == nullptr)
+	{
+		return;
+	}
 
-		if (SoftAsset.IsValid())
-		{
-			RetrieveSharedAssetPaths(SoftAsset, SharedVariableSourcePaths);
-			const FSceneOutlinerTreeItemPtr AssetItem = Mode->CreateItemFor<FOutlinerAssetItem>(FOutlinerAssetItem::FItemData{SoftAsset, 0, SharedVariableSourcePaths}, bCreate);
-			
-			UE_CLOG(GOutputHierarchyDebugInformationCvar.GetValueOnAnyThread(), LogAnimation, Warning, TEXT("* FOutlinerAssetItem: %s %i"), *SoftAsset.GetAssetName(), 0);
-			return AssetItem;
-		}
-		else if (Struct)
-		{
-			return Mode->CreateItemFor<FVariablesOutlinerStructSharedVariablesItem>(Struct, bCreate);
-		}
+	if (WeakProcessedAssets.Contains(Asset))
+	{
+		return;
 	}
 
-	return nullptr;
-}
+	WeakProcessedAssets.Add(Asset);
 
-void FVariablesOutlinerHierarchy::CreateItemsForAsset(const UAnimNextRigVMAssetEditorData* EditorData, TArray<FSceneOutlinerTreeItemPtr>& OutItems) const
-{
+	UE_CLOG(GOutputHierarchyDebugInformationCvar.GetValueOnAnyThread(),LogAnimation, Warning, TEXT("ProcessAsset: %s %i"), *InSoftAsset.GetAssetName(), Depth);
+		
 	TArray<FSceneOutlinerTreeItemPtr> ParentItems;
-	EditorData->ForEachEntryOfType<UAnimNextSharedVariablesEntry>([this, /*&OutSharedVariableSourcePaths,*/ &ParentItems, /*Depth,*/ &OutItems](UAnimNextSharedVariablesEntry* InSharedVariablesEntry)
+	EditorData->ForEachEntryOfType<UAnimNextSharedVariablesEntry>([this, &OutSharedVariableSourcePaths, &ParentItems, Depth, &OutItems](UAnimNextSharedVariablesEntry* InSharedVariablesEntry)
 	{
-		SCommonOutliner* Outliner = static_cast<FCommonOutlinerMode* const>(Mode)->GetOutliner();
+		SVariablesOutliner* Outliner = static_cast<FVariablesOutlinerMode* const>(Mode)->GetOutliner();
 
 		switch (InSharedVariablesEntry->Type)
 		{
@@ -151,11 +147,13 @@ void FVariablesOutlinerHierarchy::CreateItemsForAsset(const UAnimNextRigVMAssetE
 			{
 				if (const UAnimNextSharedVariables* SharedVariables = InSharedVariablesEntry->GetAsset())
 				{
+					OutSharedVariableSourcePaths.Add(SharedVariables);
+				
 					// Would have expected shared variable (assets) to be part of the outer export chain, so should have been handled as regular assets.
-					if (!Outliner->GetAssetsView().Contains(SharedVariables))
+					if (!Outliner->Assets.Contains(SharedVariables))
 					{	
 						TArray<const FSoftObjectPath> SharedVariableSourcePaths;
-						ProcessAsset( const_cast<UAnimNextSharedVariables*>(SharedVariables), 0, ParentItems, SharedVariableSourcePaths);
+						ProcessAsset( const_cast<UAnimNextSharedVariables*>(SharedVariables), Depth - 1, ParentItems, SharedVariableSourcePaths);
 					}
 				}				
 			}
@@ -172,22 +170,26 @@ void FVariablesOutlinerHierarchy::CreateItemsForAsset(const UAnimNextRigVMAssetE
 
 					WeakProcessedStructs.Add(Struct);
 
-					if (FSceneOutlinerTreeItemPtr Item = Mode->CreateItemFor<FVariablesOutlinerStructSharedVariablesItem>(Struct))
+					if (FSceneOutlinerTreeItemPtr Item = Mode->CreateItemFor<FVariablesOutlinerStructSharedVariablesItem>(InSharedVariablesEntry))
 					{
 						OutItems.Add(Item);
 					}
 					
-					UE_CLOG(GOutputHierarchyDebugInformationCvar.GetValueOnAnyThread(), LogAnimation, Warning, TEXT("~ FVariablesOutlinerStructSharedVariablesItem: %s %i"), *InSharedVariablesEntry->GetObjectPath().ToString(), 0);
+					UE_CLOG(GOutputHierarchyDebugInformationCvar.GetValueOnAnyThread(), LogAnimation, Warning, TEXT("~ FVariablesOutlinerStructSharedVariablesItem: %s %i"), *InSharedVariablesEntry->GetObjectPath().ToString(), Depth);
 					
+					OutSharedVariableSourcePaths.Add(Struct);
 					int32 Index = 0;
 					for (TFieldIterator<FProperty> It(Struct); It; ++It)
 					{
 						const FProperty* Property = *It;
-						if (FSceneOutlinerTreeItemPtr Item = Mode->CreateItemFor<FVariablesOutlinerEntryItem>(FVariablesOutlinerEntryItem::FItemData{nullptr, Property, Index++}))
+						if (FSceneOutlinerTreeItemPtr Item = Mode->CreateItemFor<FVariablesOutlinerEntryItem>(Property))
 						{
 							FVariablesOutlinerEntryItem* EntryItem = Item->CastTo<FVariablesOutlinerEntryItem>();
 							EntryItem->WeakSharedVariablesEntry = InSharedVariablesEntry;
+							EntryItem->SortValue = Index++;
 							OutItems.Add(Item);
+
+							UE_CLOG(GOutputHierarchyDebugInformationCvar.GetValueOnAnyThread(), LogAnimation, Warning, TEXT("\t- FVariablesOutlinerEntryItem: %s %i"), *Property->GetName(), EntryItem->SortValue);
 						}
 					}
 				}
@@ -200,7 +202,8 @@ void FVariablesOutlinerHierarchy::CreateItemsForAsset(const UAnimNextRigVMAssetE
 
 	OutItems.Append(ParentItems);
 
-	TSet<FString> ParsedCategories;
+	TMap<FString, TSharedPtr<FVariablesOutlinerCategoryItem>> CategoryItems;
+
 	for (int32 CategoryIndex = 0; CategoryIndex < EditorData->VariableAndFunctionCategories.Num(); ++CategoryIndex)
 	{
 		const FString& CategoryName = EditorData->VariableAndFunctionCategories[CategoryIndex];
@@ -220,12 +223,12 @@ void FVariablesOutlinerHierarchy::CreateItemsForAsset(const UAnimNextRigVMAssetE
 			
 			const FString FullCategoryName = RebuildCategories;
 
-			if (!ParsedCategories.Contains(FullCategoryName))
+			if (!CategoryItems.Contains(FullCategoryName))
 			{
-				if (TSharedPtr<FVariablesOutlinerCategoryItem> Item = StaticCastSharedPtr<FVariablesOutlinerCategoryItem>(Mode->CreateItemFor<FVariablesOutlinerCategoryItem>(FVariablesOutlinerCategoryItem::FItemData{UncookedOnly::FUtils::GetAsset<UAnimNextRigVMAsset>(EditorData), CategoryIndex, SubCategoryName, ParentCategoryName, FullCategoryName})))
+				if (TSharedPtr<FVariablesOutlinerCategoryItem> Item = StaticCastSharedPtr<FVariablesOutlinerCategoryItem>( Mode->CreateItemFor<FVariablesOutlinerCategoryItem>(FVariablesOutlinerCategoryItem::FItemData{Asset, SubCategoryName, ParentCategoryName, FullCategoryName})))
 				{
-					ParsedCategories.Add(FullCategoryName);
-					OutItems.AddUnique(Item);
+					Item->CastTo<FVariablesOutlinerCategoryItem>()->SortValue = CategoryIndex;
+					CategoryItems.Add(FullCategoryName, Item);
 
 					UE_CLOG(GOutputHierarchyDebugInformationCvar.GetValueOnAnyThread(),LogAnimation, Warning, TEXT("+ FVariablesOutlinerCategoryItem: %s  %i"), *FullCategoryName, CategoryIndex);
 				}
@@ -233,24 +236,47 @@ void FVariablesOutlinerHierarchy::CreateItemsForAsset(const UAnimNextRigVMAssetE
 		}
 	}
 
+	// Adjust asset depth to ensure SharedVariables are always pinned to 0 depth
+	const int32 AssetDepth = Asset->GetClass() && Asset->GetClass() == UAnimNextSharedVariables::StaticClass() ? 0 : Depth;
+
 	constexpr int32 NumberOfVariableSortValueEntriesInCategory = 1024;
-	
+
+	auto EnsureCategoryItemExists = [&CategoryItems, &OutItems](const FString& CategoryName)
+	{		 
+		TSharedPtr<FVariablesOutlinerCategoryItem>& CategoryItemToAdd = CategoryItems.FindChecked(CategoryName);
+		OutItems.AddUnique(CategoryItemToAdd);
+
+		auto CategoryPtr = CategoryItems.Find(CategoryItemToAdd->ParentCategoryName);
+		while (CategoryPtr)
+		{
+			OutItems.AddUnique(*CategoryPtr);
+			CategoryPtr = CategoryItems.Find((*CategoryPtr)->ParentCategoryName);
+		}
+	};
+
+	if (FSceneOutlinerTreeItemPtr AssetItem = Mode->CreateItemFor<FVariablesOutlinerAssetItem>(FVariablesOutlinerAssetItem::FItemData{InSoftAsset, AssetDepth, OutSharedVariableSourcePaths}))
 	{
+		OutItems.Add(AssetItem);
+
+		UE_CLOG(GOutputHierarchyDebugInformationCvar.GetValueOnAnyThread(), LogAnimation, Warning, TEXT("* FVariablesOutlinerAssetItem: %s %i"), *InSoftAsset.GetAssetName(), AssetDepth);
+
 		int32 Index = 1;
-		EditorData->ForEachEntryOfType<UAnimNextVariableEntry>([this, &OutItems, &Index, EditorData](UAnimNextVariableEntry* InVariable)
+		EditorData->ForEachEntryOfType<UAnimNextVariableEntry>([this, &OutItems, EnsureCategoryItemExists, &Index, EditorData](UAnimNextVariableEntry* InVariable)
 		{
-			const int32 CategoryIndex = EditorData->VariableAndFunctionCategories.IndexOfByKey(InVariable->Category);
-			if (FSceneOutlinerTreeItemPtr Item = Mode->CreateItemFor<FVariablesOutlinerEntryItem>(
-				FVariablesOutlinerEntryItem::FItemData
-				{
-					CastChecked<UAnimNextVariableEntry>(InVariable),
-					nullptr,
-					(CategoryIndex != INDEX_NONE ? ((CategoryIndex + 1) * NumberOfVariableSortValueEntriesInCategory) : NumberOfVariableSortValueEntriesInCategory) + Index++
-				}))
+			if (FSceneOutlinerTreeItemPtr Item = Mode->CreateItemFor<FVariablesOutlinerEntryItem>(CastChecked<UAnimNextVariableEntry>(InVariable)))
 			{
+				FVariablesOutlinerEntryItem* EntryItem = Item->CastTo<FVariablesOutlinerEntryItem>();
+				const int32 CategoryIndex = EditorData->VariableAndFunctionCategories.IndexOfByKey(InVariable->Category);
+				EntryItem->SortValue = (CategoryIndex != INDEX_NONE ? ((CategoryIndex + 1) * NumberOfVariableSortValueEntriesInCategory) : 0) + Index++;
 				OutItems.Add(Item);
 				
 				UE_CLOG(GOutputHierarchyDebugInformationCvar.GetValueOnAnyThread(), LogAnimation, Warning, TEXT("\t- FVariablesOutlinerEntryItem: %s %i %s"), *InVariable->GetEntryName().ToString(), Item->CastTo<FVariablesOutlinerEntryItem>()->SortValue, InVariable->GetVariableCategory().GetData());
+
+				// Add category item on demand (to ensure we do not end up with empty category item hierarchies)
+				if (!InVariable->Category.IsEmpty())
+				{
+					EnsureCategoryItemExists(InVariable->Category);
+				}
 			}
 			return true;
 		});
diff --git a/Engine/Plugins/Experimental/UAF/UAF/Source/UAFEditor/Private/Variables/Outliner/VariablesOutlinerHierarchy.h b/Engine/Plugins/Experimental/UAF/UAF/Source/UAFEditor/Private/Variables/Outliner/VariablesOutlinerHierarchy.h
--- a/Engine/Plugins/Experimental/UAF/UAF/Source/UAFEditor/Private/Variables/Outliner/VariablesOutlinerHierarchy.h
+++ b/Engine/Plugins/Experimental/UAF/UAF/Source/UAFEditor/Private/Variables/Outliner/VariablesOutlinerHierarchy.h
@@ -4,7 +4,6 @@
 
 #include "CoreMinimal.h"
 #include "ISceneOutlinerHierarchy.h"
-#include "Common/Outliner/OutlinerHierarchy.h"
 
 class UAnimNextRigVMAssetEditorData;
 class UAnimNextRigVMAsset;
@@ -12,7 +11,7 @@ class UAnimNextRigVMAsset;
 namespace UE::UAF::Editor
 {
 
-class FVariablesOutlinerHierarchy : public FOutlinerHierarchy
+class FVariablesOutlinerHierarchy : public ISceneOutlinerHierarchy
 {
 public:
 	FVariablesOutlinerHierarchy(ISceneOutlinerMode* Mode);
@@ -20,10 +19,15 @@ public:
 	FVariablesOutlinerHierarchy& operator=(const FVariablesOutlinerHierarchy&) = delete;
 
 	// Begin ISceneOutlinerHierarchy overrides
+	virtual void CreateItems(TArray<FSceneOutlinerTreeItemPtr>& OutItems) const override;
+	virtual void CreateChildren(const FSceneOutlinerTreeItemPtr& Item, TArray<FSceneOutlinerTreeItemPtr>& OutChildren) const override {}
 	virtual FSceneOutlinerTreeItemPtr FindOrCreateParentItem(const ISceneOutlinerTreeItem& Item, const TMap<FSceneOutlinerTreeItemID, FSceneOutlinerTreeItemPtr>& Items, bool bCreate = false) override;
 	// End ISceneOutlinerHierarchy overrides
+protected:
+	void ProcessAsset(const TSoftObjectPtr<UAnimNextRigVMAsset>& InSoftAsset, int32 Depth, TArray<FSceneOutlinerTreeItemPtr>& OutItems, TArray<const FSoftObjectPath>& OutSharedVariableSourcePaths) const;
 
-	virtual void CreateItemsForAsset(const UAnimNextRigVMAssetEditorData* EditorData, TArray<FSceneOutlinerTreeItemPtr>& OutItems) const override;
+	mutable TSet<TWeakObjectPtr<const UAnimNextRigVMAsset>> WeakProcessedAssets;
+	mutable TSet<TWeakObjectPtr<const UScriptStruct>> WeakProcessedStructs;
 };
 
 }
\ No newline at end of file
diff --git a/Engine/Plugins/Experimental/UAF/UAF/Source/UAFEditor/Private/Variables/Outliner/VariablesOutlinerMode.cpp b/Engine/Plugins/Experimental/UAF/UAF/Source/UAFEditor/Private/Variables/Outliner/VariablesOutlinerMode.cpp
--- a/Engine/Plugins/Experimental/UAF/UAF/Source/UAFEditor/Private/Variables/Outliner/VariablesOutlinerMode.cpp
+++ b/Engine/Plugins/Experimental/UAF/UAF/Source/UAFEditor/Private/Variables/Outliner/VariablesOutlinerMode.cpp
@@ -16,20 +16,20 @@
 #include "ToolMenus.h"
 #include "Framework/Commands/GenericCommands.h"
 #include "Entries/AnimNextRigVMAssetEntry.h"
+#include "Variables/AnimNextVariableItemMenuContext.h"
 #include "AnimNextRigVMAsset.h"
 #include "ContentBrowserModule.h"
 #include "FileHelpers.h"
 #include "IWorkspaceEditorModule.h"
 #include "StructViewerFilter.h"
 #include "StructViewerModule.h"
+#include "VariablesOutlinerAssetItem.h"
 #include "VariablesOutlinerCommands.h"
 #include "VariablesOutlinerStructSharedVariablesItem.h"
 #include "VariablesOutlinerCategoryItem.h"
-#include "Common/Outliner/OutlinerDragDropOps.h"
+#include "VariablesOutlinerDragDrop.h"
 #include "Common/AnimNextAssetFindReplaceVariables.h"
 #include "Common/GraphEditorSchemaActions.h"
-#include "Common/Outliner/OutlinerAssetItem.h"
-#include "Common/Outliner/OutlinerCommands.h"
 #include "Entries/AnimNextVariableEntry.h"
 #include "Variables/SVariablesView.h"
 #include "Entries/AnimNextSharedVariablesEntry.h"
@@ -39,658 +39,216 @@
 #include "Variables/AnimNextVariableSettings.h"
 #include "Widgets/Docking/SDockTab.h"
 #include "AnimNextScopedCompileJob.h"
-#include "Framework/Notifications/NotificationManager.h"
-#include "HAL/PlatformApplicationMisc.h"
-#include "Common/Outliner/OutlinerItemMenuContexts.h"
-#include "Common/Outliner/SCommonOutliner.h"
-#include "Widgets/Notifications/SNotificationList.h"
 
 #define LOCTEXT_NAMESPACE "VariablesOutlinerMode"
 
 namespace UE::UAF::Editor
 {
-void FVariablesOutlinerMode::RegisterToolMenus()
+
+FVariablesOutlinerMode::FVariablesOutlinerMode(SVariablesOutliner* InVariablesOutliner, const TSharedRef<UE::Workspace::IWorkspaceEditor>& InWorkspaceEditor)
+	: ISceneOutlinerMode(InVariablesOutliner)
+	, WeakWorkspaceEditor(InWorkspaceEditor)
 {
-	FCommonOutlinerMode::RegisterToolMenus();
-	
+	CommandList = MakeShared<FUICommandList>();
+
+	static const FName MenuName("VariablesOutliner.AddVariablesMenu");
 	UToolMenus* ToolMenus = UToolMenus::Get();
+	if (!ToolMenus->IsMenuRegistered(MenuName))
 	{
 		FToolMenuOwnerScoped ToolMenuOwnerScope(this);
-		
-		if (UToolMenu* Menu = ToolMenus->ExtendMenu(ContextToolMenuName))
-		{
-			FToolMenuSection& AssetSection = Menu->AddDynamicSection("DynamicAssetSection", FNewToolMenuDelegate::CreateLambda([](UToolMenu* InMenu)
-			{
-				const UCommonOutlinerItemMenuContext* MenuContext = InMenu->FindContext<UCommonOutlinerItemMenuContext>();			
-				const UVariablesOutlinerItemMenuContext* VariablesContext = InMenu->FindContext<UVariablesOutlinerItemMenuContext>();
-				if (MenuContext == nullptr || VariablesContext == nullptr)
-				{
-					return;
-				}
-				
-				TSharedPtr<SCommonOutliner> SharedOutliner = MenuContext->WeakOutliner.Pin();
-				if (SharedOutliner.IsValid() && MenuContext->WeakEditorDatas.Num() == 1)
-				{
-					FToolMenuSection& AssetsSection = InMenu->AddSection(ContextToolAssetSectionName, LOCTEXT("AssetsSectionLabel", "Assets"));
-				
-					const FVariablesOutlinerMode* Outliner = static_cast<const FVariablesOutlinerMode*>(SharedOutliner->GetMode());
-					AssetsSection.AddMenuEntry("AddNewCategory",  LOCTEXT("AddCategoryLabel", "Add New Category"), LOCTEXT("AddCategoryTooltip", "Adds a new Category to the selected Asset"), FSlateIcon(FAppStyle::Get().GetStyleSetName(), "Icons.PlusCircle"), FUIAction(FExecuteAction::CreateRaw(Outliner, &FVariablesOutlinerMode::AddCategory, MenuContext->WeakEditorDatas[0])));
-				}
-			}));
-			
-			
-			FToolMenuSection& VariablesSection = Menu->AddDynamicSection("Variables", FNewToolMenuDelegate::CreateLambda([](UToolMenu* InMenu)
-			{
-				const UAssetEditorToolkitMenuContext* EditorContext = InMenu->FindContext<UAssetEditorToolkitMenuContext>();
-				const UCommonOutlinerItemMenuContext* OutlinerContext = InMenu->FindContext<UCommonOutlinerItemMenuContext>();
-				const UVariablesOutlinerItemMenuContext* MenuContext = InMenu->FindContext<UVariablesOutlinerItemMenuContext>();
-				if(EditorContext == nullptr || OutlinerContext == nullptr || MenuContext == nullptr)
-				{
-					return;
-				}
-				
-				FToolMenuSection& VariablesSection = InMenu->AddSection("Variables", LOCTEXT("VariablesSectionLabel", "Variables"));
 
-				const int32 NumSelectedVariables = MenuContext->WeakEntries.Num();
-				/** Variable specific actions */			
-				{
-					// Populates entries for adding of different variable types/sources					
-					PopulateNewVariableToolMenuEntries(InMenu, false);
-					
-					if (NumSelectedVariables >= 1)
-					{
-						VariablesSection.AddMenuEntryWithCommandList(FVariablesOutlinerCommands::Get().CreateSharedVariablesAssets, OutlinerContext->WeakCommandList.Pin(), TAttribute<FText>(), TAttribute<FText>(), FSlateIcon(FUAFStyle::Get().GetStyleSetName(), "ClassIcon.AnimNextSharedVariables"));
-					}
-				}
-			}), FToolMenuInsert(FCommonOutlinerMode::ContextToolAssetSectionName, EToolMenuInsertType::After));
-		}
-		
-		
-		if (UToolMenu* Menu = ToolMenus->ExtendMenu(AddEntryContextMenuName))
+		if (UToolMenu* Menu = ToolMenus->RegisterMenu(MenuName))
 		{
-			Menu->AddDynamicSection("VariablesDynamic", FNewToolMenuDelegate::CreateLambda([this](UToolMenu* InMenu)
+			Menu->AddDynamicSection(TEXT("Variables"), FNewToolMenuDelegate::CreateLambda([](UToolMenu* InMenu)
 			{
-				const UCommonOutlinerItemMenuContext* OutlinerContext = InMenu->FindContext<UCommonOutlinerItemMenuContext>();
-				const UVariablesOutlinerItemMenuContext* MenuContext = InMenu->FindContext<UVariablesOutlinerItemMenuContext>();
-				if(MenuContext == nullptr)
-				{
-					return;
-				}
-				
-				TSharedPtr<SCommonOutliner> SharedOutliner = OutlinerContext->WeakOutliner.Pin();
-				if (SharedOutliner.IsValid() && OutlinerContext->WeakEditorDatas.Num() == 1)
-				{
-					FToolMenuSection& AssetSection = InMenu->AddSection(ContextToolAssetSectionName, LOCTEXT("AssetSectionLabel", "Asset"));
-					{
-						const FVariablesOutlinerMode* Outliner = static_cast<const FVariablesOutlinerMode*>(SharedOutliner->GetMode());
-						AssetSection.AddMenuEntry("AddNewCategory",  LOCTEXT("AddCategoryLabel", "Add New Category"), LOCTEXT("AddCategoryTooltip", "Adds a new Category to the selected Asset"), FSlateIcon(FAppStyle::Get().GetStyleSetName(), "Icons.PlusCircle"), FUIAction(FExecuteAction::CreateRaw(Outliner, &FVariablesOutlinerMode::AddCategory, OutlinerContext->WeakEditorDatas[0])));
-					}
-				}
-			
-				PopulateNewVariableToolMenuEntries(InMenu, false);			
+				PopulateNewVariableToolMenuEntries(InMenu, false);				
 			}));
 		}
 	}
 }
 
-void FVariablesOutlinerMode::InitToolMenuContext(FToolMenuContext& InOutContext) const
+void FVariablesOutlinerMode::Rebuild()
 {
-	FCommonOutlinerMode::InitToolMenuContext(InOutContext);
-	
-	UVariablesOutlinerItemMenuContext* MenuContext = NewObject<UVariablesOutlinerItemMenuContext>();
-	TArray<FSceneOutlinerTreeItemPtr> SelectedItems = GetOutliner()->GetSelectedItems();
-	for(const FSceneOutlinerTreeItemPtr& Item : SelectedItems)
-	{
-		if (const FVariablesOutlinerEntryItem* EntryItem = Item->CastTo<FVariablesOutlinerEntryItem>())
-		{
-			UAnimNextRigVMAssetEntry* Entry = EntryItem->WeakEntry.Get();
-			if(Entry == nullptr)
-			{
-				continue;
-			}
-
-			MenuContext->WeakEntries.Add(Entry);
-		}
-	}
-	
-	InOutContext.AddObject(MenuContext);
+	Hierarchy = CreateHierarchy();
 }
 
-FVariablesOutlinerMode::FVariablesOutlinerMode(SCommonOutliner* InOutliner, const TSharedRef<Workspace::IWorkspaceEditor>& InWorkspaceEditor)
-	: FCommonOutlinerMode(InOutliner, InWorkspaceEditor)
+TSharedPtr<SWidget> FVariablesOutlinerMode::CreateContextMenu()
 {
-	static const FName MenuName("VariablesOutliner.AddVariablesMenu");
+	static const FName MenuName("VariablesOutliner.ItemContextMenu");
 	UToolMenus* ToolMenus = UToolMenus::Get();
 	if (!ToolMenus->IsMenuRegistered(MenuName))
 	{
 		FToolMenuOwnerScoped ToolMenuOwnerScope(this);
-
 		if (UToolMenu* Menu = ToolMenus->RegisterMenu(MenuName))
 		{
-			Menu->AddDynamicSection(TEXT("Variables"), FNewToolMenuDelegate::CreateLambda([](UToolMenu* InMenu)
+			Menu->AddDynamicSection(TEXT("Assets"), FNewToolMenuDelegate::CreateLambda([](UToolMenu* InMenu)
 			{
-				PopulateNewVariableToolMenuEntries(InMenu, false);				
-			}));
-		}
-	}
-	
-	RegisterDragDropOperationConstructor<FVariablesOutlinerEntryItem>([](const FSceneOutlinerTreeItemPtr& InItem)-> TSharedPtr<FDragDropOperation>
-	{
-		FVariablesOutlinerEntryItem* VariableItem = InItem->CastTo<FVariablesOutlinerEntryItem>();
-		check(VariableItem);
-		
-		if(UAnimNextVariableEntry* Entry = Cast<UAnimNextVariableEntry>(VariableItem->WeakEntry.Get()))
-		{
-			const UAnimNextSharedVariables* Asset = nullptr;
-			if (VariableItem->WeakSharedVariablesEntry.Get())
-			{
-				Asset = VariableItem->WeakSharedVariablesEntry->GetAsset();
-			}
-			else
-			{
-				Asset = Entry->GetTypedOuter<UAnimNextSharedVariables>();
-			}
-	
-			if (Asset)
-			{
-				TSharedPtr<FAnimNextSchemaAction_Variable> Action = MakeShared<FAnimNextSchemaAction_Variable>(Entry->GetVariableName(), Asset, Entry->GetType(), FAnimNextSchemaAction_Variable::EVariableAccessorChoice::Deferred);
-
-				return FVariableDragDropOp::New(Action, StaticCastSharedRef<FVariablesOutlinerEntryItem>(VariableItem->AsShared()));
-			}
-		}
-		else if (const FProperty* Property = VariableItem->PropertyPath.Get())
-		{
-			if (const UScriptStruct* Struct = Property->GetOwner<UScriptStruct>())
-			{
-				TSharedPtr<FAnimNextSchemaAction_Variable> Action = MakeShared<FAnimNextSchemaAction_Variable>(Property->GetFName(), Struct, FAnimNextParamType::FromProperty(Property), FAnimNextSchemaAction_Variable::EVariableAccessorChoice::Deferred);
-			
-				TSharedPtr<FVariableDragDropOp> Op = FVariableDragDropOp::New(Action, StaticCastSharedRef<FVariablesOutlinerEntryItem>(VariableItem->AsShared()));
-				return Op;
-			}
-		}
-		
-		return nullptr;
-	});
-	
-	RegisterDragDropOperationParser<FVariableDragDropOp>([](FSceneOutlinerDragDropPayload& OutPayload, const FDragDropOperation& InDragDropOperation)
-	{
-		check(InDragDropOperation.IsOfType<FVariableDragDropOp>());
-		const FVariableDragDropOp& VariableDragDropOp = static_cast<const FVariableDragDropOp&>(InDragDropOperation);
-
-		OutPayload.DraggedItems.Add(VariableDragDropOp.WeakItem.Pin());
-	});
-	
-	RegisterDragDropOperator<FVariablesOutlinerEntryItem, FOutlinerCategoryItem, FVariableDragDropOp>([](const ISceneOutlinerTreeItem& DragOverTreeItem, TSharedRef<FDragDropOperation> DragDropOperation, FSceneOutlinerDragValidationInfo& OutValidationInfo) -> FReply
-	{
-		TSharedRef<FVariableDragDropOp> VariableDropOp = StaticCastSharedRef<FVariableDragDropOp>(DragDropOperation);			
-		const FOutlinerCategoryItem* DragOverItem = DragOverTreeItem.CastTo<FOutlinerCategoryItem>();
-		
-		if (VariableDropOp->WeakItem.IsValid())
-		{
-			TSharedPtr<FVariablesOutlinerEntryItem> DropVariableItem = VariableDropOp->WeakItem.Pin();
-			
-			if (DropVariableItem->HasStructOwner())
-			{
-				const FText FormattedMessage = LOCTEXT("StructItemFeedback", "Cannot make changes to Struct Entries");						
-				OutValidationInfo = FSceneOutlinerDragValidationInfo(ESceneOutlinerDropCompatibility::IncompatibleGeneric, FormattedMessage);
-				
-				return FReply::Handled();
-			}
-			
-			if (UAnimNextVariableEntry* VariableEntry = DropVariableItem->WeakEntry.Get())
-			{
-				const FText OtherVariableNameText = FText::FromName(VariableEntry->GetVariableName());
-				const FText ThisVariableCategoryText = FText::FromString(DragOverItem->GetDisplayString());
-
-				const UAnimNextRigVMAssetEditorData* VariableEditorData = VariableEntry->GetTypedOuter<UAnimNextRigVMAssetEditorData>();
-				const UAnimNextRigVMAssetEditorData* CategoryEditorData = UncookedOnly::FUtils::GetEditorData(DragOverItem->WeakOwner.Get());
-
-				if (VariableEditorData == CategoryEditorData)
+				const UAssetEditorToolkitMenuContext* EditorContext = InMenu->FindContext<UAssetEditorToolkitMenuContext>();
+				const UAnimNextVariableItemMenuContext* MenuContext = InMenu->FindContext<UAnimNextVariableItemMenuContext>();
+				if(EditorContext == nullptr || MenuContext == nullptr)
 				{
-					if (VariableEntry->GetVariableCategory() == DragOverItem->CategoryPath)
-					{
-						const FText FormattedMessage = FText::Format(LOCTEXT("VariableAlreadyPartOfCategoryFormat", "{0} is already part of Category {1}"), OtherVariableNameText, ThisVariableCategoryText);
-						
-						OutValidationInfo = FSceneOutlinerDragValidationInfo(ESceneOutlinerDropCompatibility::Incompatible, FormattedMessage);
-					}
-					else if (!VariableEntry->GetVariableCategory().IsEmpty())
-					{
-						const FText FormattedMessage = FText::Format(LOCTEXT("MoveVariableToCategoryFormat", "Moving {0} to Category {1}"), OtherVariableNameText, ThisVariableCategoryText);
-						OutValidationInfo = FSceneOutlinerDragValidationInfo(ESceneOutlinerDropCompatibility::CompatibleGeneric, FormattedMessage);
-					}
-					else
-					{
-						const FText FormattedMessage = FText::Format(LOCTEXT("AddVariableToCategoryFormat", "Add {0} to Category {1}"), OtherVariableNameText, ThisVariableCategoryText);
-						OutValidationInfo = FSceneOutlinerDragValidationInfo(ESceneOutlinerDropCompatibility::CompatibleGeneric, FormattedMessage);
-					}
+					return;
 				}
-				else
+
+				const int32 NumSelectedVariables = MenuContext->WeakEntries.Num();
+				
+				/** Asset specific actions */
+				FToolMenuSection& AssetSection = InMenu->AddSection("Asset", LOCTEXT("AssetSectionLabel", "Asset"));
 				{
-					const FText ThisAssetText = FText::FromString(DragOverItem->GetDisplayString());
-					const FText FormattedMessage = FText::Format(LOCTEXT("MoveVariableToCategoryAndAssetFormat", "Moving {0} to Category {1} in {2}"), OtherVariableNameText, ThisVariableCategoryText, ThisAssetText);
-						OutValidationInfo = FSceneOutlinerDragValidationInfo(ESceneOutlinerDropCompatibility::CompatibleGeneric, FormattedMessage);		
+					AssetSection.AddMenuEntryWithCommandList(FVariablesOutlinerCommands::Get().SaveAsset, MenuContext->WeakCommandList.Pin(), TAttribute<FText>(), TAttribute<FText>(), FSlateIcon(FAppStyle::Get().GetStyleSetName(), "AssetEditor.SaveAsset"));
+					AssetSection.AddMenuEntryWithCommandList(FGenericCommands::Get().Paste, MenuContext->WeakCommandList.Pin());
 				}
-				
-				return FReply::Handled();
-			}
-		}
-		
-	
-		return FReply::Unhandled();
-	},
-	[this](ISceneOutlinerTreeItem& DroppedOnTreeItem, TSharedRef<const FDragDropOperation> DragDropOperation)
-	{			
-		TSharedRef<const FVariableDragDropOp> VariableDropOp = StaticCastSharedRef<const FVariableDragDropOp>(DragDropOperation);
-		const FOutlinerCategoryItem* DroppedOnCategory = DroppedOnTreeItem.CastTo<FOutlinerCategoryItem>();			
-		TSharedPtr<FVariablesOutlinerEntryItem> DropVariableItem = VariableDropOp->WeakItem.Pin();
-		
-		if (DropVariableItem.IsValid())
-		{
-			if (UAnimNextVariableEntry* VariableEntry = DropVariableItem->WeakEntry.Get())
-			{
-				UAnimNextRigVMAsset* VariableAsset = VariableEntry->GetTypedOuter<UAnimNextRigVMAsset>();
-				UAnimNextRigVMAsset* CategoryAsset = DroppedOnCategory->WeakOwner.Get();
 
-				const FText OtherVariableNameText = FText::FromName(VariableEntry->GetVariableName());
-				const FText ThisVariableCategoryText = FText::FromString(DroppedOnCategory->GetDisplayString());
-				const FText FormattedMessage = FText::Format(LOCTEXT("MoveVariableToCategoryFormat", "Moving {0} to Category {1}"), OtherVariableNameText, ThisVariableCategoryText);	
-				FScopedTransaction Transaction(FormattedMessage);
-
-				UAnimNextVariableEntry* NewVariableEntry = VariableEntry;
-					
-				if (VariableAsset != CategoryAsset)
+				/** Variable specific actions */
+				FToolMenuSection& VariablesSection = InMenu->AddSection("Variables", LOCTEXT("VariablesSectionLabel", "Variables"));
 				{
-					const FName VariableEntryName = VariableEntry->GetVariableName();
-					UncookedOnly::FUtils::MoveVariableToAsset(VariableEntry, CategoryAsset, true, true);
-					NewVariableEntry = Cast<UAnimNextVariableEntry>(UncookedOnly::FUtils::GetEditorData(CategoryAsset)->FindEntry(VariableEntryName));
-				}
-					
-				check(NewVariableEntry);
+					// Populates entries for adding of different variable types/sources					
+					PopulateNewVariableToolMenuEntries(InMenu, true);
 					
-				// Only add if not already part of category
-				if (NewVariableEntry->GetVariableCategory() != DroppedOnCategory->CategoryPath)
-				{
-					NewVariableEntry->SetVariableCategory(DroppedOnCategory->CategoryPath);
-				}
-				
-				return true;
-			}
-		}
-		
-		return false;
-	});	
-	
-	RegisterDragDropOperator<FVariablesOutlinerEntryItem, FVariablesOutlinerEntryItem, FVariableDragDropOp>([](const ISceneOutlinerTreeItem& DragOverTreeItem, TSharedRef<FDragDropOperation> DragDropOperation, FSceneOutlinerDragValidationInfo& OutValidationInfo) -> FReply
-	{
-		TSharedRef<FVariableDragDropOp> VariableDropOp = StaticCastSharedRef<FVariableDragDropOp>(DragDropOperation);			
-		const FVariablesOutlinerEntryItem* DragOverItem = DragOverTreeItem.CastTo<FVariablesOutlinerEntryItem>();			
-		TSharedPtr<FVariablesOutlinerEntryItem> DropVariableItem = VariableDropOp->WeakItem.Pin();
+					VariablesSection.AddSubMenu("FindReferencesMenu", LOCTEXT("FindReferencesMenuLabel", "Find References..."), LOCTEXT("FindReferencesMenuToolTip", "Find references to selected variables by different means."),
+					FNewToolMenuDelegate::CreateLambda(
+						[MenuContext](UToolMenu* InSubmenu)
+							{
+								InSubmenu->AddMenuEntry(NAME_None, FToolMenuEntry::InitMenuEntryWithCommandList(FVariablesOutlinerCommands::Get().FindReferences, MenuContext->WeakCommandList.Pin(), TAttribute<FText>(), TAttribute<FText>(), UAnimNextAssetFindReplaceVariables::GetIconFromSearchScope(ESearchScope::Global)));
 
-		if (DragOverItem && DropVariableItem.IsValid())
-		{
-			const UAnimNextVariableEntry* ThisEntry = DragOverItem->WeakEntry.Get();
+								InSubmenu->AddMenuEntry(NAME_None, FToolMenuEntry::InitMenuEntryWithCommandList(FVariablesOutlinerCommands::Get().FindReferencesInWorkspace, MenuContext->WeakCommandList.Pin(), TAttribute<FText>(), TAttribute<FText>(),
+									UAnimNextAssetFindReplaceVariables::GetIconFromSearchScope(ESearchScope::Workspace)));
+							
+								InSubmenu->AddMenuEntry(NAME_None, FToolMenuEntry::InitMenuEntryWithCommandList(FVariablesOutlinerCommands::Get().FindReferencesInAsset, MenuContext->WeakCommandList.Pin(), TAttribute<FText>(), TAttribute<FText>(),
+									UAnimNextAssetFindReplaceVariables::GetIconFromSearchScope(ESearchScope::Asset)));
+							}),
+							FUIAction(FExecuteAction(), FCanExecuteAction(), FGetActionCheckState(), FIsActionButtonVisible::CreateLambda([MenuContext]() -> bool
+							{
+								if (TSharedPtr<UE::UAF::Editor::SVariablesOutliner> SharedOutliner = MenuContext->WeakOutliner.Pin())
+								{
+									if (const FVariablesOutlinerMode* Mode = static_cast<const FVariablesOutlinerMode*>(SharedOutliner->GetMode()))
+									{
+										return Mode->CanFindReferences();
+									}
+								}
 
-			if (DragOverItem->HasStructOwner() || DropVariableItem->HasStructOwner())
-			{
-				const FText FormattedMessage = LOCTEXT("StructItemFeedback", "Cannot make changes to Struct Entries");						
-				OutValidationInfo = FSceneOutlinerDragValidationInfo(ESceneOutlinerDropCompatibility::IncompatibleGeneric, FormattedMessage);
-			}
-			else if (ThisEntry)
-			{
-				if (const UAnimNextVariableEntry* OtherEntry = DropVariableItem->WeakEntry.Get())
-				{
-					const FText OtherVariableNameText = FText::FromName(OtherEntry->GetVariableName());					
-					const FText OtherVariableCategoryText = FText::FromStringView(OtherEntry->GetVariableCategory());
-					const FText ThisVariableNameText = FText::FromName(ThisEntry->GetVariableName());
-					const FText ThisVariableCategoryText = FText::FromStringView(ThisEntry->GetVariableCategory());
+								return false;
+							})),
+							EUserInterfaceActionType::Button,
+							false,
+							FSlateIcon(FAppStyle::GetAppStyleSetName(), TEXT("Icons.Search"))
+						);
+
+					VariablesSection.AddMenuEntryWithCommandList(FGenericCommands::Get().Rename, MenuContext->WeakCommandList.Pin());
+					VariablesSection.AddMenuEntryWithCommandList(FGenericCommands::Get().Delete, MenuContext->WeakCommandList.Pin());
+					VariablesSection.AddMenuEntryWithCommandList(FGenericCommands::Get().Duplicate, MenuContext->WeakCommandList.Pin());
 					
-					if (ThisEntry && ThisEntry != OtherEntry)
+					VariablesSection.AddMenuEntryWithCommandList(FGenericCommands::Get().Copy, MenuContext->WeakCommandList.Pin());
+
+					// Check which export access related context menu entries we need to populate
+					if (NumSelectedVariables > 0)
 					{
-						UAnimNextRigVMAsset* ThisAsset = DragOverItem->WeakEntry.Get()->GetTypedOuter<UAnimNextRigVMAsset>();
-						const UAnimNextRigVMAsset* OtherAsset = DropVariableItem->WeakEntry.Get()->GetTypedOuter<UAnimNextRigVMAsset>();
-						if (ThisAsset && ThisAsset == OtherAsset)
+						bool bHasPublicVariables = false;
+						bool bHasPrivateVariables = false;
+						for (const TWeakObjectPtr<UAnimNextRigVMAssetEntry> WeakEntry : MenuContext->WeakEntries)
 						{
-							if (ThisEntry->GetVariableCategory() != OtherEntry->GetVariableCategory())
+							if (const UAnimNextVariableEntry* VariableEntry = Cast<UAnimNextVariableEntry>(WeakEntry))
 							{
-								if (ThisEntry->GetVariableCategory().IsEmpty())
+								if (VariableEntry->GetExportAccessSpecifier() == EAnimNextExportAccessSpecifier::Public)
 								{
-									// Remove from category
-									const FText FormattedMessage = FText::Format(LOCTEXT("RemoveVariableToCategoryFormat", "Remove {0} from Category {1}"), OtherVariableNameText, OtherVariableCategoryText);
-									OutValidationInfo = FSceneOutlinerDragValidationInfo(ESceneOutlinerDropCompatibility::CompatibleGeneric, FormattedMessage);
+									bHasPublicVariables = true;
 								}
-								else
+								else if (VariableEntry->GetExportAccessSpecifier() == EAnimNextExportAccessSpecifier::Private)
 								{
-									// Add to category
-									const FText FormattedMessage = FText::Format(LOCTEXT("AddVariableToCategoryFormat", "Add {0} to Category {1}"), OtherVariableNameText, ThisVariableCategoryText);
-									OutValidationInfo = FSceneOutlinerDragValidationInfo(ESceneOutlinerDropCompatibility::CompatibleGeneric, FormattedMessage);
+									bHasPrivateVariables = true;
 								}
 							}
-							else
+
+							if (bHasPrivateVariables && bHasPublicVariables)
 							{
-								// Reorder
-								const FText FormattedMessage = FText::Format(LOCTEXT("ReorderVariableFormat", "Reorder {0} before {1}"), OtherVariableNameText, ThisVariableNameText);
-								OutValidationInfo = FSceneOutlinerDragValidationInfo(ESceneOutlinerDropCompatibility::CompatibleGeneric, FormattedMessage);
+								break;
 							}
 						}
-						else if (ThisAsset)
-						{
-							const FText VariableNameText = FText::FromName(OtherEntry->GetVariableName());
-							const FText AssetNameText = FText::FromName(ThisAsset->GetFName());
-							
-							// [TODO] should we only allow move variables "up" the asset chain? If so validate that according to the workspace export chain, that VariableOuter is a child of Asset
-							const FText FormattedMessage = FText::Format(LOCTEXT("MoveVariableToAssetFormat", "Move {0} to {1}"), VariableNameText, AssetNameText);
-							
-							OutValidationInfo = FSceneOutlinerDragValidationInfo(ESceneOutlinerDropCompatibility::CompatibleGeneric, FormattedMessage);
-						}
-					}
-					else
-					{
-						// Reorder
-						const FText FormattedMessage = LOCTEXT("ReorderSelfVariableLabel", "Cannot Reorder Variable before itself");
-						OutValidationInfo = FSceneOutlinerDragValidationInfo(ESceneOutlinerDropCompatibility::IncompatibleGeneric, FormattedMessage);
-					}
-				}					
-			}
-			
-			return FReply::Handled();
-		}
-	
-		return FReply::Unhandled();
-	},
-	[this](ISceneOutlinerTreeItem& DroppedOnTreeItem, TSharedRef<const FDragDropOperation> DragDropOperation)
-	{
-		TSharedRef<const FVariableDragDropOp> VariableDropOp = StaticCastSharedRef<const FVariableDragDropOp>(DragDropOperation);			
-		const FVariablesOutlinerEntryItem* DroppedOnItem = DroppedOnTreeItem.CastTo<FVariablesOutlinerEntryItem>();			
-		TSharedPtr<FVariablesOutlinerEntryItem> DropVariableItem = VariableDropOp->WeakItem.Pin();
 
-		if (DroppedOnItem)
-		{
-			const UAnimNextVariableEntry* ThisEntry = DroppedOnItem->WeakEntry.Get();
-			TSharedPtr<FVariablesOutlinerEntryItem> EntryItem = VariableDropOp->WeakItem.Pin();
-
-			if (ThisEntry && EntryItem.IsValid())
-			{
-				if (UAnimNextVariableEntry* OtherEntry = EntryItem->WeakEntry.Get())
-				{
-					const FText OtherVariableNameText = FText::FromName(OtherEntry->GetVariableName());					
-					const FText OtherVariableCategoryText = FText::FromStringView(OtherEntry->GetVariableCategory());
-					const FText ThisVariableNameText = FText::FromName(ThisEntry->GetVariableName());
-					const FText ThisVariableCategoryText = FText::FromStringView(ThisEntry->GetVariableCategory());
-					
-					if (ThisEntry && ThisEntry != OtherEntry)
-					{
-						UAnimNextRigVMAsset* ThisAsset = DroppedOnItem->WeakEntry.Get()->GetTypedOuter<UAnimNextRigVMAsset>();
-						const UAnimNextRigVMAsset* OtherAsset = EntryItem->WeakEntry.Get()->GetTypedOuter<UAnimNextRigVMAsset>();
-						if (ThisAsset && ThisAsset == OtherAsset)
-						{
-							UAnimNextVariableEntry* VariableEntry = EntryItem->WeakEntry.Get();
-							const UAnimNextVariableEntry* ThisVariableEntry = DroppedOnItem->WeakEntry.Get();
-								
-							if (VariableEntry && ThisVariableEntry)
-							{	
-								if (ThisEntry->GetVariableCategory() != OtherEntry->GetVariableCategory())
-								{
-									UAnimNextRigVMAssetEditorData* VariableEditorData = OtherEntry->GetTypedOuter<UAnimNextRigVMAssetEditorData>();
-									UAnimNextRigVMAssetEditorData* ThisEditorData = ThisEntry->GetTypedOuter<UAnimNextRigVMAssetEditorData>();																			
-									if (VariableEditorData && ThisEditorData)
-									{
-										if (VariableEditorData == ThisEditorData)
-										{
-											// Add or remove to category of this variable entry
-											const bool bRemovingFromCategory = ThisEntry->GetVariableCategory().IsEmpty();
-																				
-											const FText AddFormattedMessage = FText::Format(LOCTEXT("AddVariableToCategoryFormat", "Add {0} to Category {1}"), OtherVariableNameText, ThisVariableCategoryText);
-											const FText RemoveFormattedMessage = FText::Format(LOCTEXT("RemoveVariableToCategoryFormat", "Remove {0} from Category {1}"), OtherVariableNameText, OtherVariableCategoryText);
-												
-											FScopedTransaction Transaction(bRemovingFromCategory ? RemoveFormattedMessage : AddFormattedMessage);
-											OtherEntry->SetVariableCategory(ThisEntry->GetVariableCategory(), true);		
-											// Also reorder variables such that dropped item is above dropped-on item
-											VariableEditorData->ReorderVariable(VariableEntry, ThisVariableEntry);
-											
-											return true;
-										}
-									}
-								}	
-								else
-								{
-									const FText FormattedMessage = FText::Format(LOCTEXT("ReorderVariableFormat", "Reorder {0} before {1}"), OtherVariableNameText, ThisVariableNameText);
-									FScopedTransaction Transaction(FormattedMessage);
-								
-									UAnimNextRigVMAssetEditorData* VariableEditorData = VariableEntry->GetTypedOuter<UAnimNextRigVMAssetEditorData>();
-									UAnimNextRigVMAssetEditorData* ThisEditorData = ThisVariableEntry->GetTypedOuter<UAnimNextRigVMAssetEditorData>();
-									if (VariableEditorData && VariableEditorData == ThisEditorData)
-									{
-										VariableEditorData->ReorderVariable(VariableEntry, ThisVariableEntry);
-										return true;
-									}
-								}	
-							}	
-						}
-						else if (ThisAsset)
+						if (bHasPrivateVariables)
 						{
-							if (UAnimNextVariableEntry* VariableEntry = EntryItem->WeakEntry.Get())
-							{
-								const FText VariableNameText = FText::FromName(OtherEntry->GetVariableName());
-								const FText AssetNameText = FText::FromName(ThisAsset->GetFName());
-																	
-								const FText FormattedMessage = FText::Format(LOCTEXT("MoveVariableToAssetFormatTransaction", "Moving {0} to {1}"), VariableNameText, AssetNameText);
-								FScopedTransaction Transaction(FormattedMessage);
-								UncookedOnly::FUtils::MoveVariableToAsset(VariableEntry, ThisAsset, true, true);
-								
-								return true;
-							}
+							VariablesSection.AddMenuEntryWithCommandList(FVariablesOutlinerCommands::Get().SetVariableExportPublic, MenuContext->WeakCommandList.Pin(), LOCTEXT("MakePublicMenuLabel", "Expose as Public"), TAttribute<FText>(), FSlateIcon(FAppStyle::Get().GetStyleSetName(), "Level.VisibleIcon16x"));
 						}
-					}
-				}					
-			}
-		}
-	
-		return false;
-	});
-	
-	RegisterDragDropOperator<FVariablesOutlinerEntryItem, FVariablesOutlinerStructSharedVariablesItem, FVariableDragDropOp>([](const ISceneOutlinerTreeItem& DragOverTreeItem, TSharedRef<FDragDropOperation> DragDropOperation, FSceneOutlinerDragValidationInfo& OutValidationInfo) -> FReply
-	{
-		const FText FormattedMessage = LOCTEXT("StructItemFeedback", "Cannot make changes to Struct Entries");						
-		OutValidationInfo = FSceneOutlinerDragValidationInfo(ESceneOutlinerDropCompatibility::IncompatibleGeneric, FormattedMessage);			
-		return FReply::Handled();
-	},
-	[this](ISceneOutlinerTreeItem& DroppedOnTreeItem, TSharedRef<const FDragDropOperation> DragDropOperation){ return false; }
-	);
-	
-	RegisterDragDropOperator<FVariablesOutlinerEntryItem, FOutlinerAssetItem, FVariableDragDropOp>([](const ISceneOutlinerTreeItem& DragOverTreeItem, TSharedRef<FDragDropOperation> DragDropOperation, FSceneOutlinerDragValidationInfo& OutValidationInfo) -> FReply
-	{	
-		TSharedRef<FVariableDragDropOp> VariableDropOp = StaticCastSharedRef<FVariableDragDropOp>(DragDropOperation);			
-		const FOutlinerAssetItem* DragOverItem = DragOverTreeItem.CastTo<FOutlinerAssetItem>();			
-		TSharedPtr<FVariablesOutlinerEntryItem> DropVariableItem = VariableDropOp->WeakItem.Pin();
-		
-		if (DropVariableItem->HasStructOwner())
-		{
-			const FText FormattedMessage = LOCTEXT("StructItemFeedback", "Cannot make changes to Struct Entries");						
-			OutValidationInfo = FSceneOutlinerDragValidationInfo(ESceneOutlinerDropCompatibility::IncompatibleGeneric, FormattedMessage);
-			return FReply::Handled();
-		}
-		
-		if (const UAnimNextRigVMAsset* Asset = DragOverItem->SoftAsset.Get())
-		{
-			if (DropVariableItem.IsValid())
-			{
-				if (const UAnimNextVariableEntry* VariableEntry = DropVariableItem->WeakEntry.Get())
-				{
-					const FText VariableNameText = FText::FromName(VariableEntry->GetVariableName());
-					const FText AssetNameText = FText::FromName(Asset->GetFName());
-				
-					const UAnimNextRigVMAsset* VariableOuter = VariableEntry->GetTypedOuter<UAnimNextRigVMAsset>();
-					if (VariableOuter != Asset)
-					{
-						// [TODO] should we only allow move variables "up" the asset chain? If so validate that according to the workspace export chain, that VariableOuter is a child of Asset
-						const FText FormattedMessage = FText::Format(LOCTEXT("MoveVariableToAssetFormat", "Move {0} to {1}"), VariableNameText, AssetNameText);
-						OutValidationInfo = FSceneOutlinerDragValidationInfo(ESceneOutlinerDropCompatibility::CompatibleGeneric, FormattedMessage);	
-					}
-					else
-					{
-						const FStringView CategoryPath = VariableEntry->GetVariableCategory();
-						if (CategoryPath.IsEmpty())
+
+						if (bHasPublicVariables)
 						{
-							const FText FormattedMessage = FText::Format(LOCTEXT("VariableAlreadyPartOfAssetFormat", "{0} is already part of {1}"), VariableNameText, AssetNameText);								
-							OutValidationInfo = FSceneOutlinerDragValidationInfo(ESceneOutlinerDropCompatibility::Incompatible, FormattedMessage);	
+							VariablesSection.AddMenuEntryWithCommandList(FVariablesOutlinerCommands::Get().SetVariableExportPrivate, MenuContext->WeakCommandList.Pin(), LOCTEXT("MakePrivateMenuLabel", "Make Private"), TAttribute<FText>(), FSlateIcon(FAppStyle::Get().GetStyleSetName(), "Level.NotVisibleHighlightIcon16x"));
 						}
-						else
-						{
-							const FText FormattedMessage = FText::Format(LOCTEXT("RemoveVariableFromCategoryFormat", "Removing {0} from Category {1}"), VariableNameText, FText::FromStringView(CategoryPath));	
-							OutValidationInfo = FSceneOutlinerDragValidationInfo(ESceneOutlinerDropCompatibility::CompatibleGeneric, FormattedMessage);
-						}							
 					}
 					
-					return FReply::Handled();
+					AssetSection.AddMenuEntryWithCommandList(FVariablesOutlinerCommands::Get().CreateSharedVariablesAssets, MenuContext->WeakCommandList.Pin(), TAttribute<FText>(), TAttribute<FText>(), FSlateIcon(FUAFStyle::Get().GetStyleSetName(), "ClassIcon.AnimNextSharedVariables"));
 				}
-			}
-		}
-		
-		return FReply::Unhandled();
-	},
-	[this](ISceneOutlinerTreeItem& DroppedOnTreeItem, TSharedRef<const FDragDropOperation> DragDropOperation)
-	{
-		if (DragDropOperation->IsOfType<FVariableDragDropOp>() && DroppedOnTreeItem.IsA<FOutlinerAssetItem>())
-		{
-			TSharedRef<const FVariableDragDropOp> VariableDropOp = StaticCastSharedRef<const FVariableDragDropOp>(DragDropOperation);			
-			const FOutlinerAssetItem* DroppedOnItem = DroppedOnTreeItem.CastTo<FOutlinerAssetItem>();			
-			TSharedPtr<FVariablesOutlinerEntryItem> DropVariableItem = VariableDropOp->WeakItem.Pin();
-			
-			if (UAnimNextRigVMAsset* Asset = DroppedOnItem->SoftAsset.Get())
-			{
-				if (DropVariableItem.IsValid())
+
+				if (TSharedPtr<SVariablesOutliner> Outliner = MenuContext->WeakOutliner.Pin())
 				{
-					if (UAnimNextVariableEntry* VariableEntry = DropVariableItem->WeakEntry.Get())
-					{
-						const UAnimNextRigVMAsset* VariableOuter = VariableEntry->GetTypedOuter<UAnimNextRigVMAsset>();
-						const FText VariableNameText = FText::FromName(VariableEntry->GetVariableName());
-						
-						if (VariableOuter != Asset)
-						{								
-							const FText AssetNameText = FText::FromName(Asset->GetFName());						
-							const FText FormattedMessage = FText::Format(LOCTEXT("MoveVariableToAssetFormatTransaction", "Moving {0} to {1}"), VariableNameText, AssetNameText);
-							FScopedTransaction Transaction(FormattedMessage);
-							UncookedOnly::FUtils::MoveVariableToAsset(VariableEntry, Asset, true, true);
-							
-							return true;
-						}
-						else
-						{
-							const FStringView CategoryPath = VariableEntry->GetVariableCategory();
-							if (!CategoryPath.IsEmpty())
-							{
-								const FText FormattedMessage = FText::Format(LOCTEXT("RemoveVariableFromCategoryFormat", "Removing {0} from Category {1}"), VariableNameText, FText::FromStringView(CategoryPath));	
-								
-								FScopedTransaction Transaction(FormattedMessage);
-								VariableEntry->SetVariableCategory(TEXT(""));
-								
-								return true;
-							}							
-						}
-					}
+					Outliner->AddSourceControlMenuOptions(InMenu);
 				}
-			}
-		}
-		
-		return false;
-	});
-}
-
-bool FVariablesOutlinerMode::UpdateOperationDecorator(const FDragDropEvent& Event, const FSceneOutlinerDragValidationInfo* ValidationInfo) const
-{
-	TSharedPtr<FVariableDragDropOp> GraphDropOp = Event.GetOperationAs<FVariableDragDropOp>();
-	if (GraphDropOp.IsValid())
-	{
-		if (ValidationInfo)
-		{
-			const FSlateBrush* ValidHoverIcon = FAppStyle::GetBrush(TEXT("Graph.ConnectorFeedback.OK"));
-			const FSlateBrush* InvalidSymbol = FAppStyle::GetBrush(TEXT("Graph.ConnectorFeedback.Error"));			
-			GraphDropOp->SetSimpleFeedbackMessage(ValidationInfo->IsValid() ? ValidHoverIcon : InvalidSymbol, FLinearColor::White, ValidationInfo->ValidationText);		
-		}
-		else
-		{
-			GraphDropOp->HoverTargetChanged();
-		}
-		
-		return true;
-	}
-	
-	return false;
-}
-
-void FVariablesOutlinerMode::ReorderCategoryItem(TSharedPtr<FOutlinerCategoryItem> EntryItem, TSharedPtr<FOutlinerCategoryItem> BeforeItem) const
-{
-	const UAnimNextRigVMAssetEditorData* ThisEditorData = UncookedOnly::FUtils::GetEditorData(EntryItem->WeakOwner.Get());
-	 UAnimNextRigVMAssetEditorData* EditorData = UncookedOnly::FUtils::GetEditorData(BeforeItem->WeakOwner.Get());
-	if (EditorData && ThisEditorData)
-	{
-		if (EditorData == ThisEditorData)
-		{		
-			EditorData->ReorderCategory(EntryItem->CategoryPath, BeforeItem->CategoryPath);
+			}));
 		}
 	}
-}
-
-void FVariablesOutlinerMode::ReorderCategory(const FString& CategoryPath, const FString& BeforeCategoryPath, UAnimNextRigVMAsset* Asset) const
-{
-	if (UAnimNextRigVMAssetEditorData* EditorData = UncookedOnly::FUtils::GetEditorData(Asset))
-	{
-		EditorData->ReorderCategory(CategoryPath, BeforeCategoryPath);
-	}
-}
 
-void FVariablesOutlinerMode::MoveCategoryToAsset(const FString& CategoryPath, UAnimNextRigVMAsset* FromAsset, UAnimNextRigVMAsset* ToAsset) const
-{
-	UAnimNextRigVMAssetEditorData* ThisEditorData = UncookedOnly::FUtils::GetEditorData(FromAsset);
-	UAnimNextRigVMAssetEditorData* EditorData = UncookedOnly::FUtils::GetEditorData(ToAsset);
-	
-	TArray<FString> CategoriesToMove;
-	Algo::TransformIf(ThisEditorData->VariableAndFunctionCategories, CategoriesToMove, [&CategoryPath](const FString& CategoryName)
-	{
-		return CategoryName.StartsWith(CategoryPath);
-	},
-	[](const FString& CategoryName)
-	{
-		return CategoryName;
-	});	
-	
-	TArray<UAnimNextVariableEntry*> EntriesToMove;
-	ThisEditorData->ForEachEntryOfType<UAnimNextVariableEntry>([&EntriesToMove, &CategoriesToMove](UAnimNextVariableEntry* Entry)
-	{
-		if (CategoriesToMove.Contains(Entry->GetVariableCategory()))
-		{
-			EntriesToMove.Add(Entry);
-		}
-	
-		return true;
-	});
-	
-	for (const FString& CategoryName : CategoriesToMove)
-	{	
-		EditorData->AddCategory(CategoryName);
-		ThisEditorData->RemoveCategory(CategoryName);
-	}
-	
-	if (EntriesToMove.Num())
 	{
-		for (UAnimNextVariableEntry* ToMoveEntry : EntriesToMove)
+		UAnimNextVariableItemMenuContext* MenuContext = NewObject<UAnimNextVariableItemMenuContext>();
+		MenuContext->WeakWorkspaceEditor = WeakWorkspaceEditor;
+		MenuContext->WeakOutliner = StaticCastSharedRef<SVariablesOutliner>(SceneOutliner->AsShared());
+		MenuContext->WeakCommandList = CommandList;
+		TArray<FSceneOutlinerTreeItemPtr> SelectedItems = GetOutliner()->GetSelectedItems();
+		for(const FSceneOutlinerTreeItemPtr& Item : SelectedItems)
 		{
-			if (UAnimNextVariableEntry* NewEntry = UncookedOnly::FUtils::MoveVariableToAsset(ToMoveEntry, UncookedOnly::FUtils::GetAsset(EditorData)))
+			if (const FVariablesOutlinerAssetItem* AssetItem = Item->CastTo<FVariablesOutlinerAssetItem>())
 			{
-				NewEntry->SetVariableCategory(ToMoveEntry->GetVariableCategory());
+				UAnimNextRigVMAsset* Asset = AssetItem->SoftAsset.Get();
+				if(Asset == nullptr)
+				{
+					continue;
+				}
+
+				UAnimNextRigVMAssetEditorData* EditorData = UncookedOnly::FUtils::GetEditorData<UAnimNextRigVMAssetEditorData>(Asset);
+				if(EditorData == nullptr)
+				{
+					continue;
+				}
+				
+				MenuContext->WeakEditorDatas.Add(EditorData);
+			}
+			else if (const FVariablesOutlinerEntryItem* EntryItem = Item->CastTo<FVariablesOutlinerEntryItem>())
+			{
+				UAnimNextRigVMAssetEntry* Entry = EntryItem->WeakEntry.Get();
+				if(Entry == nullptr)
+				{
+					continue;
+				}
+
+				MenuContext->WeakEntries.Add(Entry);
 			}
 		}
+
+		FToolMenuContext Context;
+		Context.AddObject(MenuContext);
+		WeakWorkspaceEditor.Pin()->InitToolMenuContext(Context);
+		return UToolMenus::Get()->GenerateWidget(MenuName, Context);
 	}
 }
 
-void FVariablesOutlinerMode::Rebuild()
+FReply FVariablesOutlinerMode::OnKeyDown(const FKeyEvent& InKeyEvent)
 {
-	Hierarchy = CreateHierarchy();
+	if (CommandList->ProcessCommandBindings(InKeyEvent))
+	{
+		return FReply::Handled();
+	}
+
+	return FReply::Unhandled();
+}
+
+void FVariablesOutlinerMode::OnItemClicked(FSceneOutlinerTreeItemPtr Item)
+{
+	const FSceneOutlinerItemSelection& Selection = SceneOutliner->GetSelection();
+	HandleItemSelection(Selection);
 }
 
 void FVariablesOutlinerMode::HandleItemSelection(const FSceneOutlinerItemSelection& Selection) const
 {
-	if (TSharedPtr<Workspace::IWorkspaceEditor> WorkspaceEditor = WeakWorkspaceEditor.Pin())
+	if (TSharedPtr<UE::Workspace::IWorkspaceEditor> WorkspaceEditor = WeakWorkspaceEditor.Pin())
 	{
 		TArray<FSceneOutlinerTreeItemPtr> SelectedItems;
 		Selection.Get(SelectedItems);
@@ -705,19 +263,36 @@ void FVariablesOutlinerMode::HandleItemSelection(const FSceneOutlinerItemSelecti
 					EntriesToShow.Add(VariableEntry);
 				}
 			}
+			else if(const FVariablesOutlinerStructSharedVariablesItem* SharedVariablesItem = Item->CastTo<FVariablesOutlinerStructSharedVariablesItem>())
+			{
+				if(const UAnimNextSharedVariablesEntry* SharedVariablesEntry = SharedVariablesItem->WeakEntry.Get())
+				{
+					EntriesToShow.Add(const_cast<UAnimNextSharedVariablesEntry*>(SharedVariablesEntry));
+				}
+			}
 		}
 
 		WorkspaceEditor->SetDetailsObjects(EntriesToShow);
 	}
 }
 
+void FVariablesOutlinerMode::OnItemSelectionChanged(FSceneOutlinerTreeItemPtr Item, ESelectInfo::Type SelectionType, const FSceneOutlinerItemSelection& Selection)
+{	
+	if (TSharedPtr<UE::Workspace::IWorkspaceEditor> SharedWorkspaceEditor = WeakWorkspaceEditor.Pin())
+	{
+		SharedWorkspaceEditor->SetGlobalSelection(SceneOutliner->AsShared(), UE::Workspace::FOnClearGlobalSelection::CreateRaw(this, &FVariablesOutlinerMode::ResetOutlinerSelection));
+	}
+	
+	HandleItemSelection(Selection);
+}
+
 bool FVariablesOutlinerMode::CanRenameItem(const ISceneOutlinerTreeItem& Item) const
 {
 	if (const FVariablesOutlinerEntryItem* EntryItem = Item.CastTo<FVariablesOutlinerEntryItem>())
 	{
 		return !EntryItem->HasStructOwner();
 	}
-	else if (Item.IsA<FOutlinerAssetItem>())
+	else if (Item.IsA<FVariablesOutlinerAssetItem>())
 	{
 		return true;
 	}
@@ -731,8 +306,95 @@ bool FVariablesOutlinerMode::CanRenameItem(const ISceneOutlinerTreeItem& Item) c
 
 void FVariablesOutlinerMode::BindCommands(const TSharedRef<FUICommandList>& OutCommandList)
 {
-	FCommonOutlinerMode::BindCommands(OutCommandList);
+	CommandList->MapAction( 
+		FGenericCommands::Get().Rename, 
+		FExecuteAction::CreateRaw(this, &FVariablesOutlinerMode::Rename),
+		FCanExecuteAction::CreateRaw(this, &FVariablesOutlinerMode::CanRename),
+		FIsActionChecked(),
+		FIsActionButtonVisible::CreateRaw(this, &FVariablesOutlinerMode::CanRename)
+	);
+
+	CommandList->MapAction(
+		FGenericCommands::Get().Delete,
+		FExecuteAction::CreateRaw(this, &FVariablesOutlinerMode::Delete),
+		FCanExecuteAction::CreateRaw(this, &FVariablesOutlinerMode::CanDelete),
+		FIsActionChecked(),
+		FIsActionButtonVisible::CreateRaw(this, &FVariablesOutlinerMode::CanDelete)	
+		);
+
+	// [TODO] implement copy/paste for variables
+	CommandList->MapAction(
+		FGenericCommands::Get().Copy, 
+		FExecuteAction::CreateRaw(this, &FVariablesOutlinerMode::Copy),
+		FCanExecuteAction::CreateRaw(this, &FVariablesOutlinerMode::CanCopy),
+		FIsActionChecked(),
+		FIsActionButtonVisible::CreateRaw(this, &FVariablesOutlinerMode::CanCopy)
+	);
+
+	CommandList->MapAction( 
+		FGenericCommands::Get().Paste, 
+		FExecuteAction::CreateRaw(this, &FVariablesOutlinerMode::Paste),
+		FCanExecuteAction::CreateRaw(this, &FVariablesOutlinerMode::CanPaste),
+		FIsActionChecked(),
+		FIsActionButtonVisible::CreateRaw(this, &FVariablesOutlinerMode::CanPaste)
+	);
+
+	CommandList->MapAction(
+		FGenericCommands::Get().Duplicate,
+		FExecuteAction::CreateRaw(this, &FVariablesOutlinerMode::Duplicate),
+		FCanExecuteAction::CreateRaw(this, &FVariablesOutlinerMode::CanDuplicate),
+		FIsActionChecked(),
+		FIsActionButtonVisible::CreateRaw(this, &FVariablesOutlinerMode::CanDuplicate)
+	);
+
+	CommandList->MapAction(
+		FVariablesOutlinerCommands::Get().FindReferences,
+		FExecuteAction::CreateRaw(this, &FVariablesOutlinerMode::FindReferences, ESearchScope::Global),
+		FCanExecuteAction::CreateRaw(this, &FVariablesOutlinerMode::CanFindReferences),
+		FIsActionChecked(),
+		FIsActionButtonVisible::CreateRaw(this, &FVariablesOutlinerMode::IsFindReferencesVisible, ESearchScope::Global)
+	);
+	
+	CommandList->MapAction(
+		FVariablesOutlinerCommands::Get().FindReferencesInWorkspace,
+		FExecuteAction::CreateRaw(this, &FVariablesOutlinerMode::FindReferences, ESearchScope::Workspace),
+		FCanExecuteAction::CreateRaw(this, &FVariablesOutlinerMode::CanFindReferences),
+		FIsActionChecked(),
+	FIsActionButtonVisible::CreateRaw(this, &FVariablesOutlinerMode::IsFindReferencesVisible, ESearchScope::Workspace)
+	);
+	
+	CommandList->MapAction(
+		FVariablesOutlinerCommands::Get().FindReferencesInAsset,
+		FExecuteAction::CreateRaw(this, &FVariablesOutlinerMode::FindReferences, ESearchScope::Asset),
+		FCanExecuteAction::CreateRaw(this, &FVariablesOutlinerMode::CanFindReferences),
+		FIsActionChecked(),
+		FIsActionButtonVisible::CreateRaw(this, &FVariablesOutlinerMode::IsFindReferencesVisible, ESearchScope::Asset)
+	);
+
+	CommandList->MapAction(
+		FVariablesOutlinerCommands::Get().SaveAsset,
+		FExecuteAction::CreateRaw(this, &FVariablesOutlinerMode::SaveAsset),
+		FCanExecuteAction::CreateRaw(this, &FVariablesOutlinerMode::CanSaveAsset),
+		FIsActionChecked(),
+		FIsActionButtonVisible::CreateRaw(this, &FVariablesOutlinerMode::IsAsset)
+		);
 	
+	CommandList->MapAction(
+		FVariablesOutlinerCommands::Get().SetVariableExportPublic,
+		FExecuteAction::CreateRaw(this, &FVariablesOutlinerMode::SetVariableExportAccessSpecifiersOnSelection, EAnimNextExportAccessSpecifier::Public),
+		FCanExecuteAction::CreateRaw(this, &FVariablesOutlinerMode::CanModifyVariableExport),
+		FIsActionChecked(),
+		FIsActionButtonVisible::CreateRaw(this, &FVariablesOutlinerMode::CanModifyVariableExport)
+	);
+
+	CommandList->MapAction(
+		FVariablesOutlinerCommands::Get().SetVariableExportPrivate,
+		FExecuteAction::CreateRaw(this, &FVariablesOutlinerMode::SetVariableExportAccessSpecifiersOnSelection, EAnimNextExportAccessSpecifier::Private),
+		FCanExecuteAction::CreateRaw(this, &FVariablesOutlinerMode::CanModifyVariableExport),
+		FIsActionChecked(),
+		FIsActionButtonVisible::CreateRaw(this, &FVariablesOutlinerMode::CanModifyVariableExport)
+	);
+
 	CommandList->MapAction(		
 		FVariablesOutlinerCommands::Get().CreateSharedVariablesAssets,
 		FExecuteAction::CreateRaw(this, &FVariablesOutlinerMode::CreateSharedVariablesAssets),
@@ -747,6 +409,16 @@ TUniquePtr<ISceneOutlinerHierarchy> FVariablesOutlinerMode::CreateHierarchy()
 	return MakeUnique<FVariablesOutlinerHierarchy>(this);
 }
 
+void FVariablesOutlinerMode::ResetOutlinerSelection() const
+{
+	SceneOutliner->ClearSelection();
+}
+
+SVariablesOutliner* FVariablesOutlinerMode::GetOutliner() const
+{
+	return static_cast<SVariablesOutliner*>(SceneOutliner);
+}
+
 void FVariablesOutlinerMode::Rename() const
 {
 	const FSceneOutlinerItemSelection& Selection = SceneOutliner->GetSelection();
@@ -777,9 +449,6 @@ void FVariablesOutlinerMode::Delete() const
 {
 	int32 NumEntries = 0;
 	TMap<UAnimNextRigVMAssetEditorData*, TArray<UAnimNextRigVMAssetEntry*>> EntriesToDeletePerAsset;
-	
-	TMap<UAnimNextRigVMAssetEditorData*, TArray<FString>> CategoryPathsToDeletePerAsset;
-	
 	TArray<FSceneOutlinerTreeItemPtr> SelectedItems = SceneOutliner->GetSelectedItems();
 	for(const FSceneOutlinerTreeItemPtr& Item : SelectedItems)
 	{
@@ -805,656 +474,152 @@ void FVariablesOutlinerMode::Delete() const
 				NumEntries++;
 			}
 		}
-		else if (const FVariablesOutlinerCategoryItem* CategoryItem = Item->CastTo<FVariablesOutlinerCategoryItem>())
-		{
-			if (UAnimNextRigVMAsset* Asset = CategoryItem->WeakOwner.Get())
-			{
-				UAnimNextRigVMAssetEditorData* EditorData = UncookedOnly::FUtils::GetEditorData(Asset);
-				CategoryPathsToDeletePerAsset.FindOrAdd(EditorData).Add(CategoryItem->CategoryPath);
-				NumEntries++;
-
-				EditorData->ForEachEntryOfType<UAnimNextVariableEntry>([CategoryPath = CategoryItem->CategoryPath, EditorData, &EntriesToDeletePerAsset, &NumEntries, &CategoryPathsToDeletePerAsset](UAnimNextVariableEntry* Entry)
-				{
-					if (Entry->GetVariableCategory().StartsWith(CategoryPath))
-					{
-						CategoryPathsToDeletePerAsset.FindOrAdd(EditorData).AddUnique(Entry->GetVariableCategory().GetData());
-						EntriesToDeletePerAsset.FindOrAdd(EditorData).Add(Entry);
-						NumEntries++;
-					}
-					
-					return true;
-				});
-			}		
-		}
-	}
-
-	if(EntriesToDeletePerAsset.Num() || CategoryPathsToDeletePerAsset.Num())
-	{
-		// [TODO] prompt user when deletion impacts assets not part of the current workspace (local context), as these have a hidden (read unknown) impact
-		FScopedTransaction Transaction(FText::Format(LOCTEXT("DeleteVariablesFormat", "Delete {0}|plural(one=item, other=items)"), NumEntries));
-		for(const TPair<UAnimNextRigVMAssetEditorData*, TArray<UAnimNextRigVMAssetEntry*>>& EntriesPair : EntriesToDeletePerAsset)
-		{
-			for (UAnimNextRigVMAssetEntry* Entry : EntriesPair.Value)
-			{
-				UncookedOnly::FUtils::DeleteVariable(CastChecked<UAnimNextVariableEntry>(Entry), true, true);
-			}
-		}
-
-		for(const TPair<UAnimNextRigVMAssetEditorData*, TArray<FString>>& CategoriesPair : CategoryPathsToDeletePerAsset)
-		{
-			for (const FString& CategoryPath : CategoriesPair.Value)
-			{
-				CategoriesPair.Key->RemoveCategory(CategoryPath);
-			}
-		}
-	}	
-}
-
-bool FVariablesOutlinerMode::CanDelete() const
-{
-	TArray<FSceneOutlinerTreeItemPtr> SelectedItems = SceneOutliner->GetSelectedItems();
-
-	bool bCanDelete = false;
-	for(const FSceneOutlinerTreeItemPtr& Item : SelectedItems)
-	{
-		if (const FVariablesOutlinerEntryItem* VariablesItem = Item->CastTo<FVariablesOutlinerEntryItem>())
-		{
-			if (!VariablesItem->HasStructOwner())
-			{
-				bCanDelete = true;
-			}
-		}
-		
-		if (const FVariablesOutlinerCategoryItem* CategoryItem = Item->CastTo<FVariablesOutlinerCategoryItem>())
-		{
-			bCanDelete = true;
-		}
-	}
-	
-	return bCanDelete;
-}
-
-bool RecursivelyFindVariableItem(const FVariablesOutlinerCategoryItem& Category)
-{	
-	const TSet<TWeakPtr<ISceneOutlinerTreeItem>>& Children = Category.GetChildren();
-	for (const TWeakPtr<ISceneOutlinerTreeItem>& ChildItem : Children)
-	{
-		if (TSharedPtr<ISceneOutlinerTreeItem> SharedChild = ChildItem.Pin())
-		{
-			if (SharedChild->IsA<FVariablesOutlinerEntryItem>())
-			{
-				return true;
-			}
-			
-			if (FVariablesOutlinerCategoryItem* ChildCategoryItem = SharedChild->CastTo<FVariablesOutlinerCategoryItem>())
-			{
-				if (RecursivelyFindVariableItem(*ChildCategoryItem))
-				{
-					return true;
-				}
-			}
-		}
-	}
-	
-	return false;
-}
-
-bool FVariablesOutlinerMode::CanCopy() const
-{
-	/**
-	 *	- Variable(s) from a graph, module of shared variables. (from multiple assets at the same time)
-	 *	- Entire category (by category selection only)
-	 *	- Variables from structs
-	 */
-	const FSceneOutlinerItemSelection& Selection = SceneOutliner->GetSelection();
-	const bool bSelectionContainsVariable = Selection.Has<FVariablesOutlinerEntryItem>();
-	const bool bSelectionContainsCategory = Selection.Has<FVariablesOutlinerCategoryItem>();	
-	const bool bSelectedCategoryContainsVariables = bSelectionContainsCategory && [&Selection]()-> bool
-	{
-		bool bContainsVariable = false;
-		Selection.ForEachItem<FVariablesOutlinerCategoryItem>([&bContainsVariable](const FVariablesOutlinerCategoryItem& Category)
-		{
-			bContainsVariable |= RecursivelyFindVariableItem(Category);
-			return !bContainsVariable;
-		});
-		
-		return bContainsVariable;
-	}();
-	
-	const bool bSelectionContainsAssetItem = Selection.Has<FOutlinerAssetItem>();
-	const bool bSelectionContainsStructItem = Selection.Has<FVariablesOutlinerStructSharedVariablesItem>();
-	
-	return !bSelectionContainsAssetItem && !bSelectionContainsStructItem && (bSelectionContainsVariable || (bSelectionContainsCategory && bSelectedCategoryContainsVariables));
-}
-
-void FVariablesOutlinerMode::Copy() const
-{
-	FVariablesOutlinerClipboardData ClipboardData;
-
-	auto CopyVariableEntry = [&ClipboardData](const UAnimNextVariableEntry* Entry, const UAnimNextRigVMAssetEditorData* EditorData)
-	{
-		FVariableClipboardData VariableData;
-		VariableData.SoftSourceObjectPath = EditorData;
-
-		VariableData.ParameterName = Entry->GetVariableName();
-		VariableData.Type = Entry->GetType();
-		if (Entry->Binding.IsValid())
-		{
-			VariableData.Binding = Entry->Binding;
-		}
-		VariableData.Comment = Entry->Comment;
-		VariableData.Category = Entry->Category;
-		VariableData.Access = Entry->GetExportAccessSpecifier();
-
-		Entry->GetDefaultValueString(VariableData.DefaultValue);
-
-		ClipboardData.Variables.AddUnique(VariableData);
-	};
-
-	auto CopyStructProperty = [&ClipboardData](const FProperty* Property)
-	{		
-		if (const UScriptStruct* Struct = Property->GetOwner<UScriptStruct>())
-		{
-			FVariableClipboardData VariableData;
-			VariableData.SoftSourceObjectPath = Struct;
-			
-			FStructOnScope DefaultStruct;
-			DefaultStruct.Initialize(Struct);
-			VariableData.ParameterName = Property->GetFName();
-			VariableData.Type = FAnimNextParamType::FromProperty(Property);			
-			VariableData.Access = EAnimNextExportAccessSpecifier::Public;
-
-			Property->ExportTextItem_InContainer(VariableData.DefaultValue, DefaultStruct.GetStructMemory(), nullptr, nullptr, 0);
-			
-			ClipboardData.Variables.AddUnique(VariableData);
-		}
-	};
-	
-	const FSceneOutlinerItemSelection& Selection = SceneOutliner->GetSelection();
-	Selection.ForEachItem<FVariablesOutlinerCategoryItem>([&CopyVariableEntry, &ClipboardData](const FVariablesOutlinerCategoryItem& SelectedItem)
-	{
-		const FString& CategoryPath = SelectedItem.CategoryPath;
-		if (UAnimNextRigVMAsset* Asset = SelectedItem.WeakOwner.Get())
-		{
-			if (UAnimNextRigVMAssetEditorData* EditorData = UncookedOnly::FUtils::GetEditorData(Asset))
-			{
-				EditorData->ForEachEntryOfType<UAnimNextVariableEntry>([&CopyVariableEntry, &CategoryPath, &EditorData](const UAnimNextVariableEntry* VariableEntry)
-				{
-					// Category is stored as set of concatenated (sub)categories, so this will cover _all_ variables
-					if (VariableEntry->GetVariableCategory().StartsWith(CategoryPath))
-					{
-						CopyVariableEntry(VariableEntry, EditorData);
-					}
-					return true;
-				});
-			}
-
-			ClipboardData.bCategorySelectedForOperation = true;
-		}
-	});
-
-	Selection.ForEachItem<FVariablesOutlinerEntryItem>([&CopyVariableEntry, &CopyStructProperty](const FVariablesOutlinerEntryItem& SelectedItem)
-	{
-		if (UAnimNextVariableEntry* Entry = SelectedItem.WeakEntry.Get())
-		{
-			if (UAnimNextRigVMAssetEditorData* EditorData = CastChecked<UAnimNextRigVMAssetEditorData>(Entry->GetOuter()))
-			{
-				CopyVariableEntry(Entry, EditorData);
-			}
-		}
-		else if (SelectedItem.bStructOwner)
-		{
-			if (const FProperty* Property = SelectedItem.PropertyPath.Get())
-			{
-				CopyStructProperty(Property);
-			}
-		}
-	});
-	
-	ensure(ClipboardData.Variables.Num());
-	{		
-		FString ClipboardDataString;
-		FVariablesOutlinerClipboardData::StaticStruct()->ExportText(ClipboardDataString, &ClipboardData, nullptr, nullptr, 0, nullptr);	
-		FPlatformApplicationMisc::ClipboardCopy(*ClipboardDataString);
-	}
-}
-
-bool FVariablesOutlinerMode::CanCut() const
-{
-	/**
-	 *	- Variable(s) from a graph, module of shared variables. (from multiple assets at the same time?) BUT not from Structs
-	 *	- Entire category (by category selection only)
-	*/
-	const FSceneOutlinerItemSelection& Selection = SceneOutliner->GetSelection();
-	const bool bSelectionContainsVariable = Selection.Has<FVariablesOutlinerEntryItem>();
-	const bool bSelectionContainsCategory = Selection.Has<FVariablesOutlinerCategoryItem>();
-	const bool bSelectedCategoryContainsVariables = bSelectionContainsCategory && [&Selection]()-> bool
-	{
-		bool bContainsVariable = false;
-		Selection.ForEachItem<FVariablesOutlinerCategoryItem>([&bContainsVariable](const FVariablesOutlinerCategoryItem& Category)
-		{
-			bContainsVariable |= RecursivelyFindVariableItem(Category);
-			return !bContainsVariable;
-		});
-	
-		return bContainsVariable;
-	}();
-	const bool bSelectionContainsAssetItem = Selection.Has<FOutlinerAssetItem>();
-
-	bool bSelectionContainsVariableFromStruct = false;
-	if (bSelectionContainsVariable)
-	{
-		TArray<FVariablesOutlinerEntryItem*> SelectedItems;
-		Selection.Get<FVariablesOutlinerEntryItem>(SelectedItems);
-
-		for (FVariablesOutlinerEntryItem* Item : SelectedItems)
-		{
-			bSelectionContainsVariableFromStruct |= Item->HasStructOwner();
-		}
-	}
-	if (bSelectionContainsCategory)
-	{
-		TArray<FVariablesOutlinerCategoryItem*> SelectedItems;
-		Selection.Get<FVariablesOutlinerCategoryItem>(SelectedItems);
-
-		for (FVariablesOutlinerCategoryItem* Item : SelectedItems)
-		{
-			bSelectionContainsVariableFromStruct |= Item->WeakOwner.Get() == nullptr;
-		}
-	}
-
-	return !bSelectionContainsAssetItem && !bSelectionContainsVariableFromStruct && (bSelectionContainsVariable || (bSelectionContainsCategory && bSelectedCategoryContainsVariables));
-}
-
-void FVariablesOutlinerMode::Cut() const
-{
-	FVariablesOutlinerClipboardData ClipboardData;
-	ClipboardData.bCutOperation = true;
-
-	auto CutVariableEntry = [&ClipboardData](UAnimNextVariableEntry* Entry)
-	{
-		UAnimNextRigVMAssetEditorData* EditorData = CastChecked<UAnimNextRigVMAssetEditorData>(Entry->GetOuter());
-		
-		FVariableClipboardData& VariableData = ClipboardData.Variables.AddDefaulted_GetRef();
-		VariableData.SoftSourceObjectPath = EditorData;
-
-		VariableData.ParameterName = Entry->GetVariableName();
-		VariableData.Type = Entry->GetType();
-		if (Entry->Binding.IsValid())
-		{
-			VariableData.Binding = Entry->Binding;
-		}
-		VariableData.Comment = Entry->Comment;
-		VariableData.Category = Entry->Category;
-		VariableData.Access = Entry->GetExportAccessSpecifier();
-		Entry->GetDefaultValueString(VariableData.DefaultValue);
-
-		EditorData->RemoveEntry(Entry);
-	};
-
-	FScopedTransaction Transaction(LOCTEXT("CutVariables", "Cut Variable(s)"));
-	TArray<UAnimNextVariableEntry*> EntriesToCut;
-	TArray<UAnimNextRigVMAssetEditorData*> EditorDatas;
-	
-	const FSceneOutlinerItemSelection& Selection = SceneOutliner->GetSelection();
-	Selection.ForEachItem<FVariablesOutlinerCategoryItem>([&CutVariableEntry, &ClipboardData, &EntriesToCut, &EditorDatas](const FVariablesOutlinerCategoryItem& SelectedItem)
-	{
-		const FString& CategoryPath = SelectedItem.CategoryPath;
-		if (UAnimNextRigVMAsset* Asset = SelectedItem.WeakOwner.Get())
+		else if (const FVariablesOutlinerStructSharedVariablesItem* SharedVariablesItem = Item->CastTo<FVariablesOutlinerStructSharedVariablesItem>())
 		{
-			if (UAnimNextRigVMAssetEditorData* EditorData = UncookedOnly::FUtils::GetEditorData(Asset))
+			UAnimNextSharedVariablesEntry* SharedVariablesEntry = SharedVariablesItem->WeakEntry.Get();
+			if(SharedVariablesEntry == nullptr)
 			{
-				EditorData->ForEachEntryOfType<UAnimNextVariableEntry>([&CutVariableEntry, &CategoryPath, &EntriesToCut](UAnimNextVariableEntry* VariableEntry)
-				{
-					// Category is stored as set of concatenated (sub)categories, so this will cover _all_ variables
-					if (VariableEntry->GetVariableCategory().StartsWith(CategoryPath))
-					{
-						EntriesToCut.Add(VariableEntry);
-					}
-					return true;
-				});
-
-				EditorDatas.AddUnique(EditorData);
-				EditorData->RemoveCategory(CategoryPath);
+				continue;
 			}
 
-			ClipboardData.bCategorySelectedForOperation = true;
-		}
-	});
-	
-	
-	TMap<UAnimNextRigVMAssetEditorData*, TArray<FString>> CategoriesToCheck;
-	Selection.ForEachItem<FVariablesOutlinerEntryItem>([&EntriesToCut, &EditorDatas, &CategoriesToCheck](const FVariablesOutlinerEntryItem& SelectedItem)
-	{
-		if (UAnimNextVariableEntry* Entry = SelectedItem.WeakEntry.Get())
-		{
-			const FStringView Category = Entry->GetVariableCategory();
-			if (!Category.IsEmpty())
+			UAnimNextRigVMAssetEditorData* EditorData = SharedVariablesEntry->GetTypedOuter<UAnimNextRigVMAssetEditorData>();
+			if(EditorData == nullptr)
 			{
-				CategoriesToCheck.FindOrAdd(UncookedOnly::FUtils::GetEditorData(SelectedItem.WeakOwner.Get())).Add(FString(Category));
+				continue;
 			}
-		
-			EntriesToCut.Add(Entry);
-			EditorDatas.AddUnique(Entry->GetTypedOuter<UAnimNextRigVMAssetEditorData>());
-		}
-	});
-
-	TArray<TUniquePtr<FRigVMControllerCompileBracketScope>> RigVMScopes;
-	for (UAnimNextRigVMAssetEditorData* EditorData : EditorDatas)
-	{
-		// Standalone SharedVariableAssets don't have a graph, and thus no controller
-		if (URigVMController* Controller = EditorData->GetOrCreateController())
-		{
-			RigVMScopes.Add(MakeUnique<FRigVMControllerCompileBracketScope>(Controller));
-		}
-	}
 
-	for (UAnimNextVariableEntry* VariableEntry : EntriesToCut)
-	{
-		CutVariableEntry(VariableEntry);
-	}
-	
-	// Handle categories for cut variables (as the might end up empty)
-	for (TPair<UAnimNextRigVMAssetEditorData*, TArray<FString>> Pair : CategoriesToCheck)
-	{	
-		TArray<FString>& Categories = Pair.Value;
-		// Remove all non-empty categories
-		Pair.Key->ForEachEntryOfType<UAnimNextVariableEntry>([&Categories](UAnimNextVariableEntry* VariableEntry)
-		{
-			Categories.RemoveAll([CategoryPath = VariableEntry->GetVariableCategory()](const FString& PathToCheck) { return PathToCheck == CategoryPath; });
-			return true;
-		});
-		
-		// Now remove all empty categories
-		for (const FString& CategoryPath : Categories)
-		{
-			Pair.Key->RemoveCategory(CategoryPath);
+			TArray<UAnimNextRigVMAssetEntry*>& EntriesToDelete = EntriesToDeletePerAsset.FindOrAdd(EditorData);
+			EntriesToDelete.Add(SharedVariablesEntry);
+			NumEntries++;
 		}
 	}
-	
-	ensure(ClipboardData.Variables.Num());
-	{
-		FString ClipboardDataString;
-		FVariablesOutlinerClipboardData::StaticStruct()->ExportText(ClipboardDataString, &ClipboardData, nullptr, nullptr, 0, nullptr);	
-		FPlatformApplicationMisc::ClipboardCopy(*ClipboardDataString);
-	}
-}
 
-void FVariablesOutlinerMode::CheckOverlappingVariableNamesRecursively(const UAnimNextRigVMAssetEditorData* EditorData, const TArray<FName>& BaseVariableNames, TMap<FName, FSoftObjectPath>& OutOverlappingVariableNames) const
-{
-	EditorData->ForEachEntryOfType<UAnimNextVariableEntry>([&BaseVariableNames, &OutOverlappingVariableNames, &EditorData](const UAnimNextVariableEntry* Entry)
+	if(NumEntries > 0)
 	{
-		if (BaseVariableNames.Contains(Entry->GetVariableName()))
-		{
-			OutOverlappingVariableNames.Add(Entry->GetVariableName(), EditorData->GetOuter());
-		}
-		
-		return true;
-	});	
-	
-	EditorData->ForEachEntryOfType<UAnimNextSharedVariablesEntry>([this, &BaseVariableNames, &OutOverlappingVariableNames](const UAnimNextSharedVariablesEntry* Entry)
-	{
-		if (const UAnimNextSharedVariables* SharedVariables = Entry->Asset.Get())
+		// [TODO] prompt user when deletion impacts assets not part of the current workspace (local context), as these have a hidden (read unknown) impact
+		FScopedTransaction Transaction(FText::Format(LOCTEXT("DeleteVariablesFormat", "Delete {0}|plural(one=variable, other=variables)"), NumEntries));
+		for(const TPair<UAnimNextRigVMAssetEditorData*, TArray<UAnimNextRigVMAssetEntry*>>& EntriesPair : EntriesToDeletePerAsset)
 		{
-			if (const UAnimNextRigVMAssetEditorData* SharedVariablesEditorData = UncookedOnly::FUtils::GetEditorData<const UAnimNextRigVMAssetEditorData>(SharedVariables))
+			for (UAnimNextRigVMAssetEntry* Entry : EntriesPair.Value)
 			{
-				CheckOverlappingVariableNamesRecursively(SharedVariablesEditorData, BaseVariableNames, OutOverlappingVariableNames);
+				UncookedOnly::FUtils::DeleteVariable(CastChecked<UAnimNextVariableEntry>(Entry), true, true);
 			}
 		}
-		
-		return true;
-	});
+	}
 }
 
-bool FVariablesOutlinerMode::CanPaste() const
+bool FVariablesOutlinerMode::CanDelete() const
 {
-	/**
-	 *	- Asset (not a struct) selected as target
-	 *	- (Sub)Category selected as target
-	 */
-	const FSceneOutlinerItemSelection& Selection = SceneOutliner->GetSelection();
-	const bool bSelectionIsCategory = Selection.Has<FVariablesOutlinerCategoryItem>();
-	const bool bSelectionIsAsset = Selection.Has<FOutlinerAssetItem>();
-
-	UAnimNextRigVMAsset* TargetAsset = nullptr;
-	if (bSelectionIsCategory)
-	{
-		TArray<FVariablesOutlinerCategoryItem*> SelectedItems;
-		Selection.Get<FVariablesOutlinerCategoryItem>(SelectedItems);
+	TArray<FSceneOutlinerTreeItemPtr> SelectedItems = SceneOutliner->GetSelectedItems();
 
-		if (SelectedItems.Num() == 1)
-		{
-			const FVariablesOutlinerCategoryItem& SelectedCategory = *SelectedItems[0];
-			TargetAsset = SelectedCategory.WeakOwner.Get();
-		}
-	}
-	else
+	bool bCanDelete = false;
+	for(const FSceneOutlinerTreeItemPtr& Item : SelectedItems)
 	{
-		TArray<FOutlinerAssetItem*> SelectedItems;
-		Selection.Get<FOutlinerAssetItem>(SelectedItems);
-
-		// This should have been enforced by CanPaste
-		if (SelectedItems.Num() == 1)
+		if (const FVariablesOutlinerEntryItem* VariablesItem = Item->CastTo<FVariablesOutlinerEntryItem>())
 		{
-			const FOutlinerAssetItem& SelectedAsset = *SelectedItems[0];
-			TargetAsset = SelectedAsset.SoftAsset.Get();
+			if (!VariablesItem->HasStructOwner())
+			{
+				bCanDelete = true;
+			}
 		}
 	}
+	
+	return bCanDelete;
+}
 
-	FVariablesOutlinerClipboardData ClipboardData;
-	if (TargetAsset)
+void FVariablesOutlinerMode::Copy() const
+{
+	const FSceneOutlinerItemSelection& Selection = SceneOutliner->GetSelection();
+	Selection.ForEachItem([](const FSceneOutlinerTreeItemPtr& SelectedItem)
 	{
-		FString ClipboardDataString;
-		FPlatformApplicationMisc::ClipboardPaste(ClipboardDataString);
-
-		FVariablesOutlinerClipboardData::StaticStruct()->ImportText(*ClipboardDataString, &ClipboardData, nullptr, 0, nullptr, FVariablesOutlinerClipboardData::StaticStruct()->GetName());
-	}
+		
+	});
+}
 
-	const bool bClipboardHasValidPasteData = ClipboardData.Variables.Num() != 0;
-	return bClipboardHasValidPasteData;
+bool FVariablesOutlinerMode::CanCopy() const
+{
+	// [TODO] implement copy behaviour
+	return false;
 }
 
 void FVariablesOutlinerMode::Paste() const
 {
-	const FSceneOutlinerItemSelection& Selection = SceneOutliner->GetSelection();
-	const bool bSelectionIsCategory = Selection.Has<FVariablesOutlinerCategoryItem>();
-	const bool bSelectionIsAsset = Selection.Has<FOutlinerAssetItem>();
-
-	const FVariablesOutlinerCategoryItem* SelectedCategoryPtr = nullptr;
-
-	UAnimNextRigVMAsset* TargetAsset = nullptr;
-	if (bSelectionIsCategory)
-	{
-		TArray<FVariablesOutlinerCategoryItem*> SelectedItems;
-		Selection.Get<FVariablesOutlinerCategoryItem>(SelectedItems);
+}
 
-		// This should have been enforced by CanPaste
-		check(SelectedItems.Num() == 1);
+bool FVariablesOutlinerMode::CanPaste() const
+{
+	// [TODO] implement paste behaviour
+	return false;
+}
 
-		SelectedCategoryPtr = SelectedItems[0];
-		TargetAsset = SelectedCategoryPtr->WeakOwner.Get();
-	}
-	else if (bSelectionIsAsset)
+void FVariablesOutlinerMode::SetVariableExportAccessSpecifiersOnSelection(EAnimNextExportAccessSpecifier VariableAccessSpecifier) const
+{
+	if (!SceneOutliner)
 	{
-		TArray<FOutlinerAssetItem*> SelectedItems;
-		Selection.Get<FOutlinerAssetItem>(SelectedItems);
-
-		// This should have been enforced by CanPaste
-		check(SelectedItems.Num() == 1);
-
-		const FOutlinerAssetItem& SelectedAsset = *SelectedItems[0];
-		TargetAsset = SelectedAsset.SoftAsset.Get();
+		return;
 	}
-
-	if (TargetAsset)
+	
+	const FSceneOutlinerItemSelection& Selection = SceneOutliner->GetSelection();
+	if (Selection.Num() > 0)
 	{
-		FString ClipboardDataString;
-		FPlatformApplicationMisc::ClipboardPaste(ClipboardDataString);
-
-		FVariablesOutlinerClipboardData ClipboardData;
-
-		FVariablesOutlinerClipboardData::StaticStruct()->ImportText(*ClipboardDataString, &ClipboardData, nullptr, 0, nullptr, FVariablesOutlinerClipboardData::StaticStruct()->GetName());
-		
-		// Verify that to-be-pasted variables don't overlap with any existing within the asset, or its upward dependencies
+		TArray<UObject*> ObjectsToCompile;
+		TArray<IAnimNextRigVMExportInterface*> ExportsToChange;
+		Selection.ForEachItem<FVariablesOutlinerEntryItem>([&ObjectsToCompile, &ExportsToChange](const FVariablesOutlinerEntryItem& SelectedItem)
 		{
-			TArray<FName> PasteVariableNames;
-			Algo::Transform(ClipboardData.Variables, PasteVariableNames, [](const FVariableClipboardData& VariablePair)
-			{
-				return VariablePair.ParameterName;
-			});
-		
-			TMap<FName, FSoftObjectPath> OverlappingVariableNames;
-			UAnimNextRigVMAssetEditorData* TargetEditorData = UncookedOnly::FUtils::GetEditorData(TargetAsset);
-			CheckOverlappingVariableNamesRecursively(TargetEditorData, PasteVariableNames, OverlappingVariableNames);
-			
-			if (OverlappingVariableNames.Num())
+			if (IAnimNextRigVMExportInterface* Export = Cast<IAnimNextRigVMExportInterface>(SelectedItem.WeakEntry.Get()))
 			{
-				FString OverlappingVariableNameInfo;
-				for (const TPair<FName, FSoftObjectPath>& Pair : OverlappingVariableNames)
+				if (UAnimNextRigVMAsset* OuterAsset = SelectedItem.WeakEntry.Get()->GetTypedOuter<UAnimNextRigVMAsset>())
 				{
-					if (OverlappingVariableNameInfo.Len())
+					if (OuterAsset->GetClass() && OuterAsset->GetClass() != UAnimNextSharedVariables::StaticClass())
 					{
-						OverlappingVariableNameInfo.Append(TEXT("\n"));
+						ObjectsToCompile.AddUnique(OuterAsset);
+						ExportsToChange.AddUnique(Export);
 					}
-					OverlappingVariableNameInfo.Append(FString::Printf(TEXT("%s: %s"), *Pair.Key.ToString(), *Pair.Value.GetAssetName()));
 				}
+			}
+		});
 
-				const FText NotificationText = FText::Format(LOCTEXT("FailedPasteOverlappingNames", "Asset or references already contain Variable Name(s):\n{0}"), FText::FromString(OverlappingVariableNameInfo));
-				
-				FNotificationInfo Notification(NotificationText);
-				Notification.ExpireDuration = 5.0f;
-				FSlateNotificationManager::Get().AddNotification(Notification);	
+		UncookedOnly::FScopedCompileJob CompilerResults(LOCTEXT("ModifiedAssets_SetAccessSpecifierOnSelection", "Modified Assets Set Access Specifier on Selection"), ObjectsToCompile);
 
-				return;
-			}
-		}
+		FScopedTransaction Transaction(LOCTEXT("SetAccessSpecifierSelectionTransaction", "Set Access Specifier on Selection"));
 
-		UAnimNextRigVMAssetEditorData* TargetEditorData = UncookedOnly::FUtils::GetEditorData(TargetAsset);
-		if (TargetEditorData && ClipboardData.Variables.Num())
+		for (IAnimNextRigVMExportInterface* Export : ExportsToChange)
 		{
-			FScopedTransaction Transaction(LOCTEXT("PasteVariables", "Paste Variable(s)"));
-			
-			// Standalone SharedVariableAssets don't have a graph, and thus no controller
-			TUniquePtr<FRigVMControllerCompileBracketScope> Scope;
-			if (URigVMController* Controller = TargetEditorData->GetOrCreateController())
-			{
-				Scope = MakeUnique<FRigVMControllerCompileBracketScope>(Controller);
-			}
-			
-			TargetEditorData->Modify();
-
-			const bool bShouldRetainCopiedCategories = ClipboardData.bCategorySelectedForOperation;
-			
-			for (FVariableClipboardData& VariableData : ClipboardData.Variables)
+			if (Export)
 			{
-				const FSoftObjectPath& SourceObjectPath = VariableData.SoftSourceObjectPath;
-				UAnimNextVariableEntry* NewVariable = TargetEditorData->AddVariable(VariableData.ParameterName, VariableData.Type, VariableData.DefaultValue);
-				NewVariable->SetExportAccessSpecifier(VariableData.Access);
-
-				// In case the user explicitly selected categories when copy/cutting
-				if (bShouldRetainCopiedCategories && !VariableData.Category.IsEmpty())
-				{
-					// Maintain the categories in case of pasting on an asset
-					if (SelectedCategoryPtr == nullptr)
-					{
-						TargetEditorData->AddCategory(VariableData.Category);
-						NewVariable->SetVariableCategory(VariableData.Category);
-					}
-					// If pasting on a category (make the pasted category a child category)
-					else
-					{
-						const FString BaseCategoryPath = SelectedCategoryPtr->CategoryPath;
-						const FString NewCategoryPath = BaseCategoryPath + TEXT("|") + VariableData.Category;
-
-						TargetEditorData->AddCategory(NewCategoryPath);
-						NewVariable->SetVariableCategory(NewCategoryPath);
-					}
-				}
-				// Otherwise, in case the user explicitly selected a category to paste in - set it as the category for all pasted variables
-				else if (bSelectionIsCategory && SelectedCategoryPtr)
-				{
-					NewVariable->SetVariableCategory(SelectedCategoryPtr->CategoryPath);
-				}
-
-				if(VariableData.Binding.IsValid())
-				{
-					NewVariable->SetBinding(MoveTemp(VariableData.Binding.BindingData));
-				}
-
-				NewVariable->Comment = VariableData.Comment;
-
-				if (ClipboardData.bCutOperation)
-				{
-					if (UAnimNextRigVMAssetEditorData* CurrentEditorData = Cast<UAnimNextRigVMAssetEditorData>(SourceObjectPath.ResolveObject()))
-					{
-						UAnimNextRigVMAssetEditorData* NewEditorData = TargetEditorData;
-						
-						const FAnimNextSoftVariableReference FindReference(VariableData.ParameterName, UncookedOnly::FUtils::GetAsset(CurrentEditorData));
-						const FAnimNextSoftVariableReference ReplaceReference(VariableData.ParameterName, UncookedOnly::FUtils::GetAsset(NewEditorData));
-
-						UncookedOnly::FUtils::ReplaceVariableReferencesAcrossProject(FindReference, ReplaceReference);
-					}
-				}
+				Export->SetExportAccessSpecifier(VariableAccessSpecifier);
 			}
 		}
 	}
 }
 
-
-void FVariablesOutlinerMode::AddCategory(TWeakObjectPtr<UAnimNextRigVMAssetEditorData> InWeakEditorData) const
+bool FVariablesOutlinerMode::CanModifyVariableExport() const
 {
-	if (UAnimNextRigVMAssetEditorData* EditorData = InWeakEditorData.Get())
+	const FSceneOutlinerItemSelection& Selection = SceneOutliner->GetSelection();
+	if (Selection.Num())
 	{
-		auto GetCategoryName = [&EditorData]() -> FString
+		bool bCanModify = false;
+		
+		FScopedTransaction Transaction(LOCTEXT("SetAccessSpecifierTransaction", "Set Access Specifier"));
+		Selection.ForEachItem<FVariablesOutlinerEntryItem>([&bCanModify](const FVariablesOutlinerEntryItem& EntryItem)
 		{
-			auto NameExists = [ExistingNames = EditorData->VariableAndFunctionCategories](FString InName)
+			if (UAnimNextVariableEntry* VariableEntry = Cast<UAnimNextVariableEntry>(EntryItem.WeakEntry.Get()))
 			{
-				for(FString AdditionalName : ExistingNames)
+				if (UAnimNextRigVMAsset* OuterAsset = VariableEntry->GetTypedOuter<UAnimNextRigVMAsset>())
 				{
-					if(AdditionalName == InName)
+					// Disable modifying of variable export for pure UAnimNextSharedVariables objects				
+					if (OuterAsset->GetClass() && OuterAsset->GetClass() != UAnimNextSharedVariables::StaticClass())
 					{
-						return true;
+						bCanModify = true;
 					}
 				}
-
-				return false;
-			};
-			
-			const FString BaseCategoryName(TEXT("Category"));
-
-			if(!NameExists(BaseCategoryName))
-			{
-				// Early out - name is valid
-				return BaseCategoryName;
-			}
-
-			int32 PostFixIndex = 0;
-			TStringBuilder<128> StringBuilder;
-			while(true)
-			{
-				StringBuilder.Reset();
-				StringBuilder.Append(BaseCategoryName);
-				StringBuilder.Appendf(TEXT("_%d"), PostFixIndex++);
-
-				const FString TestString(StringBuilder.ToString()); 
-				if(!NameExists(TestString))
-				{
-					return TestString;
-				}
 			}
-		};
+		});
 
-		FScopedTransaction Transaction(LOCTEXT("AddNewCategory", "Adding New Category"));
-		const FString NewCategoryName = GetCategoryName();
-		check(EditorData->AddCategory(NewCategoryName));
-		
-		if (SceneOutliner)
-		{
-			const FSoftObjectPath OwnerPath = EditorData->GetTypedOuter<UAnimNextRigVMAsset>();
-			const FSceneOutlinerTreeItemID CategoryID = HashCombine(GetTypeHash(OwnerPath), GetTypeHash(NewCategoryName));
-			SceneOutliner->OnItemAdded(CategoryID, SceneOutliner::ENewItemAction::Select | SceneOutliner::ENewItemAction::Rename);
-		}
+		return bCanModify;
 	}
+
+	return false;
 }
 
 void FVariablesOutlinerMode::Duplicate() const
@@ -1466,27 +631,12 @@ void FVariablesOutlinerMode::Duplicate() const
 
 		Selection.ForEachItem<FVariablesOutlinerEntryItem>([](const FVariablesOutlinerEntryItem& EntryItem)
 		{
-			if (UAnimNextVariableEntry* VariableEntry = EntryItem.WeakEntry.Get())
+			if (const UAnimNextVariableEntry* VariableEntry = EntryItem.WeakEntry.Get())
 			{
 				if (UAnimNextRigVMAssetEditorData* EditorData = VariableEntry->GetTypedOuter<UAnimNextRigVMAssetEditorData>())
 				{	
-					const FName VariableName = UncookedOnly::FUtils::GetValidVariableName(EditorData, VariableEntry->GetVariableName());
-					FString DefaultValueString;
-					VariableEntry->GetDefaultValueString(DefaultValueString);
-					if (UAnimNextVariableEntry* DuplicatedVariable = EditorData->AddVariable(VariableName, VariableEntry->GetType(), DefaultValueString))
-					{
-						if (!VariableEntry->GetVariableCategory().IsEmpty())
-						{
-							DuplicatedVariable->SetVariableCategory(VariableEntry->GetVariableCategory());
-						}
-
-						if(VariableEntry->GetBinding().IsValid())
-						{
-							TInstancedStruct<FAnimNextVariableBindingData> BindingCopy = VariableEntry->Binding.BindingData;
-							DuplicatedVariable->SetBinding(MoveTemp(BindingCopy));
-						}					
-						DuplicatedVariable->SetExportAccessSpecifier(VariableEntry->GetExportAccessSpecifier());
-					}
+					const FName VariableName = UncookedOnly::FUtils::GetValidVariableName(EditorData, VariableEntry->GetVariableName());		
+					EditorData->AddVariable(VariableName, VariableEntry->GetType());
 				}
 			}
 		});
@@ -1518,13 +668,13 @@ void FVariablesOutlinerMode::FindReferences(ESearchScope InSearchScope) const
 		FSceneOutlinerTreeItemPtr ItemToFind = Selection.SelectedItems[0].Pin();
 		if (const FVariablesOutlinerEntryItem* EntryItem = ItemToFind->CastTo<FVariablesOutlinerEntryItem>())
 		{
-			if (TSharedPtr<SDockTab> FindAndReplaceTab = WeakWorkspaceEditor.Pin()->GetTabManager()->TryInvokeTab(FTabId(UAF::Editor::FindAndReplaceTabName)))
+			if (TSharedPtr<SDockTab> FindAndReplaceTab = WeakWorkspaceEditor.Pin()->GetTabManager()->TryInvokeTab(FTabId(UE::UAF::Editor::FindAndReplaceTabName)))
 			{
 				TSharedRef<IAnimAssetFindReplace> FindAndReplace = StaticCastSharedRef<IAnimAssetFindReplace>(FindAndReplaceTab->GetContent());
 				FindAndReplace->SetCurrentProcessor(UAnimNextAssetFindReplaceVariables::StaticClass());
 				UAnimNextAssetFindReplaceVariables* AnimNextFindAndReplaceProcessor = FindAndReplace->GetProcessor<UAnimNextAssetFindReplaceVariables>();
 
-				if (TSharedPtr<Workspace::IWorkspaceEditor> SharedWorkspaceEditor = WeakWorkspaceEditor.Pin())
+				if (TSharedPtr<UE::Workspace::IWorkspaceEditor> SharedWorkspaceEditor = WeakWorkspaceEditor.Pin())
 				{
 					AnimNextFindAndReplaceProcessor->SetWorkspaceEditor(SharedWorkspaceEditor.ToSharedRef());
 				}
@@ -1559,31 +709,30 @@ bool FVariablesOutlinerMode::IsFindReferencesVisible(ESearchScope InSearchScope)
 
 void FVariablesOutlinerMode::PopulateNewVariableToolMenuEntries(UToolMenu* InMenu, bool bAddSeparator)
 {
-	const UCommonOutlinerItemMenuContext* OutlinerContext = InMenu->FindContext<UCommonOutlinerItemMenuContext>();
-	const UVariablesOutlinerItemMenuContext* MenuContext = InMenu->FindContext<UVariablesOutlinerItemMenuContext>();
-	if(OutlinerContext == nullptr || MenuContext == nullptr)
+	const UAnimNextVariableItemMenuContext* MenuContext = InMenu->FindContext<UAnimNextVariableItemMenuContext>();
+	if(MenuContext == nullptr)
 	{
 		return;
 	}
 
 	TMap<TWeakObjectPtr<UAnimNextRigVMAssetEditorData>, FString> AssetsAndCategoriesToAddTo;
-	if (OutlinerContext->WeakEditorDatas.Num())
+	if (MenuContext->WeakEditorDatas.Num())
 	{
-		if (UAnimNextRigVMAssetEditorData* EditorData = OutlinerContext->WeakEditorDatas[0].Get())
+		if (UAnimNextRigVMAssetEditorData* EditorData = MenuContext->WeakEditorDatas[0].Get())
 		{
 			AssetsAndCategoriesToAddTo.Add(EditorData, TEXT(""));
 		}
 	}
 	else
 	{
-		TSharedPtr<SCommonOutliner> SharedOutliner = OutlinerContext->WeakOutliner.Pin();
+		TSharedPtr<SVariablesOutliner> SharedOutliner = MenuContext->WeakOutliner.Pin();
 		if (SharedOutliner.IsValid())
 		{
 			TArray<FSceneOutlinerTreeItemPtr> SelectedItems = SharedOutliner->GetSelectedItems();
 		
 			for(const FSceneOutlinerTreeItemPtr& Item : SelectedItems)
 			{
-				const FOutlinerAssetItem* AssetItem = Item->CastTo<FOutlinerAssetItem>();
+				const FVariablesOutlinerAssetItem* AssetItem = Item->CastTo<FVariablesOutlinerAssetItem>();
 				const FVariablesOutlinerCategoryItem* CategoryItem = Item->CastTo<FVariablesOutlinerCategoryItem>();
 				if (AssetItem == nullptr && CategoryItem == nullptr)
 				{
@@ -1626,7 +775,7 @@ void FVariablesOutlinerMode::PopulateNewVariableToolMenuEntries(UToolMenu* InMen
 
 		FToolMenuSection& Section = InMenu->AddSection("Variables", LOCTEXT("VariablesSectionLabel", "Variables"));
 		Section.AddMenuEntry("AddNewVariable", LOCTEXT("VariableLabel", "Add Variable"), LOCTEXT("VariableTooltip", "Adds a new Variable"),FSlateIcon(FAppStyle::GetAppStyleSetName(), "Icons.PlusCircle"),
-			FUIAction(FExecuteAction::CreateLambda([WeakOutliner = OutlinerContext->WeakOutliner, AssetsAndCategoriesToAddTo]()
+			FUIAction(FExecuteAction::CreateLambda([WeakOutliner = MenuContext->WeakOutliner, AssetsAndCategoriesToAddTo]()
 			{
 				for (auto& Pair : AssetsAndCategoriesToAddTo)
 				{
@@ -1643,7 +792,7 @@ void FVariablesOutlinerMode::PopulateNewVariableToolMenuEntries(UToolMenu* InMen
 							Settings->SetLastVariableName(VariableName);
 
 							// Prompt user to rename added variable in outliner
-							TSharedPtr<SCommonOutliner> SharedOutliner = WeakOutliner.Pin();
+							TSharedPtr<SVariablesOutliner> SharedOutliner = WeakOutliner.Pin();
 							if (SharedOutliner.IsValid())
 							{
 								// FVariablesOutlinerEntryItem::SharedVariablesEntry is only populated for Struct based SharedVariables
@@ -1825,10 +974,104 @@ void FVariablesOutlinerMode::PopulateNewVariableToolMenuEntries(UToolMenu* InMen
 	}
 }
 
+void FVariablesOutlinerMode::SetHighlightedItem(FSceneOutlinerTreeItemID InID) const
+{
+	if (SVariablesOutliner* Outliner = static_cast<SVariablesOutliner*>(SceneOutliner))
+	{
+		if (FSceneOutlinerTreeItemPtr ItemPtr = Outliner->GetTreeItem(InID))
+		{
+			Outliner->FrameItem(InID);
+			Outliner->SetHighlightedItem(ItemPtr);
+		}
+	}
+}
+
+void FVariablesOutlinerMode::ClearHighlightedItem(FSceneOutlinerTreeItemID InID) const
+{
+	if (SVariablesOutliner* Outliner = static_cast<SVariablesOutliner*>(SceneOutliner))
+	{
+		if (FSceneOutlinerTreeItemPtr ItemPtr = Outliner->GetTreeItem(InID))
+		{
+			Outliner->ClearHighlightedItem(ItemPtr);
+		}
+	}
+}
+
+void FVariablesOutlinerMode::SetSelectedVariablesExportAccessSpecifiers(FSceneOutlinerTreeItemPtr ItemToChange, EAnimNextExportAccessSpecifier VariableAccessSpecifier) const
+{
+	if (!SceneOutliner || !ItemToChange.IsValid())
+	{
+		return;
+	}
+
+	const FSceneOutlinerItemSelection& Selection = SceneOutliner->GetSelection();
+	const bool bEditChangedItemOnly = !Selection.SelectedItems.Contains(ItemToChange);
+
+	FScopedTransaction Transaction(LOCTEXT("SetAccessSpecifierTransaction", "Set Access Specifier"));
+
+	if (bEditChangedItemOnly)
+	{
+		if (TSharedPtr<FVariablesOutlinerEntryItem> VariablesOutlinerItem = StaticCastSharedPtr<FVariablesOutlinerEntryItem>(ItemToChange))
+		{
+			if (IAnimNextRigVMExportInterface* Export = Cast<IAnimNextRigVMExportInterface>(VariablesOutlinerItem->WeakEntry.Get()))
+			{
+				Export->SetExportAccessSpecifier(VariableAccessSpecifier);
+			}
+		}
+	}
+	else
+	{
+		SetVariableExportAccessSpecifiersOnSelection(VariableAccessSpecifier);
+	}
+}
+
+void FVariablesOutlinerMode::ToggleSelectedVariablesExportAccessSpecifiers(FSceneOutlinerTreeItemPtr ItemToChange) const
+{
+	if (TSharedPtr<FVariablesOutlinerEntryItem> VariablesOutlinerItem = StaticCastSharedPtr<FVariablesOutlinerEntryItem>(ItemToChange))
+	{
+		if (IAnimNextRigVMExportInterface* Export = Cast<IAnimNextRigVMExportInterface>(VariablesOutlinerItem->WeakEntry.Get()))
+		{
+			SetSelectedVariablesExportAccessSpecifiers(ItemToChange, Export->GetExportAccessSpecifier() == EAnimNextExportAccessSpecifier::Public ? EAnimNextExportAccessSpecifier::Private : EAnimNextExportAccessSpecifier::Public);
+		}
+	}	
+}
+
+void FVariablesOutlinerMode::SaveAsset() const
+{
+	const FSceneOutlinerItemSelection& Selection = SceneOutliner->GetSelection();
+	Selection.ForEachItem<FVariablesOutlinerAssetItem>([](const FVariablesOutlinerAssetItem& AssetItem)
+	{
+		if (const UAnimNextRigVMAsset* Asset = AssetItem.SoftAsset.Get())
+		{
+			FEditorFileUtils::PromptForCheckoutAndSave({Asset->GetPackage()}, false, /*bPromptToSave=*/ false);
+		}
+	});	
+}
+
+bool FVariablesOutlinerMode::CanSaveAsset() const
+{
+	const FSceneOutlinerItemSelection& Selection = SceneOutliner->GetSelection();	
+	return Selection.SelectedItems.ContainsByPredicate([](const TWeakPtr<ISceneOutlinerTreeItem>& WeakItem)
+	{
+		if (WeakItem.IsValid())
+		{
+			if (const FVariablesOutlinerAssetItem* AssetItem = WeakItem.Pin()->CastTo<FVariablesOutlinerAssetItem>())
+			{
+				if (const UAnimNextRigVMAsset* Asset = AssetItem->SoftAsset.Get())
+				{
+					return Asset->GetPackage()->IsDirty();
+				}
+			}
+		}
+
+		return false;
+	});
+}
+
 bool FVariablesOutlinerMode::IsAsset() const
 {
 	const FSceneOutlinerItemSelection& Selection = SceneOutliner->GetSelection();	
-	return Selection.SelectedItems.ContainsByPredicate([](const TWeakPtr<ISceneOutlinerTreeItem>& WeakItem) { return WeakItem.IsValid() && WeakItem.Pin()->IsA<FOutlinerAssetItem>(); });	
+	return Selection.SelectedItems.ContainsByPredicate([](const TWeakPtr<ISceneOutlinerTreeItem>& WeakItem) { return WeakItem.IsValid() && WeakItem.Pin()->IsA<FVariablesOutlinerAssetItem>(); });	
 }
 
 void FVariablesOutlinerMode::CreateSharedVariablesAssets() const
@@ -1853,7 +1096,7 @@ void FVariablesOutlinerMode::CreateSharedVariablesAssets() const
 			FSaveAssetDialogConfig SaveAssetDialogConfig;
 			{
 				// Populate default path to match currently focussed assets
-				TSharedPtr<Workspace::IWorkspaceEditor> SharedWorkspaceEditor = WeakWorkspaceEditor.Pin();
+				TSharedPtr<UE::Workspace::IWorkspaceEditor> SharedWorkspaceEditor = WeakWorkspaceEditor.Pin();
 				if (SharedWorkspaceEditor.IsValid())
 				{
 					const Workspace::FWorkspaceDocument& Document =SharedWorkspaceEditor->GetFocusedWorkspaceDocument();
@@ -1952,6 +1195,110 @@ bool FVariablesOutlinerMode::CanCreateSharedVariablesAssets() const
 	return bSelectionOnlyContainsVariables && !bMoveWillLeaveEmptyAsset;
 }
 
+TSharedPtr<FDragDropOperation> FVariablesOutlinerMode::CreateDragDropOperation(const FPointerEvent& MouseEvent, const TArray<FSceneOutlinerTreeItemPtr>& InTreeItems) const
+{
+	if(InTreeItems.Num() != 1)
+	{
+		return nullptr;
+	}
+
+	const FSceneOutlinerTreeItemPtr& SelectedItem = InTreeItems[0];
+	if (FVariablesOutlinerEntryItem* VariableItem = SelectedItem->CastTo<FVariablesOutlinerEntryItem>())
+	{
+		if(UAnimNextVariableEntry* Entry = Cast<UAnimNextVariableEntry>(VariableItem->WeakEntry.Get()))
+		{
+			const UAnimNextSharedVariables* Asset = nullptr;
+			if (VariableItem->WeakSharedVariablesEntry.Get())
+			{
+				Asset = VariableItem->WeakSharedVariablesEntry->GetAsset();
+			}
+			else
+			{
+				Asset = Entry->GetTypedOuter<UAnimNextSharedVariables>();
+			}
+		
+			if (Asset)
+			{
+				TSharedPtr<FAnimNextSchemaAction_Variable> Action = MakeShared<FAnimNextSchemaAction_Variable>(Entry->GetVariableName(), Asset, Entry->GetType(), FAnimNextSchemaAction_Variable::EVariableAccessorChoice::Deferred);
+
+				return FVariableDragDropOp::New(Action, StaticCastSharedRef<FVariablesOutlinerEntryItem>(VariableItem->AsShared()));
+			}
+		}
+		else if (const FProperty* Property = VariableItem->PropertyPath.Get())
+		{
+			if (const UScriptStruct* Struct = Property->GetOwner<UScriptStruct>())
+			{
+				TSharedPtr<FAnimNextSchemaAction_Variable> Action = MakeShared<FAnimNextSchemaAction_Variable>(Property->GetFName(), Struct, FAnimNextParamType::FromProperty(Property), FAnimNextSchemaAction_Variable::EVariableAccessorChoice::Deferred);
+				
+				TSharedPtr<FVariableDragDropOp> Op = FVariableDragDropOp::New(Action, StaticCastSharedRef<FVariablesOutlinerEntryItem>(VariableItem->AsShared()));
+				return Op;
+			}
+		}
+	}
+	else if (FVariablesOutlinerCategoryItem* CategoryItem = SelectedItem->CastTo<FVariablesOutlinerCategoryItem>())
+	{
+		TSharedPtr<FCategoryDragDropOp> Op = FCategoryDragDropOp::New(StaticCastSharedRef<FVariablesOutlinerCategoryItem>(CategoryItem->AsShared()));
+		return Op;
+	}
+
+	return nullptr;
+
+}
+
+bool FVariablesOutlinerMode::ParseDragDrop(FSceneOutlinerDragDropPayload& OutPayload, const FDragDropOperation& Operation) const
+{
+	return ISceneOutlinerMode::ParseDragDrop(OutPayload, Operation);
+}
+
+FSceneOutlinerDragValidationInfo FVariablesOutlinerMode::ValidateDrop(const ISceneOutlinerTreeItem& DropTarget,
+	const FSceneOutlinerDragDropPayload& Payload) const
+{
+	return ISceneOutlinerMode::ValidateDrop(DropTarget, Payload);
+}
+
+void FVariablesOutlinerMode::OnDrop(ISceneOutlinerTreeItem& DropTarget, const FSceneOutlinerDragDropPayload& Payload,
+	const FSceneOutlinerDragValidationInfo& ValidationInfo) const
+{
+	ISceneOutlinerMode::OnDrop(DropTarget, Payload, ValidationInfo);
+}
+
+FReply FVariablesOutlinerMode::OnDragOverItem(const FDragDropEvent& Event, const ISceneOutlinerTreeItem& Item) const
+{
+	return ISceneOutlinerMode::OnDragOverItem(Event, Item);
+}
+
+int32 FVariablesOutlinerMode::GetTypeSortPriority(const ISceneOutlinerTreeItem& Item) const
+{
+	if(const FVariablesOutlinerAssetItem* AssetItem = Item.CastTo<FVariablesOutlinerAssetItem>())
+	{
+		return AssetItem->SortValue;
+	}
+
+	if(const FVariablesOutlinerEntryItem* EntryItem = Item.CastTo<FVariablesOutlinerEntryItem>())
+	{
+		return EntryItem->SortValue;
+	}
+
+	if(const FVariablesOutlinerCategoryItem* CategoryItem = Item.CastTo<FVariablesOutlinerCategoryItem>())
+	{
+		return CategoryItem->SortValue;
+	}
+	
+	return ISceneOutlinerMode::GetTypeSortPriority(Item);
+}
+
+void FVariablesOutlinerMode::OnItemDoubleClick(FSceneOutlinerTreeItemPtr Item)
+{
+	if (Item.IsValid())
+	{
+		if(FVariablesOutlinerAssetItem* AssetItem = Item->CastTo<FVariablesOutlinerAssetItem>())
+		{
+			// [TODO] Open asset from Variables Outliner
+		}
+	}	
+	
+	ISceneOutlinerMode::OnItemDoubleClick(Item);
+}
 }
 
-#undef LOCTEXT_NAMESPACE // "FVariablesOutlinerMode"
+#undef LOCTEXT_NAMESPACE // "FVariablesOutlinerMode"
\ No newline at end of file
diff --git a/Engine/Plugins/Experimental/UAF/UAF/Source/UAFEditor/Private/Variables/Outliner/VariablesOutlinerMode.h b/Engine/Plugins/Experimental/UAF/UAF/Source/UAFEditor/Private/Variables/Outliner/VariablesOutlinerMode.h
--- a/Engine/Plugins/Experimental/UAF/UAF/Source/UAFEditor/Private/Variables/Outliner/VariablesOutlinerMode.h
+++ b/Engine/Plugins/Experimental/UAF/UAF/Source/UAFEditor/Private/Variables/Outliner/VariablesOutlinerMode.h
@@ -5,9 +5,6 @@
 #include "AnimNextRigVMAssetEditorData.h"
 #include "ISceneOutlinerMode.h"
 #include "Common/AnimNextAssetFindReplaceVariables.h"
-#include "Common/Outliner/CommonOutlinerMode.h"
-
-#include "VariablesOutlinerMode.generated.h"
 
 class UAnimNextRigVMAssetEditorData;
 class FVariablesOutlinerMode;
@@ -16,122 +13,90 @@ namespace UE::Workspace
 	class IWorkspaceEditor;
 }
 
-USTRUCT()
-struct FVariableClipboardData
-{
-	GENERATED_BODY()
-	
-	UPROPERTY()
-	FSoftObjectPath SoftSourceObjectPath;
-
-	UPROPERTY()
-	EAnimNextExportAccessSpecifier Access = EAnimNextExportAccessSpecifier::Private;
-
-	UPROPERTY()
-	FName ParameterName;
-
-	UPROPERTY()
-	FAnimNextParamType Type;
-
-	UPROPERTY()
-	FAnimNextVariableBinding Binding;
-
-	UPROPERTY()
-	FString Comment;
-
-	UPROPERTY()
-	FString Category;
-
-	UPROPERTY()
-	FString DefaultValue;
-
-	bool operator==(const FVariableClipboardData& Other) const
-	{
-		return ParameterName == Other.ParameterName
-			&& Type == Other.Type
-			&& Binding.BindingData == Other.Binding.BindingData
-			&& Access == Other.Access
-			&& Comment == Other.Comment
-			&& Category == Other.Category
-			&& SoftSourceObjectPath == Other.SoftSourceObjectPath
-			&& DefaultValue == Other.DefaultValue;
-	}
-};
-
-USTRUCT()
-struct FVariablesOutlinerClipboardData
+namespace UE::UAF::Editor
 {
-	GENERATED_BODY()
-
-	UPROPERTY()
-	TArray<FVariableClipboardData> Variables;
-
-	UPROPERTY()
-	bool bCutOperation = false;
-
-	UPROPERTY()
-	bool bCategorySelectedForOperation = false;
-};
+	class SVariablesOutliner;
+}
 
 namespace UE::UAF::Editor
 {
-class FVariablesOutlinerMode : public FCommonOutlinerMode
+class FVariablesOutlinerMode : public ISceneOutlinerMode
 {
 public:
-	FVariablesOutlinerMode(SCommonOutliner* InOutliner, const TSharedRef<UE::Workspace::IWorkspaceEditor>& InWorkspaceEditor);
-	
-	// Begin FCommonOutlinerMode overrides
-	virtual void HandleItemSelection(const FSceneOutlinerItemSelection& Selection) const override;
-	virtual void RegisterToolMenus() override;
-	virtual void InitToolMenuContext(FToolMenuContext& InOutContext) const override;
-	virtual void ReorderCategoryItem(TSharedPtr<FOutlinerCategoryItem> EntryItem, TSharedPtr<FOutlinerCategoryItem> BeforeItem) const override;
-	virtual void ReorderCategory(const FString& CategoryName, const FString& BeforeCategoryName, UAnimNextRigVMAsset* Asset) const override;
-	virtual void MoveCategoryToAsset(const FString& CategoryPath, UAnimNextRigVMAsset* FromAsset, UAnimNextRigVMAsset* ToAsset) const override;
-	// End FCommonOutlinerMode overrides
+	FVariablesOutlinerMode(SVariablesOutliner* InVariablesOutliner, const TSharedRef<UE::Workspace::IWorkspaceEditor>& InWorkspaceEditor);
 
 	// Begin ISceneOutlinerMode overrides
 	virtual void Rebuild() override;
+	virtual TSharedPtr<SWidget> CreateContextMenu() override;
+	virtual FReply OnKeyDown(const FKeyEvent& InKeyEvent) override;
+	virtual void OnItemClicked(FSceneOutlinerTreeItemPtr Item) override;
+	void HandleItemSelection(const FSceneOutlinerItemSelection& Selection) const;
+	virtual void OnItemSelectionChanged(FSceneOutlinerTreeItemPtr Item, ESelectInfo::Type SelectionType, const FSceneOutlinerItemSelection& Selection) override;
+	virtual bool CanCustomizeToolbar() const override { return true; }
 	virtual bool CanRenameItem(const ISceneOutlinerTreeItem& Item) const override;
+	virtual ESelectionMode::Type GetSelectionMode() const override { return ESelectionMode::Multi; }
 	virtual void BindCommands(const TSharedRef<FUICommandList>& OutCommandList) override;
-	virtual bool UpdateOperationDecorator(const FDragDropEvent& Event, const FSceneOutlinerDragValidationInfo* ValidationInfo) const override;	
+	virtual int32 GetTypeSortPriority(const ISceneOutlinerTreeItem& Item) const override;
+	virtual TSharedPtr<FDragDropOperation> CreateDragDropOperation(const FPointerEvent& MouseEvent, const TArray<FSceneOutlinerTreeItemPtr>& InTreeItems) const override;
+	virtual bool ParseDragDrop(FSceneOutlinerDragDropPayload& OutPayload, const FDragDropOperation& Operation) const override;
+	virtual FSceneOutlinerDragValidationInfo ValidateDrop(const ISceneOutlinerTreeItem& DropTarget, const FSceneOutlinerDragDropPayload& Payload) const override;
+	virtual void OnDrop(ISceneOutlinerTreeItem& DropTarget, const FSceneOutlinerDragDropPayload& Payload, const FSceneOutlinerDragValidationInfo& ValidationInfo) const override;
+	virtual FReply OnDragOverItem(const FDragDropEvent& Event, const ISceneOutlinerTreeItem& Item) const override;
+	virtual bool CanSupportDragAndDrop() const override { return true; }
+	virtual void OnItemDoubleClick(FSceneOutlinerTreeItemPtr Item) override;
+	
+	void SetHighlightedItem(FSceneOutlinerTreeItemID InID) const;
+	void ClearHighlightedItem(FSceneOutlinerTreeItemID InID) const;
+
+	// Sets the export access specifier on ItemToChange and the current selection if ItemToChange is part of it
+	void SetSelectedVariablesExportAccessSpecifiers(FSceneOutlinerTreeItemPtr ItemToChange, EAnimNextExportAccessSpecifier VariableAccessSpecifier) const;
+
+	// Toggles the export access specifier on ItemToChange and sets the current selection to match, if ItemToChange is part of it
+	void ToggleSelectedVariablesExportAccessSpecifiers(FSceneOutlinerTreeItemPtr ItemToChange) const;
+
 	static void PopulateNewVariableToolMenuEntries(UToolMenu* InMenu, bool bAddSeparator);
 protected:
 	virtual TUniquePtr<ISceneOutlinerHierarchy> CreateHierarchy() override;
 	// End ISceneOutlinerMode overrides
 
+	void ResetOutlinerSelection() const;
+
+	SVariablesOutliner* GetOutliner() const;
 
-	// QQ begin/end overrides
-	virtual void Rename() const override;
+	void Rename() const;
 	virtual bool CanRename() const override;
 
-	virtual void Delete() const override;
+	void Delete() const;
 	virtual bool CanDelete() const override;
 
-	virtual void Copy() const override;
+	void Copy() const;
 	virtual bool CanCopy() const override;
 
-	virtual void Cut() const override;
-	virtual bool CanCut() const override;
-
-	virtual void Paste() const override;
+	void Paste() const;
 	virtual bool CanPaste() const override;
 
-	void CheckOverlappingVariableNamesRecursively(const UAnimNextRigVMAssetEditorData* EditorData, const TArray<FName>& BaseVariableNames, TMap<FName, FSoftObjectPath>& OutOverlappingVariableNames) const;
-	virtual void AddCategory(TWeakObjectPtr<UAnimNextRigVMAssetEditorData> InWeakEditorData) const override;
+	void SetVariableExportAccessSpecifiersOnSelection(EAnimNextExportAccessSpecifier VariableAccessSpecifier) const;
+	bool CanModifyVariableExport() const;
 
-	virtual void Duplicate() const override;
-	virtual bool CanDuplicate() const override;
+	void Duplicate() const;
+	bool CanDuplicate() const;
 
-	virtual void FindReferences(ESearchScope InSearchScope) const override;
-	virtual bool CanFindReferences() const override;
-	virtual bool IsFindReferencesVisible(ESearchScope InSearchScope) const override;
+	void FindReferences(ESearchScope InSearchScope) const;
+	bool CanFindReferences() const;
+	bool IsFindReferencesVisible(ESearchScope InSearchScope) const;
 
+	void SaveAsset() const;
+	bool CanSaveAsset() const;
 	bool IsAsset() const;
 
 	void CreateSharedVariablesAssets() const;
 	bool CanCreateSharedVariablesAssets() const;
 private:
 	friend class FVariablesOutlinerHierarchy;
+
+	TWeakPtr<UE::Workspace::IWorkspaceEditor> WeakWorkspaceEditor;
+
+	TSharedPtr<FUICommandList> CommandList;
 };
 
 }
diff --git a/Engine/Plugins/Experimental/UAF/UAF/Source/UAFEditor/Private/Variables/Outliner/VariablesOutlinerStructSharedVariablesItem.cpp b/Engine/Plugins/Experimental/UAF/UAF/Source/UAFEditor/Private/Variables/Outliner/VariablesOutlinerStructSharedVariablesItem.cpp
--- a/Engine/Plugins/Experimental/UAF/UAF/Source/UAFEditor/Private/Variables/Outliner/VariablesOutlinerStructSharedVariablesItem.cpp
+++ b/Engine/Plugins/Experimental/UAF/UAF/Source/UAFEditor/Private/Variables/Outliner/VariablesOutlinerStructSharedVariablesItem.cpp
@@ -9,7 +9,7 @@
 #include "SceneOutlinerStandaloneTypes.h"
 #include "AnimNextEditorModule.h"
 #include "AnimNextRigVMAssetEditorData.h"
-#include "Common/Outliner/OutlinerDragDropOps.h"
+#include "VariablesOutlinerDragDrop.h"
 #include "Entries/AnimNextSharedVariablesEntry.h"
 #include "Styling/SlateIconFinder.h"
 #include "Styling/StarshipCoreStyle.h"
@@ -27,7 +27,7 @@
 namespace UE::UAF::Editor
 {
 
-const FSceneOutlinerTreeItemType FVariablesOutlinerStructSharedVariablesItem::Type(&FOutlinerItem::Type);
+const FSceneOutlinerTreeItemType FVariablesOutlinerStructSharedVariablesItem::Type;
 
 class SVariablesOutlinerStructSharedVariablesLabel : FSceneOutlinerCommonLabelData, public SCompoundWidget
 {
@@ -40,9 +40,9 @@ class SVariablesOutlinerStructSharedVariablesLabel : FSceneOutlinerCommonLabelDa
 		TreeItem = StaticCastSharedRef<FVariablesOutlinerStructSharedVariablesItem>(InTreeItem.AsShared());
 
 		FText AssetName;
-		if(const UScriptStruct* Struct = InTreeItem.Struct.Get())
+		if(const UAnimNextSharedVariablesEntry* AssetEntry = InTreeItem.WeakEntry.Get())
 		{
-			AssetName = FText::FromString(Struct->GetPathName());
+			AssetName = FText::FromString(AssetEntry->GetObjectPath().ToString());
 		}
 		else
 		{
@@ -84,51 +84,58 @@ class SVariablesOutlinerStructSharedVariablesLabel : FSceneOutlinerCommonLabelDa
 		return FText();
 	}
 
-	virtual FSlateColor GetForegroundColor() const override
+	virtual void OnDragEnter( const FGeometry& MyGeometry, const FDragDropEvent& DragDropEvent ) override
 	{
-		TOptional<FLinearColor> BaseColor;
-		if (TreeItem.IsValid())
+		TSharedPtr<FVariableDragDropOp> GraphDropOp = DragDropEvent.GetOperationAs<FVariableDragDropOp>();
+		if (GraphDropOp.IsValid())
 		{
-			BaseColor = FSceneOutlinerCommonLabelData::GetForegroundColor(*TreeItem.Pin());
+			const FSlateBrush* ErrorSymbol = FAppStyle::GetBrush(TEXT("Graph.ConnectorFeedback.Error"));
+			GraphDropOp->SetSimpleFeedbackMessage(ErrorSymbol, FSlateColor(GetForegroundColor()), LOCTEXT("StructSharedVariablesFeedback", "Cannot make changes to Struct SharedVariables"));
 		}
-		return BaseColor.IsSet() ? BaseColor.GetValue() : FSlateColor::UseForeground();
 	}
 
+	virtual FSlateColor GetForegroundColor() const override
+	{
+		const TOptional<FLinearColor> BaseColor = FSceneOutlinerCommonLabelData::GetForegroundColor(*TreeItem.Pin());
+		return BaseColor.IsSet() ? BaseColor.GetValue() : FSlateColor::UseForeground();
+	}
 
 	TWeakPtr<FVariablesOutlinerStructSharedVariablesItem> TreeItem;
 	TSharedPtr<SInlineEditableTextBlock> TextBlock;
 };
 
-FVariablesOutlinerStructSharedVariablesItem::FVariablesOutlinerStructSharedVariablesItem(const TObjectPtr<const UScriptStruct> InStruct)
-: FOutlinerItem(FVariablesOutlinerStructSharedVariablesItem::Type, {nullptr, 0})
-	, Struct(InStruct)
+FVariablesOutlinerStructSharedVariablesItem::FVariablesOutlinerStructSharedVariablesItem(UAnimNextSharedVariablesEntry* InEntry)
+	: ISceneOutlinerTreeItem(FVariablesOutlinerStructSharedVariablesItem::Type)
+	, WeakEntry(InEntry)
 {
 }
 
 bool FVariablesOutlinerStructSharedVariablesItem::IsValid() const
 {
-	return Struct != nullptr;
+	return WeakEntry.Get() != nullptr;
 }
 
 FSceneOutlinerTreeItemID FVariablesOutlinerStructSharedVariablesItem::GetID() const
 {
-	if (Struct)
+	if (const UAnimNextSharedVariablesEntry* SharedVariableEntry = WeakEntry.Get())
 	{
-		const FSoftObjectPath Path = Struct;
+		check(SharedVariableEntry->GetStruct());
+		const FSoftObjectPath Path = SharedVariableEntry->GetStruct();
 		return GetTypeHash(Path);
 	}
 	
-	return GetTypeHash(Struct);
+	return GetTypeHash(WeakEntry);
 }
 
 FString FVariablesOutlinerStructSharedVariablesItem::GetDisplayString() const
 {
-	if(Struct == nullptr)
+	const UAnimNextRigVMAssetEntry* Entry = WeakEntry.Get();
+	if(Entry == nullptr)
 	{
 		return FString();
 	}
 
-	return Struct->GetName();
+	return Entry->GetDisplayName().ToString();
 }
 
 TSharedRef<SWidget> FVariablesOutlinerStructSharedVariablesItem::GenerateLabelWidget(ISceneOutliner& Outliner, const STableRow<FSceneOutlinerTreeItemPtr>& InRow)
@@ -138,14 +145,15 @@ TSharedRef<SWidget> FVariablesOutlinerStructSharedVariablesItem::GenerateLabelWi
 
 FString FVariablesOutlinerStructSharedVariablesItem::GetPackageName() const
 {
-	if(Struct == nullptr)
+	const UAnimNextRigVMAssetEntry* Entry = WeakEntry.Get();
+	if(Entry == nullptr)
 	{
 		return ISceneOutlinerTreeItem::GetPackageName();
 	}
 
-	return Struct->GetPackage()->GetName();
+	return Entry->GetPackage()->GetName();
 }
 
 }
 
-#undef LOCTEXT_NAMESPACE // "VariablesOutlinerTreeItem"
+#undef LOCTEXT_NAMESPACE // "VariablesOutlinerTreeItem"
\ No newline at end of file
diff --git a/Engine/Plugins/Experimental/UAF/UAF/Source/UAFEditor/Private/Variables/Outliner/VariablesOutlinerStructSharedVariablesItem.h b/Engine/Plugins/Experimental/UAF/UAF/Source/UAFEditor/Private/Variables/Outliner/VariablesOutlinerStructSharedVariablesItem.h
--- a/Engine/Plugins/Experimental/UAF/UAF/Source/UAFEditor/Private/Variables/Outliner/VariablesOutlinerStructSharedVariablesItem.h
+++ b/Engine/Plugins/Experimental/UAF/UAF/Source/UAFEditor/Private/Variables/Outliner/VariablesOutlinerStructSharedVariablesItem.h
@@ -3,7 +3,6 @@
 #pragma once
 
 #include "ISceneOutlinerTreeItem.h"
-#include "Common/Outliner/OutlinerEntryItem.h"
 
 class UAnimNextSharedVariablesEntry;
 
@@ -12,11 +11,11 @@ namespace UE::UAF::Editor
 
 class IWorkspaceOutlinerItemDetails;
 
-struct FVariablesOutlinerStructSharedVariablesItem : FOutlinerItem
+struct FVariablesOutlinerStructSharedVariablesItem : ISceneOutlinerTreeItem
 {
 	static const FSceneOutlinerTreeItemType Type;
 
-	FVariablesOutlinerStructSharedVariablesItem(const TObjectPtr<const UScriptStruct> InStruct);
+	FVariablesOutlinerStructSharedVariablesItem(UAnimNextSharedVariablesEntry* InEntry);
 
 	// Begin ISceneOutlinerTreeItem overrides
 	virtual bool IsValid() const override;
@@ -27,7 +26,8 @@ struct FVariablesOutlinerStructSharedVariablesItem : FOutlinerItem
 	virtual FString GetPackageName() const override;
 	// End ISceneOutlinerTreeItem overrides
 
-	TObjectPtr<const UScriptStruct> Struct;
+	// Weak-ptr to the represented entry
+	TWeakObjectPtr<UAnimNextSharedVariablesEntry> WeakEntry;
 };
 
 }
diff --git a/Engine/Plugins/Experimental/UAF/UAF/Source/UAFEditor/Private/Variables/SVariablesView.cpp b/Engine/Plugins/Experimental/UAF/UAF/Source/UAFEditor/Private/Variables/SVariablesView.cpp
--- a/Engine/Plugins/Experimental/UAF/UAF/Source/UAFEditor/Private/Variables/SVariablesView.cpp
+++ b/Engine/Plugins/Experimental/UAF/UAF/Source/UAFEditor/Private/Variables/SVariablesView.cpp
@@ -1,10 +1,9 @@
-﻿// Copyright Epic Games, Inc. All Rights Reserved.
+// Copyright Epic Games, Inc. All Rights Reserved.
 
 #include "SVariablesView.h"
 
 #include "AnimNextAssetWorkspaceAssetUserData.h"
 #include "AnimNextRigVMAsset.h"
-#include "AssetToolsModule.h"
 #include "ContentBrowserModule.h"
 #include "IContentBrowserSingleton.h"
 #include "SAssetDropTarget.h"
@@ -18,8 +17,7 @@
 #include "IWorkspaceEditor.h"
 #include "IWorkspaceEditorModule.h"
 #include "WorkspaceAssetRegistryInfo.h"
-#include "Common/Outliner/OutlinerColumns.h"
-#include "Common/Outliner/SCommonOutliner.h"
+#include "Outliner/VariablesOutlinerAssetItem.h"
 #include "Entries/AnimNextSharedVariablesEntry.h"
 
 #define LOCTEXT_NAMESPACE "SVariablesView"
@@ -28,6 +26,199 @@ namespace UE::UAF::Editor
 {
 
 const FLazyName VariablesTabName("VariablesTab");
+
+void SVariablesOutliner::RegisterAssetDelegates(const UAnimNextRigVMAsset* InAsset)
+{	
+	// Bind for any modification callbacks
+	UAnimNextRigVMAssetEditorData* EditorData = UncookedOnly::FUtils::GetEditorData<UAnimNextRigVMAssetEditorData>(InAsset);
+	if(EditorData != nullptr)
+	{
+		EditorData->ModifiedDelegate.AddSP(this, &SVariablesOutliner::OnEditorDataModified);
+
+		for (const TObjectPtr<UAnimNextRigVMAssetEntry>& Entry : EditorData->Entries)
+		{
+			if (UAnimNextSharedVariablesEntry* SharedVariablesEntry = Cast<UAnimNextSharedVariablesEntry>(Entry.Get()))
+			{
+				if (TObjectPtr<const UAnimNextSharedVariables> SharedVariablesAsset = SharedVariablesEntry->GetAsset())
+				{
+					RegisterAssetDelegates(SharedVariablesAsset);
+					Assets.AddUnique(Cast<UAnimNextRigVMAsset>(const_cast<UAnimNextSharedVariables*>(SharedVariablesAsset.Get())));					
+				}
+			}
+		}
+	}
+}
+
+void SVariablesOutliner::UnregisterAssetDelegates(const UAnimNextRigVMAsset* InAsset)
+{
+	UAnimNextRigVMAssetEditorData* EditorData = UncookedOnly::FUtils::GetEditorData<UAnimNextRigVMAssetEditorData>(InAsset);
+	if(EditorData != nullptr)
+	{
+		EditorData->ModifiedDelegate.RemoveAll(this);
+
+		for (const TObjectPtr<UAnimNextRigVMAssetEntry>& Entry : EditorData->Entries)
+		{
+			if (UAnimNextSharedVariablesEntry* SharedVariablesEntry = Cast<UAnimNextSharedVariablesEntry>(Entry.Get()))
+			{
+				if (TObjectPtr<const UAnimNextSharedVariables> SharedVariablesAsset = SharedVariablesEntry->GetAsset())
+				{
+					UnregisterAssetDelegates(SharedVariablesAsset);					
+				}
+			}
+		}
+	}
+}
+
+void SVariablesOutliner::SetAssets(TConstArrayView<TSoftObjectPtr<UAnimNextRigVMAsset>> InAssets)
+{
+	for(const TSoftObjectPtr<UAnimNextRigVMAsset>& CurrentSoftAsset : Assets)
+	{
+		if(UAnimNextRigVMAsset* CurrentAsset = CurrentSoftAsset.Get())
+		{
+			UnregisterAssetDelegates(CurrentAsset);
+		}
+	}
+
+	Assets = InAssets;
+
+	for(const TSoftObjectPtr<UAnimNextRigVMAsset>& NewSoftAsset : InAssets)
+	{
+		if(UAnimNextRigVMAsset* NewAsset = NewSoftAsset.Get())
+		{
+			RegisterAssetDelegates(NewAsset);
+		}
+	}
+
+	FullRefresh();
+}
+
+void SVariablesOutliner::UpdateAssets()
+{
+	TArray<FSoftObjectPath> AssetsToAsyncLoad;
+	TArray<TSoftObjectPtr<UAnimNextRigVMAsset>> ExportAssets;
+
+	auto HandleAssetPath = [this, &AssetsToAsyncLoad, &ExportAssets](const FSoftObjectPath& Path)
+	{
+		ExportAssets.Add(TSoftObjectPtr<UAnimNextRigVMAsset>(Path));
+
+		if (Path.ResolveObject() == nullptr)
+		{
+			AssetsToAsyncLoad.Add(Path);
+		}
+	};
+	
+	if (FWorkspaceOutlinerAssetReferenceItemData::IsAssetReference(Export)
+		|| FWorkspaceOutlinerGroupItemData::IsGroupItem(Export)
+		|| FAnimNextCollapseGraphsOutlinerDataBase::IsCollapsedGraphBase(Export))
+	{
+		// Reference item
+		TArray<FSoftObjectPath> AssetPaths;
+		Export.GetAssetPaths(AssetPaths);
+
+		for (const FSoftObjectPath& Path : AssetPaths)
+		{
+			HandleAssetPath(Path);
+		}
+	}
+	else
+	{
+		const FSoftObjectPath FirstAssetPath = Export.GetFirstAssetPath();
+		if(FirstAssetPath.IsAsset())
+		{
+			HandleAssetPath(FirstAssetPath);
+		}
+		else if(Export.HasData() && Export.GetData().GetScriptStruct()->IsChildOf(FAnimNextAssetEntryOutlinerData::StaticStruct()))
+		{
+			const FAnimNextAssetEntryOutlinerData& EntryData = Export.GetData().Get<FAnimNextAssetEntryOutlinerData>();
+			if(EntryData.SoftEntryPtr.IsValid())
+			{
+				if(UAnimNextRigVMAsset* Asset = EntryData.GetEntry()->GetTypedOuter<UAnimNextRigVMAsset>())
+				{
+					HandleAssetPath(Asset);
+				}
+			}
+		}
+	}
+
+	SetAssets(ExportAssets);
+
+	// Try async load any missing assets
+	for(const FSoftObjectPath& AssetPath : AssetsToAsyncLoad)
+	{
+		TWeakPtr<SVariablesOutliner> WeakOutliner = StaticCastSharedRef<SVariablesOutliner>(AsShared());
+		
+		AssetPath.LoadAsync(FLoadSoftObjectPathAsyncDelegate::CreateLambda([WeakOutliner](const FSoftObjectPath& InSoftObjectPath, UObject* InObject)
+		{
+			UAnimNextRigVMAsset* Asset = Cast<UAnimNextRigVMAsset>(InObject);
+			if(Asset == nullptr)
+			{
+				return;
+			}
+
+			TSharedPtr<SVariablesOutliner> PinnedVariablesOutliner = WeakOutliner.Pin();
+			if(!PinnedVariablesOutliner.IsValid())
+			{
+				return;
+			}
+
+			PinnedVariablesOutliner->HandleAssetLoaded(InSoftObjectPath, Asset);
+		}));
+	}
+}
+
+void SVariablesOutliner::SetExport(const FWorkspaceOutlinerItemExport& InExport)
+{
+	Export = InExport;
+
+	UpdateAssets();
+}
+
+void SVariablesOutliner::HandleAssetLoaded(const FSoftObjectPath& InSoftObjectPath, UAnimNextRigVMAsset* InAsset)
+{	
+	if(InAsset && Assets.Contains(TSoftObjectPtr<UAnimNextRigVMAsset>(InAsset)))
+	{
+		RegisterAssetDelegates(InAsset);
+
+		FullRefresh();
+	}
+}
+
+void SVariablesOutliner::OnEditorDataModified(UAnimNextRigVMAssetEditorData* InEditorData, EAnimNextEditorDataNotifType InType, UObject* InSubject)
+{	
+	ensure(Assets.Contains(TSoftObjectPtr<UAnimNextRigVMAsset>(UncookedOnly::FUtils::GetAsset<UAnimNextRigVMAsset>(InEditorData))));
+
+	switch(InType)
+	{
+	case EAnimNextEditorDataNotifType::UndoRedo:
+	case EAnimNextEditorDataNotifType::EntryAdded:
+	case EAnimNextEditorDataNotifType::VariableCategoryChanged:
+	case EAnimNextEditorDataNotifType::CategoryAdded:
+	case EAnimNextEditorDataNotifType::CategoryChanged:
+		FullRefresh();
+		break;
+	case EAnimNextEditorDataNotifType::EntryRemoved:
+		UpdateAssets();
+		break;
+	default:
+		break;
+	}
+}
+
+void SVariablesOutliner::SetHighlightedItem(FSceneOutlinerTreeItemPtr Item) const
+{
+	GetTreeView()->Private_SetItemHighlighted(Item, true);
+}
+
+void SVariablesOutliner::ClearHighlightedItem(FSceneOutlinerTreeItemPtr Item) const
+{
+	GetTreeView()->Private_SetItemHighlighted(Item, false);
+}
+
+bool SVariablesOutliner::HasAssets() const
+{
+	return Assets.Num() > 0;
+}
+
 void SVariablesView::Construct(const FArguments& InArgs, TSharedRef<UE::Workspace::IWorkspaceEditor> InWorkspaceEditor)
 {
 	FSceneOutlinerInitializationOptions InitOptions;
@@ -38,27 +229,11 @@ void SVariablesView::Construct(const FArguments& InArgs, TSharedRef<UE::Workspac
 		TOptional<float>(), TAttribute<FText>(), EHeaderComboVisibility::Ghosted));
 	InitOptions.ColumnMap.Add(FVariablesOutlinerValueColumn::GetID(), FSceneOutlinerColumnInfo(ESceneOutlinerColumnVisibility::Visible, 20, FCreateSceneOutlinerColumn::CreateLambda([](ISceneOutliner& InSceneOutliner) { return MakeShared<FVariablesOutlinerValueColumn>(InSceneOutliner); }), true, 0.5f));
 	InitOptions.ColumnMap.Add(FSceneOutlinerSourceControlColumn::GetID(), FSceneOutlinerColumnInfo(ESceneOutlinerColumnVisibility::Visible, 30, FCreateSceneOutlinerColumn::CreateLambda([](ISceneOutliner& InSceneOutliner) { return MakeShared<FSceneOutlinerSourceControlColumn>(InSceneOutliner); }), true));
-	InitOptions.ColumnMap.Add(FOutlinerAccessSpecifierColumn::GetID(), FSceneOutlinerColumnInfo(ESceneOutlinerColumnVisibility::Visible, 40, FCreateSceneOutlinerColumn::CreateLambda([](ISceneOutliner& InSceneOutliner) { return MakeShared<FOutlinerAccessSpecifierColumn>(InSceneOutliner); }), false));
-	InitOptions.ModeFactory = FCreateSceneOutlinerMode::CreateLambda([this, WeakWorkspaceEditor=InWorkspaceEditor.ToWeakPtr()](SSceneOutliner* InOutliner) { check(WeakWorkspaceEditor.IsValid()); return new FVariablesOutlinerMode(static_cast<SCommonOutliner*>(InOutliner), WeakWorkspaceEditor.Pin().ToSharedRef()); });
+	InitOptions.ColumnMap.Add(FVariablesOutlinerAccessSpecifierColumn::GetID(), FSceneOutlinerColumnInfo(ESceneOutlinerColumnVisibility::Visible, 40, FCreateSceneOutlinerColumn::CreateLambda([](ISceneOutliner& InSceneOutliner) { return MakeShared<FVariablesOutlinerAccessSpecifierColumn>(InSceneOutliner); }), false));
+	InitOptions.ModeFactory = FCreateSceneOutlinerMode::CreateLambda([this, WeakWorkspaceEditor=InWorkspaceEditor.ToWeakPtr()](SSceneOutliner* InOutliner) { return new FVariablesOutlinerMode(static_cast<SVariablesOutliner*>(InOutliner), WeakWorkspaceEditor.Pin().ToSharedRef()); });
 
-	VariablesOutliner = SNew(UE::UAF::Editor::SCommonOutliner, InitOptions)
-	.BindAssetDelegate_Lambda([WeakThis=StaticCastWeakPtr<SVariablesView>(AsWeak())](UAnimNextRigVMAssetEditorData* EditorData)
-	{
-		if (WeakThis.IsValid())
-		{
-			TSharedPtr<SVariablesView> SharedView = WeakThis.Pin();
-			EditorData->ModifiedDelegate.AddSP(SharedView.Get(), &SVariablesView::OnEditorDataModified);
-		}
-	})
-	.UnbindAssetDelegate_Lambda([WeakThis=StaticCastWeakPtr<SVariablesView>(AsWeak())](UAnimNextRigVMAssetEditorData* EditorData)
-	{
-		if (WeakThis.IsValid())
-		{
-			TSharedPtr<SVariablesView> SharedView = WeakThis.Pin();
-			EditorData->ModifiedDelegate.RemoveAll(SharedView.Get());
-		}
-	});
-	VariablesOutliner->SetEnabled(MakeAttributeSP(VariablesOutliner.Get(), &SCommonOutliner::HasAssets));
+	VariablesOutliner = SNew(SVariablesOutliner, InitOptions);
+	VariablesOutliner->SetEnabled(MakeAttributeSP(VariablesOutliner.Get(), &SVariablesOutliner::HasAssets));
 
 	InWorkspaceEditor->OnFocusedDocumentChanged().AddSP(this, &SVariablesView::HandleFocusedDocumentChanged);
 	const Workspace::FWorkspaceDocument& Document = InWorkspaceEditor->GetFocusedWorkspaceDocument();
@@ -84,26 +259,6 @@ void SVariablesView::HandleFocusedDocumentChanged(const UE::Workspace::FWorkspac
 	VariablesOutliner->SetExport(Export);
 }
 
-void SVariablesView::OnEditorDataModified(UAnimNextRigVMAssetEditorData* InEditorData, EAnimNextEditorDataNotifType InType, UObject* InSubject) const
-{
-	switch(InType)
-	{
-	case EAnimNextEditorDataNotifType::UndoRedo:
-	case EAnimNextEditorDataNotifType::EntryAdded:
-	case EAnimNextEditorDataNotifType::VariableCategoryChanged:
-	case EAnimNextEditorDataNotifType::CategoryAdded:
-	case EAnimNextEditorDataNotifType::CategoryChanged:
-	case EAnimNextEditorDataNotifType::CategoryRemoved:
-		VariablesOutliner->FullRefresh();
-		break;
-	case EAnimNextEditorDataNotifType::EntryRemoved:
-		VariablesOutliner->UpdateAssets();
-		break;
-	default:
-		break;
-	}
-}
-
 FAnimNextVariablesTabSummoner::FAnimNextVariablesTabSummoner(TSharedPtr<UE::Workspace::IWorkspaceEditor> InHostingApp)
 	: FWorkflowTabFactory(VariablesTabName, StaticCastSharedPtr<FAssetEditorToolkit>(InHostingApp))
 {
@@ -131,4 +286,4 @@ FText FAnimNextVariablesTabSummoner::GetTabToolTipText(const FWorkflowTabSpawnIn
 
 }
 
-#undef LOCTEXT_NAMESPACE
+#undef LOCTEXT_NAMESPACE
\ No newline at end of file
diff --git a/Engine/Plugins/Experimental/UAF/UAF/Source/UAFEditor/Private/Variables/SVariablesView.h b/Engine/Plugins/Experimental/UAF/UAF/Source/UAFEditor/Private/Variables/SVariablesView.h
--- a/Engine/Plugins/Experimental/UAF/UAF/Source/UAFEditor/Private/Variables/SVariablesView.h
+++ b/Engine/Plugins/Experimental/UAF/UAF/Source/UAFEditor/Private/Variables/SVariablesView.h
@@ -23,10 +23,41 @@ namespace UE::Workspace
 
 namespace UE::UAF::Editor
 {
-class SCommonOutliner;
+	class SVariablesOutliner;
+}
+
+namespace UE::UAF::Editor
+{
 
 extern const FLazyName VariablesTabName;
 
+class SVariablesOutliner : public SSceneOutliner
+{
+public:
+	void SetExport(const FWorkspaceOutlinerItemExport& InExport);
+	void OnEditorDataModified(UAnimNextRigVMAssetEditorData* InEditorData, EAnimNextEditorDataNotifType InType, UObject* InSubject);
+
+	void SetHighlightedItem(FSceneOutlinerTreeItemPtr Item) const;
+	void ClearHighlightedItem(FSceneOutlinerTreeItemPtr Item) const; 
+	
+	bool HasAssets() const;
+protected:
+	
+	void SetAssets(TConstArrayView<TSoftObjectPtr<UAnimNextRigVMAsset>> InAssets);
+	void UpdateAssets();
+
+	void HandleAssetLoaded(const FSoftObjectPath& InSoftObjectPath, UAnimNextRigVMAsset* InAsset);
+	
+	void RegisterAssetDelegates(const UAnimNextRigVMAsset* InAsset);
+	void UnregisterAssetDelegates(const UAnimNextRigVMAsset* InAsset);
+private:
+	friend class FVariablesOutlinerMode;
+	friend class FVariablesOutlinerHierarchy;
+
+	TArray<TSoftObjectPtr<UAnimNextRigVMAsset>> Assets;
+	FWorkspaceOutlinerItemExport Export;
+};
+
 class SVariablesView : public SCompoundWidget
 {
 	SLATE_BEGIN_ARGS(SVariablesView) {}
@@ -36,12 +67,11 @@ class SVariablesView : public SCompoundWidget
 
 	void SetExportDirectly(const FWorkspaceOutlinerItemExport& InExport) const;
 private:
-	void HandleFocusedDocumentChanged(const UE::Workspace::FWorkspaceDocument& InDocument) const;	
-	void OnEditorDataModified(UAnimNextRigVMAssetEditorData* InEditorData, EAnimNextEditorDataNotifType InType, UObject* InSubject) const;
+	void HandleFocusedDocumentChanged(const UE::Workspace::FWorkspaceDocument& InDocument) const;
 
 private:
 	friend struct FAnimNextVariablesTabSummoner;
-	TSharedPtr<SCommonOutliner> VariablesOutliner;
+	TSharedPtr<SVariablesOutliner> VariablesOutliner;
 };
 
 struct FAnimNextVariablesTabSummoner : public FWorkflowTabFactory
diff --git a/Engine/Plugins/Experimental/UAF/UAF/Source/UAFEditor/UAFEditor.Build.cs b/Engine/Plugins/Experimental/UAF/UAF/Source/UAFEditor/UAFEditor.Build.cs
--- a/Engine/Plugins/Experimental/UAF/UAF/Source/UAFEditor/UAFEditor.Build.cs
+++ b/Engine/Plugins/Experimental/UAF/UAF/Source/UAFEditor/UAFEditor.Build.cs
@@ -61,8 +61,6 @@ namespace UnrealBuildTool.Rules
 					"AdvancedPreviewScene"
 				}
 			);
-			
-			bAllowUETypesInNamespaces = true; 
 		}
 	}
 }
\ No newline at end of file
diff --git a/Engine/Plugins/Experimental/UAF/UAF/Source/UAFUncookedOnly/Internal/AnimNextRigVMAssetEditorData.h b/Engine/Plugins/Experimental/UAF/UAF/Source/UAFUncookedOnly/Internal/AnimNextRigVMAssetEditorData.h
--- a/Engine/Plugins/Experimental/UAF/UAF/Source/UAFUncookedOnly/Internal/AnimNextRigVMAssetEditorData.h
+++ b/Engine/Plugins/Experimental/UAF/UAF/Source/UAFUncookedOnly/Internal/AnimNextRigVMAssetEditorData.h
@@ -53,22 +53,16 @@ namespace UE::UAF::Editor
 	struct FVariablesOutlinerEntryItem;
 	class FVariablesOutlinerMode;
 	class FVariablesOutlinerHierarchy;
-	class FOutlinerHierarchy;
 	class SVariablesOutlinerValue;
-	class SFunctionsView;
-	class SVariablesView;
-	class FFunctionsOutlinerHierarchy;
+	class SVariablesOutliner;
 	class SAddVariablesDialog;
 	class FVariableProxyCustomization;
 	class FAnimNextAnimGraphEditorModule;
 	class FCallFunctionSharedDataDetails;
 	class FAssetCompilationHandler;
 	class FVariableBindingPropertyCustomization;
-	struct FFunctionsOutlinerEntryItem;
 	struct FVariablesOutlinerCategoryItem;
-	struct FFunctionsOutlinerCategoryItem;
-	class FFunctionsOutlinerMode;
-	class SCommonOutliner;
+	struct FVariablesOutlinerAssetItem;
 }
 
 namespace UE::UAF::Tests
@@ -96,7 +90,6 @@ enum class EAnimNextEditorDataNotifType : uint8
 	VariableCategoryChanged, // A variable entry its category changed (Subject == UAnimNextVariableEntry)
 	CategoryAdded, // A category has been added (Subject == UAnimNextRigVMAssetEditorData)
 	CategoryChanged, // A category has been changed (Subject == UAnimNextRigVMAssetEditorData)
-	CategoryRemoved, // A category has been removed (Subject == UAnimNextRigVMAssetEditorData)
 	VariablesReordered,  // Asset entries have been reordered (Subject == UAnimNextRigVMAssetEditorData)
 };
 
@@ -159,10 +152,6 @@ public:
 	/** Renames a category within an AnimNext asset */
 	UFUNCTION(BlueprintCallable, Category = "UAF|Entries", meta=(ScriptMethod))
 	static UAFUNCOOKEDONLY_API bool RenameCategory(UAnimNextRigVMAsset* InAsset, const FString& CategoryName, const FString& NewCategoryName, bool bSetupUndoRedo = true, bool bPrintPythonCommand = true);
-	
-	/** Removes a category from within an AnimNext asset */
-	UFUNCTION(BlueprintCallable, Category = "UAF|Entries", meta=(ScriptMethod))
-	static UAFUNCOOKEDONLY_API bool RemoveCategory(UAnimNextRigVMAsset* InAsset, const FString& CategoryName, bool bSetupUndoRedo = true, bool bPrintPythonCommand = true);
 };
 
 /* Base class for all AnimNext editor data objects that use RigVM */
@@ -195,9 +184,6 @@ public:
 
 	/** Renames a category using provided new name */
 	UE_API bool RenameCategory(const FString& OldName, const FString& NewName, bool bInSetupUndoRedo = true, bool bPrintPythonCommand = true);
-	
-	/** Removes category with provided name */
-	UE_API bool RemoveCategory(const FString CategoryName, bool bInSetupUndoRedo = true, bool bPrintPythonCommand = true);
 
 	/** Reorders category entry to appear before provided BeforeCategoryName */
 	UE_API void ReorderCategory(const FString& CategoryName, const FString& BeforeCategoryName, bool bInSetupUndoRedo = true, bool bPrintPythonCommand = true);
@@ -243,13 +229,11 @@ protected:
 	friend class UE::UAF::Editor::FAnimNextFunctionItemDetails;
 	friend struct UE::UAF::Editor::FVariablesOutlinerEntryItem;
 	friend struct UE::UAF::Editor::FVariablesOutlinerCategoryItem;
-	friend struct UE::UAF::Editor::FFunctionsOutlinerCategoryItem;
+	friend struct UE::UAF::Editor::FVariablesOutlinerAssetItem;
 	friend class UE::UAF::Editor::FVariablesOutlinerMode;
 	friend class UE::UAF::Editor::FVariablesOutlinerHierarchy;
-	friend class UE::UAF::Editor::FFunctionsOutlinerHierarchy;	
-	friend class UE::UAF::Editor::FOutlinerHierarchy;
 	friend class UE::UAF::Editor::SVariablesOutlinerValue;
-	friend class UE::UAF::Editor::FFunctionsOutlinerMode;	
+	friend class UE::UAF::Editor::SVariablesOutliner;
 	friend class UAnimNextModuleWorkspaceAssetUserData;
 	friend class UE::UAF::Editor::SAddVariablesDialog;
 	friend class UE::UAF::Editor::FVariableProxyCustomization;
@@ -267,10 +251,6 @@ protected:
 	friend struct FAnimNextModuleTests;
 	friend struct FAnimNextAnimationGraphTests;
 	friend FAnimNextRigVMAssetCompileContext;
-	friend UE::UAF::Editor::FFunctionsOutlinerEntryItem;
-	friend UE::UAF::Editor::SFunctionsView;
-	friend UE::UAF::Editor::SCommonOutliner;
-	friend UE::UAF::Editor::SVariablesView;
 
 	// UObject interface
 	UE_API virtual void Serialize(FArchive& Ar) override;
diff --git a/Engine/Plugins/Experimental/UAF/UAF/Source/UAFUncookedOnly/Internal/Entries/AnimNextSharedVariablesEntry.h b/Engine/Plugins/Experimental/UAF/UAF/Source/UAFUncookedOnly/Internal/Entries/AnimNextSharedVariablesEntry.h
--- a/Engine/Plugins/Experimental/UAF/UAF/Source/UAFUncookedOnly/Internal/Entries/AnimNextSharedVariablesEntry.h
+++ b/Engine/Plugins/Experimental/UAF/UAF/Source/UAFUncookedOnly/Internal/Entries/AnimNextSharedVariablesEntry.h
@@ -21,9 +21,8 @@ namespace UE::UAF::Editor
 	class SAddVariablesDialog;
 	class SVariableOverride;
 	class FVariablesOutlinerMode;
-	class FFunctionsOutlinerMode;
-	class FOutlinerHierarchy;
-	class SCommonOutliner;
+	class SVariablesOutliner;
+	struct FVariablesOutlinerAssetItem;
 }
 
 namespace UE::UAF::UncookedOnly
@@ -47,7 +46,6 @@ class UAnimNextSharedVariablesEntry : public UAnimNextRigVMAssetEntry
 	friend class UAnimNextRigVMAssetEditorData;
 	friend class UAssetDefinition_AnimNextSharedVariablesEntry;
 	friend class UE::UAF::Editor::FVariablesOutlinerHierarchy;
-	friend class UE::UAF::Editor::FOutlinerHierarchy;
 	friend struct UE::UAF::Editor::FVariablesOutlinerStructSharedVariablesItem;
 	friend struct UE::UAF::UncookedOnly::FUtils;
 	friend class UE::UAF::Editor::SVariablesOutlinerStructSharedVariablesLabel;
@@ -55,9 +53,9 @@ class UAnimNextSharedVariablesEntry : public UAnimNextRigVMAssetEntry
 	friend class UE::UAF::Editor::FVariableProxyCustomization;
 	friend class UE::UAF::Editor::SAddVariablesDialog;
 	friend class UE::UAF::Editor::SVariableOverride;
+	friend class UE::UAF::Editor::SVariablesOutliner;
+	friend struct UE::UAF::Editor::FVariablesOutlinerAssetItem;
 	friend UE::UAF::Editor::FVariablesOutlinerMode;
-	friend UE::UAF::Editor::SCommonOutliner;
-	friend UE::UAF::Editor::FFunctionsOutlinerMode;
 	friend class UAnimNextAssetFindReplaceVariables;
 
 	// UAnimNextRigVMAssetEntry interface
diff --git a/Engine/Plugins/Experimental/UAF/UAF/Source/UAFUncookedOnly/Internal/Entries/AnimNextVariableEntry.h b/Engine/Plugins/Experimental/UAF/UAF/Source/UAFUncookedOnly/Internal/Entries/AnimNextVariableEntry.h
--- a/Engine/Plugins/Experimental/UAF/UAF/Source/UAFUncookedOnly/Internal/Entries/AnimNextVariableEntry.h
+++ b/Engine/Plugins/Experimental/UAF/UAF/Source/UAFUncookedOnly/Internal/Entries/AnimNextVariableEntry.h
@@ -33,9 +33,6 @@ namespace UE::UAF::Editor
 	class SVariablesOutlinerEntryLabel;
 	class SVariablesOutlinerAssetLabel;
 	class SVariablesOutlinerCategoryLabel;
-	struct FVariablesOutlinerCategoryItem;
-	struct FVariablesOutlinerEntryItem;
-	class FVariablesOutlinerMode;
 }
 
 namespace UE::UAF::Tests
@@ -82,10 +79,6 @@ class UAnimNextVariableEntry : public UAnimNextRigVMAssetEntry, public IAnimNext
 	friend class UE::UAF::Editor::SVariablesOutlinerAssetLabel;
 	friend class UE::UAF::Editor::SVariablesOutlinerCategoryLabel;
 	friend class UAnimNextAssetFindReplaceVariables;
-	friend struct UE::UAF::Editor::FVariablesOutlinerCategoryItem;
-	friend struct UE::UAF::Editor::FVariablesOutlinerEntryItem;
-	friend class UE::UAF::Editor::FVariablesOutlinerMode;
-
 
 	UAnimNextVariableEntry();
 
diff --git a/Engine/Plugins/Experimental/UAF/UAF/Source/UAFUncookedOnly/Internal/UncookedOnlyUtils.h b/Engine/Plugins/Experimental/UAF/UAF/Source/UAFUncookedOnly/Internal/UncookedOnlyUtils.h
--- a/Engine/Plugins/Experimental/UAF/UAF/Source/UAFUncookedOnly/Internal/UncookedOnlyUtils.h
+++ b/Engine/Plugins/Experimental/UAF/UAF/Source/UAFUncookedOnly/Internal/UncookedOnlyUtils.h
@@ -210,7 +210,7 @@ struct FUtils
 	static UE_API void RenameVariable(UAnimNextVariableEntry* VariableEntry, const FName NewName, bool bSetupUndoRedo = true, bool bPrintPythonCommands = true);
 
 	// Moves the provided variable entry from its editor data to the NewOuter, replacing any references across the project its assets to reference it by its new outer asset
-	static UE_API UAnimNextVariableEntry* MoveVariableToAsset(UAnimNextVariableEntry* VariableEntry, UAnimNextRigVMAsset* NewOuter, bool bSetupUndoRedo = true, bool bPrintPythonCommands = true);
+	static UE_API void MoveVariableToAsset(UAnimNextVariableEntry* VariableEntry, UAnimNextRigVMAsset* NewOuter, bool bSetupUndoRedo = true, bool bPrintPythonCommands = true);
 
 	// Changes the provided variable entry its data type, replacing any references across the project to reference using its new data type
 	static UE_API void SetVariableType(UAnimNextVariableEntry* VariableEntry, const FAnimNextParamType NewType, bool bSetupUndoRedo = true, bool bPrintPythonCommands = true);
@@ -219,19 +219,6 @@ struct FUtils
 	static UE_API void ReplaceVariableReferences(UAnimNextRigVMAssetEditorData* InEditorData, const FAnimNextSoftVariableReference& ReferenceToFind, const FAnimNextSoftVariableReference& ReferenceToReplaceWith, bool bSetupUndoRedo = true, bool bPrintPythonCommands = true);
 	// Replaces the provided variable reference with ReferenceToReplaceWith within the context of the project
 	static UE_API void ReplaceVariableReferencesAcrossProject(const FAnimNextSoftVariableReference& ReferenceToFind, const FAnimNextSoftVariableReference& ReferenceToReplaceWith, bool bSetupUndoRedo = true, bool bPrintPythonCommands = true);
-
-	// Returns whether or not it is valid to add InSharedVariables to InEditorData, and if not OutErrorMessage will contain the reasoning
-	static UE_API bool CanAddSharedVariablesReference(UAnimNextRigVMAssetEditorData* InEditorData, const UAnimNextSharedVariables* InSharedVariables, FString* OutErrorMessage = nullptr);
-	
-	// Removes a SharedVariable reference from the provided EditorData, and ensures that any references to variables within the SharedVariables are replaced
-	static UE_API bool RemoveSharedVariablesReference(UAnimNextRigVMAssetEditorData* EditorData, const FSoftObjectPath& SharedVariablesPath);
-	
-	// Generate a (function) workspace export for the provided URigVMLibraryNode
-	static UE_API FWorkspaceOutlinerItemExport MakeFunctionExport(const URigVMLibraryNode* FunctionNode, const FWorkspaceOutlinerItemExport& ParentExport);
-
-	// Generate a (asset reference) workspace export for the provided asset path
-	static UE_API FWorkspaceOutlinerItemExport MakeAssetReferenceExport(const FSoftObjectPath& ReferencedAssetPath, const FWorkspaceOutlinerItemExport& ParentExport, bool bShouldExpand = true);
-
 private:
 	static UE_API void CompileVariableBindingsInternal(const FRigVMCompileSettings& InSettings, UAnimNextRigVMAsset* InAsset, TArray<URigVMGraph*>& OutGraphs, bool bInThreadSafe);
 	
diff --git a/Engine/Plugins/Experimental/UAF/UAF/Source/UAFUncookedOnly/Private/AnimNextRigVMAssetEditorData.cpp b/Engine/Plugins/Experimental/UAF/UAF/Source/UAFUncookedOnly/Private/AnimNextRigVMAssetEditorData.cpp
--- a/Engine/Plugins/Experimental/UAF/UAF/Source/UAFUncookedOnly/Private/AnimNextRigVMAssetEditorData.cpp
+++ b/Engine/Plugins/Experimental/UAF/UAF/Source/UAFUncookedOnly/Private/AnimNextRigVMAssetEditorData.cpp
@@ -1604,6 +1604,7 @@ bool UAnimNextRigVMAssetEditorData::AddCategory(const FString& CategoryName, boo
 		return false;
     }
 
+
 	const int32 CategoryIndex = VariableAndFunctionCategories.IndexOfByKey(CategoryName);
 	if(CategoryIndex != INDEX_NONE)
 	{
@@ -1669,40 +1670,8 @@ bool UAnimNextRigVMAssetEditorData::RenameCategory(const FString& OldName, const
 	return true;
 }
 
-bool UAnimNextRigVMAssetEditorData::RemoveCategory(const FString CategoryName, bool bInSetupUndoRedo, bool bPrintPythonCommand)
-{
-	if (CategoryName.IsEmpty())
-	{
-		ReportError(*FString::Printf(TEXT("UAnimNextRigVMAssetEditorData::RemoveCategory: Invalid empty category name provided.")));
-		return false;
-	}
-
-	const int32 CategoryIndex = VariableAndFunctionCategories.IndexOfByKey(CategoryName);
-	if(CategoryIndex == INDEX_NONE)
-	{
-		ReportError(*FString::Printf(TEXT("UAnimNextRigVMAssetEditorData::RemoveCategory: Could not find category with name %s."), *CategoryName));
-		return false;
-	}
-
-	if (bInSetupUndoRedo)
-	{
-		Modify();
-	}
-
-	VariableAndFunctionCategories.Remove(CategoryName);
-	BroadcastModified(EAnimNextEditorDataNotifType::CategoryRemoved, this);
-
-	if (bPrintPythonCommand)
-	{
-		RigVMPythonUtils::Print(GetName(),
-				FString::Printf(TEXT("asset.remove_category('%s')"),
-				 *CategoryName));
-	}
-
-	return true;
-}
-
-void UAnimNextRigVMAssetEditorData::ReorderCategory(const FString& CategoryName, const FString& BeforeCategoryName, bool bInSetupUndoRedo, bool bPrintPythonCommand)
+void UAnimNextRigVMAssetEditorData::ReorderCategory(const FString& CategoryName, const FString& BeforeCategoryName, bool bInSetupUndoRedo,
+	bool bPrintPythonCommand)
 {
 	const int32 CategoryIndex = VariableAndFunctionCategories.IndexOfByKey(CategoryName);
 	if (CategoryIndex == INDEX_NONE)
@@ -2466,7 +2435,30 @@ UAnimNextSharedVariablesEntry* UAnimNextRigVMAssetEditorData::AddSharedVariables
 	}
 
 	// Check for circularity
-	if(!FUtils::CanAddSharedVariablesReference(this, InSharedVariables))
+	auto CheckForCircularity = [this](UAnimNextSharedVariables_EditorData* InEditorData, auto& InCheckForCircularity)
+	{
+		if(InEditorData == this)
+		{
+			return true;
+		}
+
+		for(UAnimNextRigVMAssetEntry* Entry : InEditorData->Entries)
+		{
+			if(UAnimNextSharedVariablesEntry* SharedVariablesEntry = Cast<UAnimNextSharedVariablesEntry>(Entry))
+			{
+				const UAnimNextSharedVariables* SharedVariables = SharedVariablesEntry->GetAsset();
+				UAnimNextSharedVariables_EditorData* EditorData = UE::UAF::UncookedOnly::FUtils::GetEditorData<UAnimNextSharedVariables_EditorData>(SharedVariables);
+				if(InCheckForCircularity(EditorData, InCheckForCircularity))
+				{
+					return true;
+				}
+			}
+		}
+
+		return false;
+	};
+
+	if(CheckForCircularity(EditorData, CheckForCircularity))
 	{
 		ReportError(TEXT("UAnimNextRigVMAssetEditorData::AddSharedVariables: Circular reference detected."));
 		return nullptr;
@@ -2485,14 +2477,11 @@ UAnimNextSharedVariablesEntry* UAnimNextRigVMAssetEditorData::AddSharedVariables
 			}
 			else if(UAnimNextSharedVariablesEntry* SharedVariablesEntry = Cast<UAnimNextSharedVariablesEntry>(Entry))
 			{
-				// SharedVariables could be nullptr as the entry is a Struct based source, or the underlying asset (UAnimNextSharedVariablesEntry::ObjectPath) has been deleted
-				if (const UAnimNextSharedVariables* SharedVariables = SharedVariablesEntry->GetAsset())
+				const UAnimNextSharedVariables* SharedVariables = SharedVariablesEntry->GetAsset();
+				UAnimNextSharedVariables_EditorData* EditorData = UE::UAF::UncookedOnly::FUtils::GetEditorData<UAnimNextSharedVariables_EditorData>(SharedVariables);
+				if(InCheckForPublicMembers(EditorData, InCheckForPublicMembers))
 				{
-					UAnimNextSharedVariables_EditorData* EditorData = UE::UAF::UncookedOnly::FUtils::GetEditorData<UAnimNextSharedVariables_EditorData>(SharedVariables);
-					if(InCheckForPublicMembers(EditorData, InCheckForPublicMembers))
-					{
-						return true;
-					}
+					return true;
 				}
 			}
 		}
@@ -2500,6 +2489,12 @@ UAnimNextSharedVariablesEntry* UAnimNextRigVMAssetEditorData::AddSharedVariables
 		return false;
 	};
 
+	if(!CheckForPublicMembers(EditorData, CheckForPublicMembers))
+	{
+		ReportError(TEXT("UAnimNextRigVMAssetEditorData::AddSharedVariables: No public variables found."));
+		return nullptr;
+	}
+
 	// Check for duplicate entry
 	auto DuplicatePredicate = [InSharedVariables](const UAnimNextRigVMAssetEntry* InEntry)
 	{
@@ -2629,12 +2624,6 @@ bool UAnimNextRigVMAssetLibrary::RenameCategory(UAnimNextRigVMAsset* InAsset, co
 	return UE::UAF::UncookedOnly::FUtils::GetEditorData(InAsset)->RenameCategory(CategoryName, NewCategoryName, bSetupUndoRedo, bPrintPythonCommand);
 }
 
-bool UAnimNextRigVMAssetLibrary::RemoveCategory(UAnimNextRigVMAsset* InAsset, const FString& CategoryName, bool bSetupUndoRedo,
-	bool bPrintPythonCommand)
-{
-	return UE::UAF::UncookedOnly::FUtils::GetEditorData(InAsset)->RemoveCategory(CategoryName, bSetupUndoRedo, bPrintPythonCommand);
-}
-
 URigVMLibraryNode* UAnimNextRigVMAssetEditorData::AddFunction(FName InFunctionName, bool bInMutable, bool bInSetupUndoRedo, bool bPrintPythonCommand)
 {
 	URigVMController* Controller = RigVMClient.GetOrCreateController(GetLocalFunctionLibrary());
@@ -3142,7 +3131,7 @@ FInstancedPropertyBag UAnimNextRigVMAssetEditorData::GenerateCombinedPropertyBag
 					{
 						if (const FPropertyBagPropertyDesc* ExistingDesc = CombinedPropertyBag.FindPropertyDescByName(Desc.Name))
 						{
-							InCompileContext.GetAssetCompileContext().Error(Asset, LOCTEXT("DuplicateVariableEntries", "Variable(s) with overlapping Name {0} found in {1}"), FText::FromName(Desc.Name), FText::FromString(SourceAsset->GetPathName()));
+							InCompileContext.GetAssetCompileContext().Error(Asset, LOCTEXT("DuplicateVariableEntries", "Variable with overlapping Name {0} found in {1}"), FText::FromName(Desc.Name), FText::FromString(SourceAsset->GetPathName()));
 
 							bInvalidProperties = true;
 						}
diff --git a/Engine/Plugins/Experimental/UAF/UAF/Source/UAFUncookedOnly/Private/UncookedOnlyUtils.cpp b/Engine/Plugins/Experimental/UAF/UAF/Source/UAFUncookedOnly/Private/UncookedOnlyUtils.cpp
--- a/Engine/Plugins/Experimental/UAF/UAF/Source/UAFUncookedOnly/Private/UncookedOnlyUtils.cpp
+++ b/Engine/Plugins/Experimental/UAF/UAF/Source/UAFUncookedOnly/Private/UncookedOnlyUtils.cpp
@@ -38,11 +38,6 @@
 #include "ScopedTransaction.h"
 #include "Misc/UObjectToken.h"
 #include "Styling/SlateIconFinder.h"
-#include "Framework/Notifications/NotificationManager.h"
-#include "ScopedTransaction.h"
-#include "Misc/UObjectToken.h"
-#include "RigVMModel/Nodes/RigVMAggregateNode.h"
-#include "Widgets/Notifications/SNotificationList.h"
 
 #define LOCTEXT_NAMESPACE "AnimNextUncookedOnlyUtils"
 
@@ -238,13 +233,13 @@ void FUtils::CompileVariables(const FRigVMCompileSettings& InSettings, UAnimNext
 					const UAnimNextSharedVariables* SharedVariablesAsset = SharedVariablesEntry->GetAsset();
 					if (SharedVariablesAsset == nullptr)
 					{
-						SharedAssetEntryError(FText::Format(LOCTEXT("InvalidSharedVariableAssetFound", "Invalid shared variables asset found ({0})"), FText::FromString(SharedVariablesEntry->GetObjectPath().ToString())));
+						SharedAssetEntryError(FText::Format(LOCTEXT("InvalidSharedVariableAssetFound", "Invalid shared variables entry found ({0})"), FText::FromString(SharedVariablesEntry->GetObjectPath().ToString())));						
 						continue;
 					}
 
 					if (SharedVariableAssets.Contains(SharedVariablesAsset))
 					{
-						SharedAssetEntryError(FText::Format(LOCTEXT("DuplicateSharedVariableAssetFound", "Duplicate shared variables asset found ({0})"), FText::FromString(SharedVariablesEntry->GetObjectPath().ToString())));
+						SharedAssetEntryError(FText::Format(LOCTEXT("DuplicateSharedVariableAssetFound", "Duplicate shared variables entry found ({0})"), FText::FromString(SharedVariablesEntry->GetObjectPath().ToString())));						
 						continue;
 					}
 
@@ -262,7 +257,7 @@ void FUtils::CompileVariables(const FRigVMCompileSettings& InSettings, UAnimNext
 
 					if (SharedVariableStructs.Contains(SharedVariablesStruct))
 					{
-						SharedAssetEntryError(FText::Format(LOCTEXT("DuplicateSharedVariableStructFound", "Duplicate shared variables struct found ({0})"), FText::FromString(SharedVariablesEntry->GetObjectPath().ToString())));
+						SharedAssetEntryError(FText::Format(LOCTEXT("DuplicateSharedVariableStructFound", "Duplicate shared variables struct found ({0})"), FText::FromString(SharedVariablesEntry->GetObjectPath().ToString())));						
 						continue;
 					}
 
@@ -408,7 +403,7 @@ void FUtils::CompileVariableBindingsInternal(const FRigVMCompileSettings& InSett
 			return;
 		}
 		URigVMPin* CopyProxyVariablesExecPin = CopyProxyVariablesNode->FindPin(FRigVMStruct::ExecuteContextName.ToString());
-		if(CopyProxyVariablesExecPin == nullptr)
+		if(ExecPin == nullptr)
 		{
 			InSettings.ReportError(TEXT("Could not find execute pin on Copy System Proxy Variables node"));
 			return;
@@ -439,151 +434,6 @@ void FUtils::CompileVariableBindingsInternal(const FRigVMCompileSettings& InSett
 	OutGraphs.Add(BindingGraph);
 }
 
-bool FUtils::CanAddSharedVariablesReference(UAnimNextRigVMAssetEditorData* InEditorData, const UAnimNextSharedVariables* InSharedVariables, FString* OutErrorMessage)
-{
-	auto CheckForCircularity = [InEditorData](UAnimNextSharedVariables_EditorData* SourceEditorData, auto& InCheckForCircularity, TArray<const UAnimNextSharedVariables_EditorData*>& InOutReferenceChain)
-	{
-		TArray<const UAnimNextSharedVariables_EditorData*> ReferenceChain = InOutReferenceChain;
-		ReferenceChain.Add(SourceEditorData);
-		if(SourceEditorData == InEditorData)
-		{
-			InOutReferenceChain = ReferenceChain;
-			return true;
-		}
-
-		for(UAnimNextRigVMAssetEntry* Entry : SourceEditorData->Entries)
-		{
-			if(UAnimNextSharedVariablesEntry* SharedVariablesEntry = Cast<UAnimNextSharedVariablesEntry>(Entry))
-			{
-				// SharedVariables could be nullptr as the entry is a Struct based source, or the underlying asset (UAnimNextSharedVariablesEntry::ObjectPath) has been deleted
-				if (const UAnimNextSharedVariables* SharedVariables = SharedVariablesEntry->GetAsset())
-				{
-					UAnimNextSharedVariables_EditorData* EditorData = GetEditorData<UAnimNextSharedVariables_EditorData>(SharedVariables);
-					
-					if(InCheckForCircularity(EditorData, InCheckForCircularity, ReferenceChain))
-					{
-						InOutReferenceChain = ReferenceChain;
-						return true;
-					}
-				}
-			}
-		}
-
-		return false;
-	};
-	
-	if (UAnimNextSharedVariables_EditorData* SharedVariablesEditorData = GetEditorData<UAnimNextSharedVariables_EditorData>(InSharedVariables))
-	{
-		TArray<const UAnimNextSharedVariables_EditorData*> EditorDataReferenceChain;
-		EditorDataReferenceChain.Add(Cast<UAnimNextSharedVariables_EditorData>(InEditorData));
-		const bool bResult = CheckForCircularity(SharedVariablesEditorData, CheckForCircularity, EditorDataReferenceChain);
-
-		if (bResult && OutErrorMessage)
-		{
-			FString DependencyString;
-			for (const UAnimNextSharedVariables_EditorData* EditorData : EditorDataReferenceChain)
-			{
-				if (!DependencyString.IsEmpty())
-				{
-					DependencyString += TEXT("\n\t");
-				}
-				DependencyString += GetAsset<UAnimNextSharedVariables>(EditorData)->GetPackage()->GetFName().ToString();
-			}
-			*OutErrorMessage = FString::Printf(TEXT("Failed to add Shared Variables dependency as this would introduce a circular reference:\n\t%s"), *DependencyString);
-		}
-
-		return !bResult;
-	}
-
-	return false;	
-}
-
-bool FUtils::RemoveSharedVariablesReference(UAnimNextRigVMAssetEditorData* InEditorData, const FSoftObjectPath& SharedVariablesPath)
-{
-	const FText AssetNameText = FText::FromName(InEditorData->GetTypedOuter<UAnimNextRigVMAsset>()->GetFName());
-	if (UAnimNextSharedVariablesEntry* SharedVariableEntryToRemove = Cast<UAnimNextSharedVariablesEntry>(InEditorData->FindEntry(SharedVariablesPath.GetAssetFName())))
-	{
-		const FText SharedVariablesAssetNameText = FText::FromName(SharedVariablesPath.GetAssetFName());
-		const FText FormattedMessage = FText::Format(LOCTEXT("RemovingSharedAssetEntryFormat", "Removing {0} from {1}"), SharedVariablesAssetNameText, AssetNameText);
-		FScopedTransaction Transaction(FormattedMessage);
-		
-		if (const UAnimNextSharedVariables* SharedVariablesToRemove = SharedVariableEntryToRemove->GetAsset())
-		{
-			// Replace all referenced variables from the to-be-removed SharedVariables from the referencing Asset
-			FAnimNextAssetRegistryExports VariableExports;
-			GetExportedVariablesForAsset(SharedVariablesToRemove, VariableExports);
-	
-			VariableExports.ForEachExportOfType<FAnimNextVariableDeclarationData>([InEditorData, SharedVariablesPath](const FName& Identifier, const FAnimNextVariableDeclarationData& VariableDeclaration) -> bool
-			{
-				FAnimNextSoftVariableReference VariableReference(Identifier, SharedVariablesPath);
-				FAnimNextSoftVariableReference EmptyReference;
-				ReplaceVariableReferences(InEditorData, VariableReference, EmptyReference);
-			
-				return true;
-			});
-		}
-		else if (const UScriptStruct* Struct = SharedVariableEntryToRemove->GetStruct())
-		{
-			// Replace all referenced variables from the to-be-removed Native SharedVariables source from the referencing Asset
-			for (TFieldIterator<FProperty> It(Struct); It; ++It)
-			{
-				const FProperty* Property = *It;
-				if (Property->HasAnyPropertyFlags(CPF_NativeAccessSpecifierPublic))
-				{	
-					FAnimNextSoftVariableReference VariableReference(Property->GetFName(), SharedVariablesPath);
-					FAnimNextSoftVariableReference EmptyReference;
-					ReplaceVariableReferences(InEditorData, VariableReference, EmptyReference);
-				}
-			}
-		}
-	
-		// Remove the entry itself
-		return InEditorData->RemoveEntry(SharedVariableEntryToRemove);
-	}
-	
-	return false;
-}
-
-FWorkspaceOutlinerItemExport FUtils::MakeFunctionExport(const URigVMLibraryNode* FunctionNode, const FWorkspaceOutlinerItemExport& ParentExport)
-{
-	FWorkspaceOutlinerItemExport Export;
-	
-	check(FunctionNode);
-	if (const URigVMGraph* ContainedModelGraph = FunctionNode->GetContainedGraph())
-	{
-		if (const UAnimNextRigVMAssetEditorData* EditorData = FunctionNode->GetTypedOuter<UAnimNextRigVMAssetEditorData>())
-		{
-			if (URigVMEdGraph* EditorObject = Cast<URigVMEdGraph>(EditorData->GetEditorObjectForRigVMGraph(ContainedModelGraph)))
-			{
-				const FSoftObjectPath EditorObjectSoftObjPath = EditorObject;
-				ensureMsgf(EditorObjectSoftObjPath.IsSubobject(), TEXT("EditorObject for RigVMFunctionReferenceNode Graph was not a subobject as expected."));
-
-				const FName Identifier = EditorObjectSoftObjPath.IsSubobject() ? *EditorObjectSoftObjPath.GetSubPathUtf8String() : FunctionNode->GetFName();
-				Export = FWorkspaceOutlinerItemExport(Identifier, ParentExport);
-
-				Export.GetData().InitializeAsScriptStruct(FAnimNextGraphFunctionOutlinerData::StaticStruct());
-				FAnimNextGraphFunctionOutlinerData& FnGraphData = Export.GetData().GetMutable<FAnimNextGraphFunctionOutlinerData>();
-				FnGraphData.SoftEditorObject = EditorObject;
-			}
-		} 
-	}
-	
-	return Export;
-}
-
-FWorkspaceOutlinerItemExport FUtils::MakeAssetReferenceExport(const FSoftObjectPath& ReferencedAssetPath, const FWorkspaceOutlinerItemExport& ParentExport, bool bShouldExpand)
-{
-	check(ReferencedAssetPath.IsValid());
-	
-	FWorkspaceOutlinerItemExport Export = FWorkspaceOutlinerItemExport(FName(ReferencedAssetPath.ToString()), ParentExport);
-	Export.GetData().InitializeAs<FWorkspaceOutlinerAssetReferenceItemData>();
-	FWorkspaceOutlinerAssetReferenceItemData& Data = Export.GetData().GetMutable<FWorkspaceOutlinerAssetReferenceItemData>();
-	Data.ReferredObjectPath = ReferencedAssetPath;
-	Data.bShouldExpandReference = bShouldExpand;
-	
-	return Export;
-}
-
 UAnimNextRigVMAsset* FUtils::GetAsset(UAnimNextRigVMAssetEditorData* InEditorData)
 {
 	check(InEditorData);
@@ -1156,30 +1006,32 @@ void FUtils::GetAssetWorkspaceExports(const UAnimNextRigVMAssetEditorData* Edito
 				AddSharedVariablesGroupEntry();
 			}
 
-			const FSoftObjectPath ObjectPath = [SharedVariablesEntry]() -> FSoftObjectPath
+			switch(SharedVariablesEntry->Type)
 			{
-				switch(SharedVariablesEntry->Type)
+			case EAnimNextSharedVariablesType::Asset:
 				{
-				case EAnimNextSharedVariablesType::Asset:
-					{
-						const UAnimNextRigVMAsset* Asset = SharedVariablesEntry->Asset;
-						return Asset;
-					}
-					break;
-				case EAnimNextSharedVariablesType::Struct:
-					{
-						const UScriptStruct* Struct = SharedVariablesEntry->Struct;
-						return Struct;
-					}
-					break;
+					const UAnimNextRigVMAsset* Asset = SharedVariablesEntry->Asset;
+					const FSoftObjectPath ObjectPath(Asset);
+					const FWorkspaceOutlinerItemExport& ParentExport = OutExports.Exports[SharedVariablesGroupIndex];
+					FWorkspaceOutlinerItemExport& ReferenceExport = OutExports.Exports.Add_GetRef(FWorkspaceOutlinerItemExport(FName(ObjectPath.ToString()), ParentExport));
+					ReferenceExport.GetData().InitializeAs<FWorkspaceOutlinerAssetReferenceItemData>();
+					FWorkspaceOutlinerAssetReferenceItemData& Data = ReferenceExport.GetData().GetMutable<FWorkspaceOutlinerAssetReferenceItemData>();
+					Data.ReferredObjectPath = ObjectPath;
+					Data.bShouldExpandReference = false;
 				}
-				
-				return nullptr;
-			}();
-
-			if (ObjectPath.IsValid())
-			{
-				OutExports.Exports.Add(MakeAssetReferenceExport(ObjectPath, OutExports.Exports[SharedVariablesGroupIndex], false));
+				break;
+			case EAnimNextSharedVariablesType::Struct:
+				{
+					const UScriptStruct* Struct = SharedVariablesEntry->Struct;
+					const FSoftObjectPath ObjectPath(Struct);
+					const FWorkspaceOutlinerItemExport& ParentExport = OutExports.Exports[SharedVariablesGroupIndex];
+					FWorkspaceOutlinerItemExport& ReferenceExport = OutExports.Exports.Add_GetRef(FWorkspaceOutlinerItemExport(FName(ObjectPath.ToString()), ParentExport));
+					ReferenceExport.GetData().InitializeAs<FWorkspaceOutlinerAssetReferenceItemData>();
+					FWorkspaceOutlinerAssetReferenceItemData& Data = ReferenceExport.GetData().GetMutable<FWorkspaceOutlinerAssetReferenceItemData>();
+					Data.ReferredObjectPath = ObjectPath;
+					Data.bShouldExpandReference = false;
+				}
+				break;
 			}
 		}
 		return true;
@@ -1229,7 +1081,7 @@ void FUtils::GetAssetWorkspaceExports(const UAnimNextRigVMAssetEditorData* Edito
 		}
 		return true;
 	});	
-	
+
 	GetFunctionLibraryWorkspaceExportsRecursive(EditorData, OutExports, AssetIdentifier, INDEX_NONE, EditorData->GetRigVMGraphFunctionStore()->PublicFunctions, EditorData->GetRigVMGraphFunctionStore()->PrivateFunctions);
 }
 
@@ -1242,7 +1094,9 @@ void FUtils::ProcessPinAssetReferences(const URigVMPin* InPin, FWorkspaceOutline
 		if (InSoftObjectPath.ResolveObject() || IAssetRegistry::GetChecked().TryGetAssetByObjectPath(InSoftObjectPath, ReferenceAssetData) == UE::AssetRegistry::EExists::Exists)
 		{
 			const FWorkspaceOutlinerItemExport& ParentExport = ParentExportIndex == INDEX_NONE ? RootExport : OutExports.Exports[ParentExportIndex];
-			OutExports.Exports.Add(MakeAssetReferenceExport(InSoftObjectPath, ParentExport, true));
+			FWorkspaceOutlinerItemExport& ReferenceExport = OutExports.Exports.Add_GetRef(FWorkspaceOutlinerItemExport(FName(InSoftObjectPath.ToString()), ParentExport));
+			ReferenceExport.GetData().InitializeAs<FWorkspaceOutlinerAssetReferenceItemData>();
+			ReferenceExport.GetData().GetMutable<FWorkspaceOutlinerAssetReferenceItemData>().ReferredObjectPath = InSoftObjectPath;
 		}
 	};
 	
@@ -1414,22 +1268,57 @@ void FUtils::GetSubGraphWorkspaceExportsRecursive(const UAnimNextRigVMAssetEdito
 		{
 			if(ensure(EditorObject->GetModel()))
 			{
-				// Ignore aggregate nodes
-				const URigVMCollapseNode* CollapseNode = ExactCast<URigVMCollapseNode>(EditorObject->GetModel()->GetOuter());
-				if (CollapseNode)
+				const URigVMCollapseNode* CollapseNode = CastChecked<URigVMCollapseNode>(EditorObject->GetModel()->GetOuter());
+				const FSoftObjectPath EditorObjectSoftObjPath = EditorObject;
+				ensureMsgf(EditorObjectSoftObjPath.IsSubobject(), TEXT("EditorObject for RigVMCollapseNode Graph was not a subobject as expected."));
+
+				const FWorkspaceOutlinerItemExport& ParentExport = ParentExportIndex == INDEX_NONE ? RootExport : OutExports.Exports[ParentExportIndex];
+				const FName Identifier = EditorObjectSoftObjPath.IsSubobject() ? *EditorObjectSoftObjPath.GetSubPathUtf8String() : CollapseNode->GetFName();
+				FWorkspaceOutlinerItemExport& Export = OutExports.Exports.Add_GetRef(FWorkspaceOutlinerItemExport(Identifier, ParentExport));
+				Export.GetData().InitializeAsScriptStruct(FAnimNextCollapseGraphOutlinerData::StaticStruct());
+			
+				FAnimNextCollapseGraphOutlinerData& FnGraphData = Export.GetData().GetMutable<FAnimNextCollapseGraphOutlinerData>();
+				FnGraphData.SoftEditorObject = EditorObject;
+
+				int32 ExportIndex = OutExports.Exports.Num() - 1;
+				GetSubGraphWorkspaceExportsRecursive(EditorData, OutExports, RootExport, ExportIndex, EditorObject, Context);
+			}
+		}
+	}
+
+	// ---- Function References ---
+	TArray<URigVMEdGraphNode*> EdNodes;
+	RigVMEdGraph->GetNodesOfClass(EdNodes);
+
+	for (URigVMEdGraphNode* EdNode : EdNodes)
+	{
+		if (URigVMFunctionReferenceNode* FunctionReferenceNode = Cast<URigVMFunctionReferenceNode>(EdNode->GetModelNode()))
+		{
+			// Only export referenced functions which are part of same outer
+			const FRigVMGraphFunctionIdentifier FunctionIdentifier = FunctionReferenceNode->GetFunctionIdentifier();
+			if (FunctionIdentifier.HostObject == FSoftObjectPath(EditorData))
+			{
+				const URigVMLibraryNode* FunctionNode = EditorData->RigVMClient.GetFunctionLibrary()->FindFunction(FunctionIdentifier.GetFunctionFName());
+				if (const URigVMGraph* ContainedGraph = FunctionNode->GetContainedGraph())
 				{
-					const FSoftObjectPath EditorObjectSoftObjPath = EditorObject;
-					ensureMsgf(EditorObjectSoftObjPath.IsSubobject(), TEXT("EditorObject for RigVMCollapseNode Graph was not a subobject as expected."));
+					URigVMEdGraph* ContainedEdGraph = Cast<URigVMEdGraph>(EditorData->GetEditorObjectForRigVMGraph(ContainedGraph));
+					const FSoftObjectPath EditorObjectSoftObjPath = ContainedEdGraph;
+					ensureMsgf(EditorObjectSoftObjPath.IsSubobject(), TEXT("EditorObject for RigVMFunctionReferenceNode Graph was not a subobject as expected."));
 
 					const FWorkspaceOutlinerItemExport& ParentExport = ParentExportIndex == INDEX_NONE ? RootExport : OutExports.Exports[ParentExportIndex];
-					const FName Identifier = EditorObjectSoftObjPath.IsSubobject() ? *EditorObjectSoftObjPath.GetSubPathUtf8String() : CollapseNode->GetFName();
+					const FName Identifier = EditorObjectSoftObjPath.IsSubobject() ? *EditorObjectSoftObjPath.GetSubPathUtf8String() : FunctionNode->GetFName();
 					FWorkspaceOutlinerItemExport& Export = OutExports.Exports.Add_GetRef(FWorkspaceOutlinerItemExport(Identifier, ParentExport));
-					Export.GetData().InitializeAsScriptStruct(FAnimNextCollapseGraphOutlinerData::StaticStruct());
-					FAnimNextCollapseGraphOutlinerData& FnGraphData = Export.GetData().GetMutable<FAnimNextCollapseGraphOutlinerData>();
-					FnGraphData.SoftEditorObject = EditorObject;
+					Export.GetData().InitializeAsScriptStruct(FAnimNextGraphFunctionOutlinerData::StaticStruct());
+					FAnimNextGraphFunctionOutlinerData& FnGraphData = Export.GetData().GetMutable<FAnimNextGraphFunctionOutlinerData>();
+
+					if (ContainedEdGraph)
+					{
+						FnGraphData.SoftEditorObject = ContainedEdGraph;
+						FnGraphData.SoftEdGraphNode = EdNode;
 
-					int32 ExportIndex = OutExports.Exports.Num() - 1;
-					GetSubGraphWorkspaceExportsRecursive(EditorData, OutExports, RootExport, ExportIndex, EditorObject, Context);
+						int32 ExportIndex = OutExports.Exports.Num() - 1;
+						GetSubGraphWorkspaceExportsRecursive(EditorData, OutExports, RootExport, ExportIndex, ContainedEdGraph, Context);
+					}
 				}
 			}
 		}
@@ -1456,13 +1345,23 @@ void FUtils::GetFunctionsWorkspaceExportsRecursive(const UAnimNextRigVMAssetEdit
 {
 	for (const FRigVMGraphFunctionData& FunctionData : Functions)
 	{
-		const FWorkspaceOutlinerItemExport& ParentExport = ParentExportIndex == INDEX_NONE ? RootExport : OutExports.Exports[ParentExportIndex];
 		if (const URigVMLibraryNode* FunctionNode = EditorData->RigVMClient.GetFunctionLibrary()->FindFunction(FunctionData.Header.LibraryPointer.GetFunctionFName()))
 		{
-			const FWorkspaceOutlinerItemExport Export = MakeFunctionExport(FunctionNode, ParentExport);
-			if (Export.GetIdentifier() != NAME_None)
+			if (const URigVMGraph* ContainedModelGraph = FunctionNode->GetContainedGraph())
 			{
-				OutExports.Exports.Add(Export);
+				if (URigVMEdGraph* EditorObject = Cast<URigVMEdGraph>(EditorData->GetEditorObjectForRigVMGraph(ContainedModelGraph)))
+				{
+					const FSoftObjectPath EditorObjectSoftObjPath = EditorObject;
+					ensureMsgf(EditorObjectSoftObjPath.IsSubobject(), TEXT("EditorObject for RigVMFunctionReferenceNode Graph was not a subobject as expected."));
+
+					const FWorkspaceOutlinerItemExport& ParentExport = ParentExportIndex == INDEX_NONE ? RootExport : OutExports.Exports[ParentExportIndex];
+					const FName Identifier = EditorObjectSoftObjPath.IsSubobject() ? *EditorObjectSoftObjPath.GetSubPathUtf8String() : FunctionNode->GetFName();
+					FWorkspaceOutlinerItemExport& Export = OutExports.Exports.Add_GetRef(FWorkspaceOutlinerItemExport(Identifier, ParentExport));
+
+					Export.GetData().InitializeAsScriptStruct(FAnimNextGraphFunctionOutlinerData::StaticStruct());
+					FAnimNextGraphFunctionOutlinerData& FnGraphData = Export.GetData().GetMutable<FAnimNextGraphFunctionOutlinerData>();
+					FnGraphData.SoftEditorObject = EditorObject;
+				}
 			}
 		}
 	}
@@ -1669,7 +1568,7 @@ void FUtils::RenameVariable(UAnimNextVariableEntry* VariableEntry, const FName N
 	}	
 }
 
-UAnimNextVariableEntry* FUtils::MoveVariableToAsset(UAnimNextVariableEntry* VariableEntry, UAnimNextRigVMAsset* NewOuter, bool bSetupUndoRedo, bool bPrintPythonCommands)
+void FUtils::MoveVariableToAsset(UAnimNextVariableEntry* VariableEntry, UAnimNextRigVMAsset* NewOuter, bool bSetupUndoRedo, bool bPrintPythonCommands)
 {
 	if (VariableEntry && NewOuter)
 	{
@@ -1683,8 +1582,7 @@ UAnimNextVariableEntry* FUtils::MoveVariableToAsset(UAnimNextVariableEntry* Vari
 			FString DefaultValueString;
 			VariableEntry->GetDefaultValueString(DefaultValueString);
 
-			UAnimNextVariableEntry* NewEntry = NewEditorData->AddVariable(VariableEntry->GetVariableName(), VariableEntry->GetType(), DefaultValueString, bSetupUndoRedo, bPrintPythonCommands);
-			check(NewEntry);
+			check(NewEditorData->AddVariable(VariableEntry->GetVariableName(), VariableEntry->GetType(), DefaultValueString, bSetupUndoRedo, bPrintPythonCommands));
 
 			const FAnimNextSoftVariableReference FindReference(VariableEntry->GetEntryName(), CurrentOuter);
 			const FAnimNextSoftVariableReference ReplaceReference(VariableEntry->GetEntryName(), NewOuter);
@@ -1702,12 +1600,8 @@ UAnimNextVariableEntry* FUtils::MoveVariableToAsset(UAnimNextVariableEntry* Vari
 			}
 
 			CurrentEditorData->RemoveEntry(VariableEntry, bSetupUndoRedo, bPrintPythonCommands);
-			
-			return NewEntry;
 		}
 	}
-	
-	return nullptr;
 }
 
 void FUtils::SetVariableType(UAnimNextVariableEntry* VariableEntry, const FAnimNextParamType NewType, bool bSetupUndoRedo, bool bPrintPythonCommands)
@@ -1777,9 +1671,8 @@ void FUtils::ReplaceVariableReferences(UAnimNextRigVMAssetEditorData* InEditorDa
 	const bool bVariableTypeChanged = FindReferenceType != ReplaceReferenceType;
 	const bool bVariableNameChanged = ReferenceToFind.GetName() != ReferenceToReplaceWith.GetName();
 
-	// Ensure that SharedVariables have been included for replace variable (if from another asset)
-	const UAnimNextSharedVariables* SharedVariablesReplaceOuter = Cast<UAnimNextSharedVariables>(ReplaceOuterAsset);
-	if (SharedVariablesReplaceOuter && ThisOuterAsset != SharedVariablesReplaceOuter)
+	// Ensure that SharedVariables have been included for replace variable
+	if (const UAnimNextSharedVariables* SharedVariablesReplaceOuter = Cast<UAnimNextSharedVariables>(ReplaceOuterAsset))
 	{
 		bool bHasEntry = false;
 		InEditorData->ForEachEntryOfType<UAnimNextSharedVariablesEntry>([&bHasEntry, SharedVariablesReplaceOuter](const UAnimNextSharedVariablesEntry* Entry)
@@ -1798,21 +1691,9 @@ void FUtils::ReplaceVariableReferences(UAnimNextRigVMAssetEditorData* InEditorDa
 
 		if (!bHasEntry)
 		{
-			FString ErrorMessage;
-			if (CanAddSharedVariablesReference(InEditorData, SharedVariablesReplaceOuter, &ErrorMessage))
-			{
-				UE::UAF::UncookedOnly::FScopedCompileJob CompilerResults(LOCTEXT("AddSharedAssetJobName", "Add Shared Asset dependency"), { ThisOuterAsset });
-				InEditorData->AddSharedVariables(SharedVariablesReplaceOuter, bSetupUndoRedo, bPrintPythonCommands);
-				InEditorData->RecompileVM();
-			}
-			else
-			{
-				FText NotificationText = FText::FromString(ErrorMessage);
-				FNotificationInfo Notification(NotificationText);
-				Notification.ExpireDuration = 3.0f;
-				FSlateNotificationManager::Get().AddNotification(Notification);
-				UE_LOGFMT(LogAnimation, Error, "{0}", ErrorMessage);
-			}
+			UE::UAF::UncookedOnly::FScopedCompileJob CompilerResults(LOCTEXT("AddSharedAssetJobName", "Add Shared Asset dependency"), { ThisOuterAsset });
+			InEditorData->AddSharedVariables(SharedVariablesReplaceOuter, bSetupUndoRedo, bPrintPythonCommands);
+			InEditorData->RecompileVM();
 		}
 	}			
 	
diff --git a/Engine/Plugins/Experimental/Workspace/Source/WorkspaceEditor/Private/WorkspaceEditor.cpp b/Engine/Plugins/Experimental/Workspace/Source/WorkspaceEditor/Private/WorkspaceEditor.cpp
--- a/Engine/Plugins/Experimental/Workspace/Source/WorkspaceEditor/Private/WorkspaceEditor.cpp
+++ b/Engine/Plugins/Experimental/Workspace/Source/WorkspaceEditor/Private/WorkspaceEditor.cpp
@@ -141,7 +141,6 @@ void FWorkspaceEditor::CreateWidgets()
 					->SetHideTabWell(false)
 					->AddTab(WorkspaceTabs::TopLeftDocumentArea, ETabState::ClosedTab)
 					->AddTab(FBaseAssetToolkit::ViewportTabID, ETabState::OpenedTab)
-					->SetExtensionId(WorkspaceTabs::TopLeftDocumentArea)
 				)
 				->Split
 				(
@@ -149,7 +148,6 @@ void FWorkspaceEditor::CreateWidgets()
 					->SetSizeCoefficient(0.5f)
 					->SetHideTabWell(false)
 					->AddTab(WorkspaceTabs::BottomLeftDocumentArea, ETabState::ClosedTab)
-					->SetExtensionId(WorkspaceTabs::BottomLeftDocumentArea)
 				)
 			)
 			->Split
@@ -163,7 +161,6 @@ void FWorkspaceEditor::CreateWidgets()
 					->SetSizeCoefficient(0.75f)
 					->SetHideTabWell(false)
 					->AddTab(WorkspaceTabs::TopMiddleDocumentArea, ETabState::ClosedTab)
-					->SetExtensionId(WorkspaceTabs::TopMiddleDocumentArea)
 				)
 				->Split
 				(
@@ -171,7 +168,6 @@ void FWorkspaceEditor::CreateWidgets()
 					->SetSizeCoefficient(0.25f)
 					->SetHideTabWell(false)
 					->AddTab(WorkspaceTabs::BottomMiddleDocumentArea, ETabState::ClosedTab)
-					->SetExtensionId(WorkspaceTabs::BottomMiddleDocumentArea)
 				)
 			)
 			->Split
@@ -186,7 +182,6 @@ void FWorkspaceEditor::CreateWidgets()
 					->SetHideTabWell(false)
 					->AddTab(WorkspaceTabs::WorkspaceView, ETabState::OpenedTab)
 					->AddTab(WorkspaceTabs::TopRightDocumentArea, ETabState::ClosedTab)
-					->SetExtensionId(WorkspaceTabs::TopRightDocumentArea)
 				)
 				->Split
 				(
@@ -195,7 +190,6 @@ void FWorkspaceEditor::CreateWidgets()
 					->SetHideTabWell(false)
 					->AddTab(WorkspaceTabs::BottomRightDocumentArea, ETabState::ClosedTab)
 					->AddTab(FBaseAssetToolkit::DetailsTabID, ETabState::OpenedTab)
-					->SetExtensionId(WorkspaceTabs::BottomRightDocumentArea)
 				)
 			)
 		)
diff --git a/Engine/Plugins/Experimental/Workspace/Source/WorkspaceEditor/Public/WorkspaceAssetRegistryInfo.h b/Engine/Plugins/Experimental/Workspace/Source/WorkspaceEditor/Public/WorkspaceAssetRegistryInfo.h
--- a/Engine/Plugins/Experimental/Workspace/Source/WorkspaceEditor/Public/WorkspaceAssetRegistryInfo.h
+++ b/Engine/Plugins/Experimental/Workspace/Source/WorkspaceEditor/Public/WorkspaceAssetRegistryInfo.h
@@ -329,29 +329,6 @@ public:
 			}
 		}
 	}
-			
-	// Returns relative export to the provided InAssetPath
-	//  - "SoftObjectPath" - "Foo" - "SoftObjectPathTwo" - "Bar" will return "SoftObjectPath" - "Foo" when passed "SoftObjectPathTwo" or an empty export when passed "SoftObjectPath"
-	FWorkspaceOutlinerItemExport GetRelativeAssetExport(const FSoftObjectPath& InAssetPath) const
-	{
-		FWorkspaceOutlinerItemExport Export;
-		
-		if(Path.PathSegments.Num() > 0)
-		{
-			const int32 NumSegments = Path.PathSegments.Num();
-			for (int32 SegmentIndex = 0; SegmentIndex < NumSegments; SegmentIndex++)
-			{
-				FSoftObjectPath ObjectPath = Path.PathSegments[SegmentIndex].ToString();
-				if (ObjectPath.IsValid() && ObjectPath == InAssetPath && SegmentIndex != 0)
-				{
-					Export.Path.PathSegments.Append(&Path.PathSegments[0], SegmentIndex);
-					break;
-				}
-			}
-		}
-		
-		return Export;
-	}
 
 	FString GetFullPath() const
 	{
