Misc abstract skeleton set collection editor changes
* Allow moving sets to empty space to unparent them from anything
* Remove debug details panel
* Refactor drag-drop ops to use a common set drag-drop-op across editors
* Move code from UE::UAF::Sets to UE::UAF::Editor
* Add SetCollection::OnHierarchyChanged delegate
* Make the search box work
* Allow adding sets via context menu

diff --git a/Engine/Plugins/Experimental/UAF/UAF/Source/UAF/Private/UAF/AbstractSkeleton/AbstractSkeletonSetCollection.cpp b/Engine/Plugins/Experimental/UAF/UAF/Source/UAF/Private/UAF/AbstractSkeleton/AbstractSkeletonSetCollection.cpp
--- a/Engine/Plugins/Experimental/UAF/UAF/Source/UAF/Private/UAF/AbstractSkeleton/AbstractSkeletonSetCollection.cpp
+++ b/Engine/Plugins/Experimental/UAF/UAF/Source/UAF/Private/UAF/AbstractSkeleton/AbstractSkeletonSetCollection.cpp
@@ -18,6 +18,10 @@ bool UAbstractSkeletonSetCollection::AddSet(const FName InSetName, const FName I
 
 	SortSetHierarchy();
 
+#if WITH_EDITOR
+	OnSetsChanged.Broadcast();
+#endif
+
 	return true;
 }
 
@@ -45,6 +49,10 @@ bool UAbstractSkeletonSetCollection::RenameSet(const FName OldSetName, const FNa
 		}
 	}
 
+#if WITH_EDITOR
+	OnSetsChanged.Broadcast();
+#endif
+
 	return true;
 }
 
@@ -74,6 +82,10 @@ bool UAbstractSkeletonSetCollection::RemoveSet(const FName SetName)
 
 	SetHierarchy.RemoveAt(SetIndex);
 
+#if WITH_EDITOR
+	OnSetsChanged.Broadcast();
+#endif
+
 	return true;
 }
 
@@ -90,6 +102,10 @@ bool UAbstractSkeletonSetCollection::ReparentSet(const FName SetName, const FNam
 		SetHierarchy[SetIndex].ParentSetName = NewParentName;
 		SortSetHierarchy();
 
+#if WITH_EDITOR
+		OnSetsChanged.Broadcast();
+#endif
+
 		return true;
 	}
 
@@ -161,3 +177,17 @@ FAbstractSkeletonSet& UAbstractSkeletonSetCollection::GetSetRef(const FName SetN
 	check(SetIndex != INDEX_NONE)
 	return SetHierarchy[SetIndex];
 }
+
+#if WITH_EDITOR
+
+FDelegateHandle UAbstractSkeletonSetCollection::RegisterOnSetsChanged(FSimpleMulticastDelegate::FDelegate&& InDelegate)
+{
+	return OnSetsChanged.Add(InDelegate);
+}
+
+bool UAbstractSkeletonSetCollection::UnregisterOnSetsChanged(const FDelegateHandle InHandle)
+{
+	return OnSetsChanged.Remove(InHandle);
+}
+
+#endif // WITH_EDITOR
\ No newline at end of file
diff --git a/Engine/Plugins/Experimental/UAF/UAF/Source/UAF/Public/UAF/AbstractSkeleton/AbstractSkeletonSetCollection.h b/Engine/Plugins/Experimental/UAF/UAF/Source/UAF/Public/UAF/AbstractSkeleton/AbstractSkeletonSetCollection.h
--- a/Engine/Plugins/Experimental/UAF/UAF/Source/UAF/Public/UAF/AbstractSkeleton/AbstractSkeletonSetCollection.h
+++ b/Engine/Plugins/Experimental/UAF/UAF/Source/UAF/Public/UAF/AbstractSkeleton/AbstractSkeletonSetCollection.h
@@ -32,6 +32,7 @@ public:
 
 	UE_API bool RenameSet(const FName OldSetName, const FName NewSetName);
 
+	// Removes this set and moves its children to this set's parent
 	UE_API bool RemoveSet(const FName SetName);
 
 	UE_API bool ReparentSet(const FName SetName, const FName NewParentName);
@@ -41,11 +42,16 @@ public:
 
 	FName GetParentSetName(const FName SetName) const;
 
-private:
 	// Returns true if SetName is a descendant of AncestorSetName.
 	// This is for strict descendants, i.e. a set is not a descendant of itself.
-	bool IsDescendantOf(const FName SetName, const FName AncestorSetName) const;
+	UE_API bool IsDescendantOf(const FName SetName, const FName AncestorSetName) const;
+	
+#if WITH_EDITOR
+	UE_API FDelegateHandle RegisterOnSetsChanged(FSimpleMulticastDelegate::FDelegate&& InDelegate);
+	UE_API bool UnregisterOnSetsChanged(const FDelegateHandle InHandle);
+#endif
 
+private:
 	int32 GetSetIndex(const FName SetName) const;
 
 	void SortSetHierarchy();
@@ -53,7 +59,11 @@ private:
 	FAbstractSkeletonSet& GetSetRef(const FName SetName);
 
 	UPROPERTY()
-	TArray<FAbstractSkeletonSet> SetHierarchy; 
+	TArray<FAbstractSkeletonSet> SetHierarchy;
+
+#if WITH_EDITOR
+	FSimpleMulticastDelegate OnSetsChanged;
+#endif
 };
 
 #undef UE_API
diff --git a/Engine/Plugins/Experimental/UAF/UAF/Source/UAFEditor/Private/UAF/AbstractSkeleton/Sets/SAttributesList.cpp b/Engine/Plugins/Experimental/UAF/UAF/Source/UAFEditor/Private/UAF/AbstractSkeleton/Sets/SAttributesList.cpp
--- a/Engine/Plugins/Experimental/UAF/UAF/Source/UAFEditor/Private/UAF/AbstractSkeleton/Sets/SAttributesList.cpp
+++ b/Engine/Plugins/Experimental/UAF/UAF/Source/UAFEditor/Private/UAF/AbstractSkeleton/Sets/SAttributesList.cpp
@@ -9,13 +9,14 @@
 #include "PersonaModule.h"
 #include "SPositiveActionButton.h"
 #include "ScopedTransaction.h"
-#include "SetDragDrop.h"
+#include "UAF/AbstractSkeleton/Sets/SetDragDrop.h"
 #include "UAF/AbstractSkeleton/AbstractSkeletonSetBinding.h"
 #include "ToolMenus.h"
 #include "Widgets/Input/SSearchBox.h"
 
 #define LOCTEXT_NAMESPACE "UE::UAF::Sets::SAttributesList"
 
+// TODO: Move to UE::UAF::Editor
 using namespace UE::UAF::Sets;
 
 FName SAttributesList::FColumns::SetId = FName("Set");
@@ -214,7 +215,7 @@ TSharedRef<ITableRow> SAttributesList::ListView_OnGenerateRow(FListItemPtr Item,
 			{
 				TOptional<EItemDropZone> ReturnedDropZone;
 
-				const TSharedPtr<FSetDragDropOp> DragDropOp = DragDropEvent.GetOperationAs<FSetDragDropOp>();
+				const TSharedPtr<Editor::FSetDragDropOp> DragDropOp = DragDropEvent.GetOperationAs<Editor::FSetDragDropOp>();
 				if (DragDropOp.IsValid())
 				{
 					ReturnedDropZone = EItemDropZone::OntoItem;
@@ -224,14 +225,18 @@ TSharedRef<ITableRow> SAttributesList::ListView_OnGenerateRow(FListItemPtr Item,
 			})
 		.OnAcceptDrop_Lambda([this](const FDragDropEvent& DragDropEvent, EItemDropZone DropZone, FListItemPtr TargetTreeItem) -> FReply
 			{
-				const TSharedPtr<FSetDragDropOp> DragDropOp = DragDropEvent.GetOperationAs<FSetDragDropOp>();
+				const TSharedPtr<Editor::FSetDragDropOp> DragDropOp = DragDropEvent.GetOperationAs<Editor::FSetDragDropOp>();
 
 				if (DragDropOp.IsValid())
 				{
 					const FScopedTransaction Transaction(LOCTEXT("AddAttributeToSet", "Add Attribute to Set"));
 					SetBinding->Modify();
 
-					SetBinding->AddAttributeToSet(TargetTreeItem->Attribute, DragDropOp->Item->Set.SetName);
+					// TODO: Handle all dragged sets
+					if (!DragDropOp->GetDraggedSets().IsEmpty())
+					{
+						SetBinding->AddAttributeToSet(TargetTreeItem->Attribute, DragDropOp->GetDraggedSets()[0].SetName);
+					}
 					
 					RepopulateListData();
 
diff --git a/Engine/Plugins/Experimental/UAF/UAF/Source/UAFEditor/Private/UAF/AbstractSkeleton/Sets/SAttributesList.h b/Engine/Plugins/Experimental/UAF/UAF/Source/UAFEditor/Private/UAF/AbstractSkeleton/Sets/SAttributesList.h
--- a/Engine/Plugins/Experimental/UAF/UAF/Source/UAFEditor/Private/UAF/AbstractSkeleton/Sets/SAttributesList.h
+++ b/Engine/Plugins/Experimental/UAF/UAF/Source/UAFEditor/Private/UAF/AbstractSkeleton/Sets/SAttributesList.h
@@ -11,6 +11,7 @@
 
 class UAbstractSkeletonSetBinding;
 
+// TODO: Move to UE::UAF::Editor
 namespace UE::UAF::Sets
 {
 
diff --git a/Engine/Plugins/Experimental/UAF/UAF/Source/UAFEditor/Private/UAF/AbstractSkeleton/Sets/SSetBinding.cpp b/Engine/Plugins/Experimental/UAF/UAF/Source/UAFEditor/Private/UAF/AbstractSkeleton/Sets/SSetBinding.cpp
--- a/Engine/Plugins/Experimental/UAF/UAF/Source/UAFEditor/Private/UAF/AbstractSkeleton/Sets/SSetBinding.cpp
+++ b/Engine/Plugins/Experimental/UAF/UAF/Source/UAFEditor/Private/UAF/AbstractSkeleton/Sets/SSetBinding.cpp
@@ -6,7 +6,7 @@
 #include "Framework/Commands/UICommandList.h"
 #include "Framework/MultiBox/MultiBoxBuilder.h"
 #include "ScopedTransaction.h"
-#include "SetDragDrop.h"
+#include "UAF/AbstractSkeleton/Sets/SetDragDrop.h"
 #include "SkeletonTreeDragDrop.h"
 #include "UAF/AbstractSkeleton/AbstractSkeletonSetBinding.h"
 #include "UAF/AbstractSkeleton/AbstractSkeletonSetCollection.h"
@@ -14,6 +14,7 @@
 
 #define LOCTEXT_NAMESPACE "UE::UAF::Sets::SSetBinding"
 
+// TODO: Move to UE::UAF::Editor
 namespace UE::UAF::Sets
 {
 void SSetBinding::Construct(const FArguments& InArgs, TWeakObjectPtr<UAbstractSkeletonSetBinding> InSetBinding)
@@ -255,7 +256,7 @@ void SSetBinding::Construct(const FArguments& InArgs, TWeakObjectPtr<UAbstractSk
 				{
 					if (MouseEvent.IsMouseButtonDown(EKeys::LeftMouseButton))
 					{
-						const TSharedRef<FSetDragDropOp> DragDropOp = FSetDragDropOp::New(this);
+						const TSharedRef<Editor::FSetDragDropOp> DragDropOp = Editor::FSetDragDropOp::New({ Set });
 						return FReply::Handled().BeginDragDrop(DragDropOp);
 					}
 
diff --git a/Engine/Plugins/Experimental/UAF/UAF/Source/UAFEditor/Private/UAF/AbstractSkeleton/Sets/SSetCollection.cpp b/Engine/Plugins/Experimental/UAF/UAF/Source/UAFEditor/Private/UAF/AbstractSkeleton/Sets/SSetCollection.cpp
new file mode 100644
--- /dev/null
+++ b/Engine/Plugins/Experimental/UAF/UAF/Source/UAFEditor/Private/UAF/AbstractSkeleton/Sets/SSetCollection.cpp
@@ -0,0 +1,75 @@
+// Copyright Epic Games, Inc. All Rights Reserved.
+
+#include "UAF/AbstractSkeleton/Sets/SSetCollection.h"
+
+#include "SetDragDrop.h"
+#include "SPositiveActionButton.h"
+#include "Framework/Commands/GenericCommands.h"
+#include "Widgets/Text/SInlineEditableTextBlock.h"
+#include "Widgets/Input/SSearchBox.h"
+
+#define LOCTEXT_NAMESPACE "UE::UAF::Editor::SSetCollection"
+
+namespace UE::UAF::Editor
+{	
+	void SSetCollection::Construct(const FArguments& InArgs, TWeakObjectPtr<UAbstractSkeletonSetCollection> InSetCollection)
+	{
+		SetCollection = InSetCollection;
+		OnSetsChangedHandle = SetCollection->RegisterOnSetsChanged(FSimpleMulticastDelegate::FDelegate::CreateSP(this, &SSetCollection::OnSetsChanged));
+
+		ChildSlot
+		[
+			SNew(SVerticalBox)
+			+ SVerticalBox::Slot()
+			.AutoHeight()
+			.Padding(FMargin(1.0f, 1.0f))
+			[
+				SNew(SHorizontalBox)
+				+ SHorizontalBox::Slot()
+				.AutoWidth()
+				.Padding(FMargin(1.0f, 1.0f))
+				[
+					SNew(SPositiveActionButton)
+					.OnClicked_Lambda([this]()
+						{
+							TreeView->HandleAddSet();
+							return FReply::Handled();
+						})
+					.Icon(FAppStyle::Get().GetBrush("Icons.Plus"))
+					.Text(LOCTEXT("AddButton_Text", "Add Set"))
+				]
+				+ SHorizontalBox::Slot()
+				.FillWidth(1.0f)
+				.Padding(FMargin(1.0f, 1.0f))
+				[
+					SAssignNew(SearchBox, SSearchBox)
+					.SelectAllTextWhenFocused(true)
+					.OnTextChanged_Lambda([this](const FText& InText)
+					{
+						TreeView->SetFilterText(InText);
+					})
+					.HintText(LOCTEXT("SearchBox_Hint", "Search Sets..."))
+				]
+			]
+			+ SVerticalBox::Slot()
+			.FillHeight(1.0f)
+			.Padding(FMargin(1.0f, 1.0f))
+			[
+				SAssignNew(TreeView, SSetCollectionTreeView, SetCollection)
+			]
+		];
+	}
+
+	SSetCollection::~SSetCollection()
+	{
+		SetCollection->UnregisterOnSetsChanged(OnSetsChangedHandle);
+	}
+	
+	void SSetCollection::OnSetsChanged() const
+	{
+		SearchBox->SetSearchText(FText::GetEmpty());
+		TreeView->SetFilterText(FText::GetEmpty());
+	}
+}
+
+#undef LOCTEXT_NAMESPACE
\ No newline at end of file
diff --git a/Engine/Plugins/Experimental/UAF/UAF/Source/UAFEditor/Private/UAF/AbstractSkeleton/Sets/SSetCollection.h b/Engine/Plugins/Experimental/UAF/UAF/Source/UAFEditor/Private/UAF/AbstractSkeleton/Sets/SSetCollection.h
new file mode 100644
--- /dev/null
+++ b/Engine/Plugins/Experimental/UAF/UAF/Source/UAFEditor/Private/UAF/AbstractSkeleton/Sets/SSetCollection.h
@@ -0,0 +1,32 @@
+// Copyright Epic Games, Inc. All Rights Reserved.
+
+#pragma once
+
+#include "CoreMinimal.h"
+#include "EditorUndoClient.h"
+#include "UAF/AbstractSkeleton/AbstractSkeletonSetCollection.h"
+#include "UAF/AbstractSkeleton/Sets/SSetCollectionTreeView.h"
+#include "Widgets/SCompoundWidget.h"
+
+class SSearchBox;
+
+namespace UE::UAF::Editor
+{
+	class SSetCollection : public SCompoundWidget, public FSelfRegisteringEditorUndoClient
+	{
+		SLATE_BEGIN_ARGS(SSetCollection) {}
+		SLATE_END_ARGS()
+
+		void Construct(const FArguments& InArgs, TWeakObjectPtr<UAbstractSkeletonSetCollection> InSetCollection);
+		virtual ~SSetCollection() override;
+
+		void OnSetsChanged() const;
+		
+	private:
+		TWeakObjectPtr<UAbstractSkeletonSetCollection> SetCollection;
+		TSharedPtr<SSetCollectionTreeView> TreeView;
+		TSharedPtr<SSearchBox> SearchBox;
+
+		FDelegateHandle OnSetsChangedHandle;
+	};
+}
diff --git a/Engine/Plugins/Experimental/UAF/UAF/Source/UAFEditor/Private/UAF/AbstractSkeleton/Sets/SSetCollectionTreeView.cpp b/Engine/Plugins/Experimental/UAF/UAF/Source/UAFEditor/Private/UAF/AbstractSkeleton/Sets/SSetCollectionTreeView.cpp
new file mode 100644
--- /dev/null
+++ b/Engine/Plugins/Experimental/UAF/UAF/Source/UAFEditor/Private/UAF/AbstractSkeleton/Sets/SSetCollectionTreeView.cpp
@@ -0,0 +1,507 @@
+// Copyright Epic Games, Inc. All Rights Reserved.
+
+#include "UAF/AbstractSkeleton/Sets/SSetCollectionTreeView.h"
+
+#include "Framework/Commands/GenericCommands.h"
+#include "Framework/Commands/UICommandList.h"
+#include "Framework/MultiBox/MultiBoxBuilder.h"
+#include "ScopedTransaction.h"
+#include "SetDragDrop.h"
+#include "Widgets/Text/SInlineEditableTextBlock.h"
+
+#define LOCTEXT_NAMESPACE "UE::UAF::Editor::SSetCollectionTreeView"
+
+namespace UE::UAF::Editor
+{
+	void SSetCollectionTreeView::Construct(const FArguments& InArgs, TWeakObjectPtr<UAbstractSkeletonSetCollection> InSetCollection)
+	{
+		SetCollection = InSetCollection;
+		OnSetsChangedHandle = SetCollection->RegisterOnSetsChanged(FSimpleMulticastDelegate::FDelegate::CreateSP(this, &SSetCollectionTreeView::OnSetsChanged));
+		
+		STreeView<FSetCollectionTreeItemPtr>::FArguments SuperArgs;
+		SuperArgs.SelectionMode(ESelectionMode::Multi);
+		SuperArgs.HighlightParentNodesForSelection(true);
+		SuperArgs.TreeItemsSource(&RootItems);
+		SuperArgs.OnGenerateRow(this, &SSetCollectionTreeView::TreeView_OnGenerateRow);
+		SuperArgs.OnGetChildren(this, &SSetCollectionTreeView::TreeView_OnGetChildren);
+		SuperArgs.OnContextMenuOpening(this, &SSetCollectionTreeView::TreeView_OnContextMenuOpening);
+		SuperArgs.OnItemScrolledIntoView(this, &SSetCollectionTreeView::TreeView_OnItemScrolledIntoView);
+		SuperArgs.OnMouseButtonDoubleClick(this, &SSetCollectionTreeView::TreeView_OnMouseButtonDoubleClick);
+		SuperArgs.HeaderRow
+		(
+			SNew(SHeaderRow)
+			+ SHeaderRow::Column(FName("SetHierarchy"))
+			.FillWidth(1.0f)
+			.DefaultLabel(LOCTEXT("TreeView_Header", "Set Hierarchy"))
+		);
+
+		STreeView<FSetCollectionTreeItemPtr>::Construct(SuperArgs);
+		
+		BindCommands();
+	}
+
+	SSetCollectionTreeView::~SSetCollectionTreeView()
+	{
+		SetCollection->UnregisterOnSetsChanged(OnSetsChangedHandle);
+	}
+
+	TSharedRef<ITableRow> SSetCollectionTreeView::TreeView_OnGenerateRow(FSetCollectionTreeItemPtr InItem, const TSharedRef<STableViewBase>& OwnerTable)
+	{
+		TSharedPtr<SInlineEditableTextBlock> InlineWidget;
+
+		TSharedRef<STableRow<FSetCollectionTreeItemPtr>> RowWidget = SNew(STableRow<FSetCollectionTreeItemPtr>, OwnerTable)
+			.Visibility_Lambda([InItem]()
+			{
+				return InItem->bVisible ? EVisibility::Visible : EVisibility::Collapsed;
+			})
+			.Style(&FAppStyle().GetWidgetStyle<FTableRowStyle>("TableView.AlternatingRow"))
+			.OnDragDetected_Lambda([this](const FGeometry& MyGeometry, const FPointerEvent& MouseEvent)
+				{
+					if (MouseEvent.IsMouseButtonDown(EKeys::LeftMouseButton))
+					{
+						TArray<FAbstractSkeletonSet> Sets;
+						for (const auto& Selected : GetSelectedItems())
+						{
+							Sets.Add(Selected->Set);
+						}
+						const TSharedRef<FSetDragDropOp> DragDropOp = FSetDragDropOp::New(MoveTemp(Sets));
+						return FReply::Handled().BeginDragDrop(DragDropOp);
+					}
+
+					return FReply::Unhandled();
+				})
+			.OnDragEnter_Lambda([InItem, this](const FDragDropEvent& DragDropEvent)
+			{
+				const TSharedPtr<FSetDragDropOp> DragDropOp = DragDropEvent.GetOperationAs<FSetDragDropOp>();
+				if (DragDropOp.IsValid())
+				{
+					DragDropOp->SetOperation(FSetDragDropOp::EOperation::Reparent, InItem->Set.SetName);
+				}
+			})
+			.OnDragLeave_Lambda([this](const FDragDropEvent& DragDropEvent)
+			{
+				const TSharedPtr<FSetDragDropOp> DragDropOp = DragDropEvent.GetOperationAs<FSetDragDropOp>();
+				if (DragDropOp.IsValid())
+				{
+					DragDropOp->ClearOperation();
+				}
+			})
+			.OnCanAcceptDrop_Lambda([this](const FDragDropEvent& DragDropEvent, EItemDropZone DropZone, FSetCollectionTreeItemPtr TargetTreeItem) -> TOptional<EItemDropZone>
+			{
+				using namespace UE::UAF;
+
+				TOptional<EItemDropZone> ReturnedDropZone;
+
+				const TSharedPtr<FSetDragDropOp> DragDropOp = DragDropEvent.GetOperationAs<FSetDragDropOp>();
+				if (DragDropOp.IsValid())
+				{
+					bool bAnyDraggedSetAncestorOfTarget = false;
+
+					for (const FAbstractSkeletonSet& DraggedSet : DragDropOp->GetDraggedSets())
+					{
+						bAnyDraggedSetAncestorOfTarget |= SetCollection->IsDescendantOf(TargetTreeItem->Set.SetName, DraggedSet.SetName);
+					}
+
+					if (!bAnyDraggedSetAncestorOfTarget)
+					{
+						ReturnedDropZone = EItemDropZone::OntoItem;
+					}
+				}
+
+				return ReturnedDropZone;
+			})
+			.OnAcceptDrop_Lambda([this](const FDragDropEvent& DragDropEvent, EItemDropZone DropZone, FSetCollectionTreeItemPtr TargetTreeItem) -> FReply
+				{
+					const TSharedPtr<FSetDragDropOp> DragDropOp = DragDropEvent.GetOperationAs<FSetDragDropOp>();
+					if (DragDropOp.IsValid())
+					{
+						const FScopedTransaction Transaction(LOCTEXT("ReparentSets", "Reparent Set(s)"));
+						SetCollection->Modify();
+
+						for (const FAbstractSkeletonSet& Set : DragDropOp->GetDraggedSets())
+						{
+							SetCollection->ReparentSet(Set.SetName, TargetTreeItem->Set.SetName);
+						}
+
+						return FReply::Handled();
+					}
+
+					return FReply::Unhandled();
+				})
+			.ShowWires(true)
+			[
+				SNew(SHorizontalBox)
+				+ SHorizontalBox::Slot()
+				.AutoWidth()
+				.Padding(3.0f)
+				.VAlign(VAlign_Center)
+				[
+					SNew(SImage)
+					.Image(FAppStyle::Get().GetBrush("ClassIcon.GroupActor"))
+				]
+				+ SHorizontalBox::Slot()
+				.AutoWidth()
+				.Padding(3.0f)
+				[
+					SAssignNew(InlineWidget, SInlineEditableTextBlock)
+					.HighlightText_Lambda([this]()
+						{
+							return FilterText;
+						})
+					.Text(FText::FromName(InItem->Set.SetName))
+					.OnVerifyTextChanged_Lambda([InItem, this](const FText& InText, FText& OutErrorMessage)
+						{
+							const FName OldName = InItem->Set.SetName;
+							const FName NewName = FName(InText.ToString());
+
+							if (OldName == NewName)
+							{
+								return true;
+							}
+
+							if (CanRenameSet(OldName, NewName))
+							{
+								return true;
+							}
+								
+							OutErrorMessage = LOCTEXT("SetNameTaken", "Name already taken");
+							return false;
+						})
+					.OnTextCommitted_Lambda([InItem, this](const FText& InText, ETextCommit::Type CommitInfo)
+						{
+							const FName OldName = InItem->Set.SetName;
+							const FName NewName = FName(InText.ToString());
+
+							RenameSet(OldName, NewName);	
+						})
+				]
+			];
+
+		InItem->OnRequestRename.BindSP(InlineWidget.Get(), &SInlineEditableTextBlock::EnterEditingMode);
+
+		return RowWidget;
+	}
+
+	void SSetCollectionTreeView::TreeView_OnGetChildren(FSetCollectionTreeItemPtr InItem, TArray<FSetCollectionTreeItemPtr>& OutChildren)
+	{
+		OutChildren = InItem->Children;
+	}
+
+	TSharedPtr<SWidget> SSetCollectionTreeView::TreeView_OnContextMenuOpening()
+	{
+		FMenuBuilder MenuBuilder(true, CommandList.ToSharedRef());
+
+		MenuBuilder.AddMenuEntry(
+			LOCTEXT("AddSet_Label", "Add Set"),
+			LOCTEXT("AddSet_Tooltip", "Add a new set to the hierarchy"),
+			FSlateIcon(FAppStyle::GetAppStyleSetName(), "Icons.Plus"),
+			FUIAction(FExecuteAction::CreateSP(this, &SSetCollectionTreeView::HandleAddSet)));
+		
+		MenuBuilder.AddMenuEntry(FGenericCommands::Get().Rename);
+		MenuBuilder.AddMenuEntry(FGenericCommands::Get().Delete);
+
+		return MenuBuilder.MakeWidget();
+	}
+
+	void SSetCollectionTreeView::TreeView_OnItemScrolledIntoView(FSetCollectionTreeItemPtr InItem, const TSharedPtr<ITableRow>& InWidget)
+	{
+		if (RequestedSetItemToRename)
+		{
+			RequestedSetItemToRename->OnRequestRename.ExecuteIfBound();
+			RequestedSetItemToRename = nullptr;
+		}
+	}
+
+	void SSetCollectionTreeView::TreeView_OnMouseButtonDoubleClick(FSetCollectionTreeItemPtr InItem)
+	{
+		InItem->OnRequestRename.ExecuteIfBound();
+	}
+
+	void SSetCollectionTreeView::PostUndo(bool bSuccess)
+	{
+		bNeedsRebuilding = true;
+	}
+
+	void SSetCollectionTreeView::PostRedo(bool bSuccess)
+	{
+		bNeedsRebuilding = true;
+	}
+
+	void SSetCollectionTreeView::RenameSet(const FName OldSetName, const FName NewSetName)
+	{
+		SetCollection->RenameSet(OldSetName, NewSetName);
+	}
+
+	bool SSetCollectionTreeView::CanRenameSet(const FName OldSetName, const FName NewSetName) const
+	{
+		return !SetCollection->HasSet(NewSetName);
+	}
+
+	FSetCollectionTreeItemPtr SSetCollectionTreeView::GetTreeItem(const FName SetName)
+	{
+		for (FSetCollectionTreeItemPtr TreeItem : GetAllTreeItems())
+		{
+			if (TreeItem->Set.SetName == SetName)
+			{
+				return TreeItem;
+			}
+		}
+
+		return nullptr;
+	}
+
+	void SSetCollectionTreeView::HandleRemoveSet()
+	{
+		const TArray<FSetCollectionTreeItemPtr> Selection = GetSelectedItems();
+		if (!Selection.IsEmpty())
+		{
+			const FScopedTransaction Transaction(LOCTEXT("RemoveSets", "Remove Set(s)"));
+			SetCollection->Modify();
+
+			for (const FSetCollectionTreeItemPtr& Item : Selection)
+			{
+				const FName SetToRemove = Item->Set.SetName;
+				ensure(SetCollection->RemoveSet(SetToRemove));
+			}
+		}
+	}
+
+	void SSetCollectionTreeView::HandleAddSet()
+	{
+		const FScopedTransaction Transaction(LOCTEXT("AddSet", "Add Set"));
+		SetCollection->Modify();
+
+
+		const FText DefaultSetNameFormat = LOCTEXT("DefaultSetName", "NewSet_{0}");
+
+		int32 Suffix = 0;
+		FName SetName = FName(FText::Format(DefaultSetNameFormat, Suffix).ToString());
+		while (SetCollection->HasSet(SetName))
+		{
+			++Suffix;
+			SetName = FName(FText::Format(DefaultSetNameFormat, Suffix).ToString());
+		}
+
+		TArray<FSetCollectionTreeItemPtr> Selection = GetSelectedItems();
+		const FName ParentSetName = (Selection.Num() == 1) ? Selection[0]->Set.SetName : NAME_None;
+		
+		SetCollection->AddSet(SetName, ParentSetName);
+
+		// Expand the parent set tree item so we can see the new set in the tree
+		FSetCollectionTreeItemPtr ParentTreeItem = GetTreeItem(ParentSetName);
+		if (ParentTreeItem)
+		{
+			SetItemExpansion(ParentTreeItem, true);
+		}
+
+		RequestedSetNameToRename = SetName;
+	}
+	
+	void SSetCollectionTreeView::RepopulateTreeData()
+	{
+		// Make note of the currently selected item, if any
+		TArray<FSetCollectionTreeItemPtr> Selected = GetSelectedItems();
+
+		// Make note of all tree items collapsed
+		TArray<FName> CollapsedSetNames;
+
+		for (FSetCollectionTreeItemPtr TreeItem : GetAllTreeItems())
+		{
+			if (!IsItemExpanded(TreeItem) && !TreeItem->Children.IsEmpty())
+			{
+				CollapsedSetNames.Add(TreeItem->Set.SetName);
+			}
+		}
+
+		// Rebuild tree data
+		RootItems.Reset();
+
+		TMap<FName, FSetCollectionTreeItemPtr> TreeItemMap;
+
+		for (const FAbstractSkeletonSet& Set : SetCollection->GetSetHierarchy())
+		{
+			FSetCollectionTreeItemPtr SetItem = MakeShared<FSetCollectionTreeItem>();
+			SetItem->Set = Set;
+
+			if (Set.ParentSetName == NAME_None)
+			{
+				RootItems.Add(SetItem);
+			}
+			else
+			{
+				if (!TreeItemMap.Contains(Set.ParentSetName))
+				{
+					FSetCollectionTreeItemPtr ParentItem = MakeShared<FSetCollectionTreeItem>();
+					ParentItem->Set = Set;
+
+					TreeItemMap.Add(Set.ParentSetName, ParentItem);
+				}
+
+				FSetCollectionTreeItemPtr ParentItem = TreeItemMap[Set.ParentSetName];
+				ParentItem->Children.Add(SetItem);
+
+			}
+
+			TreeItemMap.Add(Set.SetName, SetItem);
+		}
+
+		// Restore expanded state of tree items and selection
+		const FName SelectedSetName = Selected.IsEmpty() ? NAME_None : Selected[0]->Set.SetName;
+		for (FSetCollectionTreeItemPtr TreeItem : GetAllTreeItems())
+		{
+			if (!CollapsedSetNames.Contains(TreeItem->Set.SetName))
+			{
+				SetItemExpansion(TreeItem, true);
+			}
+
+			if (TreeItem->Set.SetName == SelectedSetName)
+			{
+				SetSelection(TreeItem);
+			}
+		}
+
+		UpdateTreeItemVisibility();
+		RequestTreeRefresh();
+	}
+	
+	FReply SSetCollectionTreeView::OnKeyDown(const FGeometry& MyGeometry, const FKeyEvent& InKeyEvent)
+	{
+		if (CommandList->ProcessCommandBindings(InKeyEvent))
+		{
+			return FReply::Handled();
+		}
+
+		return FReply::Unhandled();
+	}
+
+	TArray<FSetCollectionTreeItemPtr> SSetCollectionTreeView::GetAllTreeItems()
+	{
+		TArray<FSetCollectionTreeItemPtr> AllItems;
+		AllItems.Append(RootItems);
+
+		for (int32 Index = 0; Index < AllItems.Num(); ++Index)
+		{
+			AllItems.Append(AllItems[Index]->Children);
+		}
+
+		return AllItems;
+	}
+
+	FReply SSetCollectionTreeView::OnDragOver(const FGeometry& MyGeometry, const FDragDropEvent& DragDropEvent)
+	{
+		if (TSharedPtr<FSetDragDropOp> DragDropOp = DragDropEvent.GetOperationAs<FSetDragDropOp>())
+		{
+			DragDropOp->SetOperation(FSetDragDropOp::EOperation::Reparent, NAME_None);
+
+			return FReply::Handled();
+		}
+		return FReply::Unhandled();
+	}
+	
+	void SSetCollectionTreeView::OnDragLeave(const FDragDropEvent& DragDropEvent)
+	{
+		if (TSharedPtr<FSetDragDropOp> DragDropOp = DragDropEvent.GetOperationAs<FSetDragDropOp>())
+		{
+			DragDropOp->ClearOperation();
+		}
+	}
+	
+	FReply SSetCollectionTreeView::OnDrop(const FGeometry& MyGeometry, const FDragDropEvent& DragDropEvent)
+	{
+		if (TSharedPtr<FSetDragDropOp> DragDropOp = DragDropEvent.GetOperationAs<FSetDragDropOp>())
+		{
+			const FScopedTransaction Transaction(LOCTEXT("UnparentSets", "Unparent Set(s)"));
+			SetCollection->Modify();
+
+			for (const FAbstractSkeletonSet& Set : DragDropOp->GetDraggedSets())
+			{
+				SetCollection->ReparentSet(Set.SetName, NAME_None);
+			}
+					
+			return FReply::Handled();
+		}
+
+		return FReply::Unhandled();
+	}
+
+	void SSetCollectionTreeView::BindCommands()
+	{
+		CommandList = MakeShareable(new FUICommandList);
+
+		CommandList->MapAction(
+			FGenericCommands::Get().Rename,
+			FExecuteAction::CreateSP(this, &SSetCollectionTreeView::HandleRenameSet)
+		);
+	
+		CommandList->MapAction(
+			FGenericCommands::Get().Delete,
+			FExecuteAction::CreateSP(this, &SSetCollectionTreeView::HandleRemoveSet)
+		);
+	}
+
+	void SSetCollectionTreeView::HandleRenameSet()
+	{
+		TArray<FSetCollectionTreeItemPtr> Selected = GetSelectedItems();
+		if (!Selected.IsEmpty())
+		{
+			RequestedSetNameToRename = Selected[0]->Set.SetName;
+		}
+	}
+
+	void SSetCollectionTreeView::OnSetsChanged()
+	{
+		bNeedsRebuilding = true;
+	}
+
+	void SSetCollectionTreeView::Tick(const FGeometry& AllottedGeometry, const double InCurrentTime, const float InDeltaTime)
+	{
+		STreeView<FSetCollectionTreeItemPtr>::Tick(AllottedGeometry, InCurrentTime, InDeltaTime);
+
+		if (bNeedsRebuilding)
+		{
+			RepopulateTreeData();
+			bNeedsRebuilding = false;
+		}
+
+		if (RequestedSetNameToRename != NAME_None)
+		{
+			RequestedSetItemToRename = GetTreeItem(RequestedSetNameToRename);
+			check(RequestedSetItemToRename);
+			
+			RequestScrollIntoView(RequestedSetItemToRename);
+			RequestedSetNameToRename = NAME_None;
+		}
+	}
+
+	void SSetCollectionTreeView::UpdateTreeItemVisibility()
+	{
+		const FString FilterString = FilterText.ToString();
+		
+		const TFunction<void(FSetCollectionTreeItemPtr&)> UpdateItemVisibility = [&](FSetCollectionTreeItemPtr& Item)
+			{
+				bool bAnyChildVisible = false;
+				for (FSetCollectionTreeItemPtr& Child : Item->Children)
+				{
+					UpdateItemVisibility(Child);
+					bAnyChildVisible |= Child->bVisible;
+				}
+
+				bool bSelfVisible = bAnyChildVisible || Item->Set.SetName.ToString().Contains(FilterString, ESearchCase::Type::IgnoreCase);
+				Item->bVisible = bSelfVisible;
+			};
+
+		for (FSetCollectionTreeItemPtr& RootItem : RootItems)
+		{
+			UpdateItemVisibility(RootItem);
+		}
+	}
+	
+	void SSetCollectionTreeView::SetFilterText(const FText& InText)
+	{
+		FilterText = InText;
+
+		UpdateTreeItemVisibility();
+	}
+}
+
+#undef LOCTEXT_NAMESPACE
\ No newline at end of file
diff --git a/Engine/Plugins/Experimental/UAF/UAF/Source/UAFEditor/Private/UAF/AbstractSkeleton/Sets/SSetCollectionTreeView.h b/Engine/Plugins/Experimental/UAF/UAF/Source/UAFEditor/Private/UAF/AbstractSkeleton/Sets/SSetCollectionTreeView.h
new file mode 100644
--- /dev/null
+++ b/Engine/Plugins/Experimental/UAF/UAF/Source/UAFEditor/Private/UAF/AbstractSkeleton/Sets/SSetCollectionTreeView.h
@@ -0,0 +1,88 @@
+// Copyright Epic Games, Inc. All Rights Reserved.
+
+#pragma once
+
+#include "CoreMinimal.h"
+#include "EditorUndoClient.h"
+#include "UAF/AbstractSkeleton/AbstractSkeletonSetCollection.h"
+#include "Widgets/Views/STreeView.h"
+
+class FUICommandList;
+
+namespace UE::UAF::Editor
+{
+	struct FSetCollectionTreeItem;
+	using FSetCollectionTreeItemPtr = TSharedPtr<FSetCollectionTreeItem>;
+	
+	struct FSetCollectionTreeItem
+	{
+		FAbstractSkeletonSet Set;
+		TArray<FSetCollectionTreeItemPtr> Children;
+		FSimpleDelegate OnRequestRename;
+		bool bVisible = true;
+	};
+
+	class SSetCollectionTreeView : public STreeView<FSetCollectionTreeItemPtr>, public FSelfRegisteringEditorUndoClient
+	{
+		SLATE_BEGIN_ARGS(SSetCollectionTreeView)
+		{}
+		SLATE_END_ARGS()
+
+		void Construct(const FArguments& InArgs, TWeakObjectPtr<UAbstractSkeletonSetCollection> InSetCollection);
+		virtual ~SSetCollectionTreeView();
+
+		/** FSelfRegisteringEditorUndoClient interface */
+		virtual void PostUndo(bool bSuccess) override;
+		virtual void PostRedo(bool bSuccess) override;
+
+		void BindCommands();
+		void RepopulateTreeData();
+
+		void HandleAddSet();
+		void HandleRenameSet();
+		void HandleRemoveSet();
+
+		void OnSetsChanged();
+		
+		void ExpandAllTreeItems();
+		TArray<FSetCollectionTreeItemPtr> GetAllTreeItems();
+
+		TSharedRef<ITableRow> TreeView_OnGenerateRow(FSetCollectionTreeItemPtr InItem, const TSharedRef<STableViewBase>& OwnerTable);
+		void TreeView_OnGetChildren(FSetCollectionTreeItemPtr InItem, TArray<FSetCollectionTreeItemPtr>& OutChildren);
+		TSharedPtr<SWidget> TreeView_OnContextMenuOpening();
+		void TreeView_OnItemScrolledIntoView(FSetCollectionTreeItemPtr InItem, const TSharedPtr<ITableRow>& InWidget);
+		void TreeView_OnMouseButtonDoubleClick(FSetCollectionTreeItemPtr InItem);
+
+		bool CanRenameSet(const FName OldSetName, const FName NewSetName) const;
+		void RenameSet(const FName OldSetName, const FName NewSetName);
+
+		virtual FReply OnDrop(const FGeometry& MyGeometry, const FDragDropEvent& DragDropEvent) override;
+
+		void SetFilterText(const FText& InText);
+
+	private:
+		virtual void Tick(const FGeometry& AllottedGeometry, const double InCurrentTime, const float InDeltaTime) override;
+		virtual FReply OnDragOver(const FGeometry& MyGeometry, const FDragDropEvent& DragDropEvent);
+		virtual void OnDragLeave(const FDragDropEvent& DragDropEvent);
+		virtual FReply OnKeyDown(const FGeometry& MyGeometry, const FKeyEvent& InKeyEvent) override;
+		
+		FSetCollectionTreeItemPtr GetTreeItem(const FName SetName);
+
+		void UpdateTreeItemVisibility();
+		
+		TArray<FSetCollectionTreeItemPtr> RootItems;
+
+		TWeakObjectPtr<UAbstractSkeletonSetCollection> SetCollection;
+
+		FName RequestedSetNameToRename;
+		FSetCollectionTreeItemPtr RequestedSetItemToRename;
+
+		TSharedPtr<FUICommandList> CommandList;
+
+		FDelegateHandle OnSetsChangedHandle;
+
+		bool bNeedsRebuilding = false;
+
+		FText FilterText;
+	};
+}
\ No newline at end of file
diff --git a/Engine/Plugins/Experimental/UAF/UAF/Source/UAFEditor/Private/UAF/AbstractSkeleton/Sets/SSetSkeletonTree.cpp b/Engine/Plugins/Experimental/UAF/UAF/Source/UAFEditor/Private/UAF/AbstractSkeleton/Sets/SSetSkeletonTree.cpp
--- a/Engine/Plugins/Experimental/UAF/UAF/Source/UAFEditor/Private/UAF/AbstractSkeleton/Sets/SSetSkeletonTree.cpp
+++ b/Engine/Plugins/Experimental/UAF/UAF/Source/UAFEditor/Private/UAF/AbstractSkeleton/Sets/SSetSkeletonTree.cpp
@@ -13,6 +13,7 @@
 
 #define LOCTEXT_NAMESPACE "UE::UAF::Sets::SSetsSkeletonTree"
 
+// TODO: Move to UE::UAF::Editor
 namespace UE::UAF::Sets
 {
 	void SSetsSkeletonTree::Construct(const FArguments& InArgs)
@@ -166,7 +167,7 @@ namespace UE::UAF::Sets
 
 					TOptional<EItemDropZone> ReturnedDropZone;
 
-					const TSharedPtr<FSetDragDropOp> DragDropOp = DragDropEvent.GetOperationAs<FSetDragDropOp>();
+					const TSharedPtr<Editor::FSetDragDropOp> DragDropOp = DragDropEvent.GetOperationAs<Editor::FSetDragDropOp>();
 					if (DragDropOp.IsValid())
 					{
 						ReturnedDropZone = EItemDropZone::OntoItem;
@@ -176,7 +177,7 @@ namespace UE::UAF::Sets
 				})
 			.OnAcceptDrop_Lambda([this](const FDragDropEvent& DragDropEvent, EItemDropZone DropZone, TSharedPtr<FTreeItem> TargetTreeItem) -> FReply
 				{
-					const TSharedPtr<FSetDragDropOp> DragDropOp = DragDropEvent.GetOperationAs<FSetDragDropOp>();
+					const TSharedPtr<Editor::FSetDragDropOp> DragDropOp = DragDropEvent.GetOperationAs<Editor::FSetDragDropOp>();
 					if (DragDropOp.IsValid())
 					{
 						TSharedPtr<FTreeItem> TargetSetTreeItem = StaticCastSharedPtr<FTreeItem>(TargetTreeItem);
@@ -184,8 +185,11 @@ namespace UE::UAF::Sets
 						const FScopedTransaction Transaction(LOCTEXT("AddAttributeToSet", "Add Attribute to Set"));
 						SetBinding->Modify();
 
-						const bool bSuccess = SetBinding->AddBoneToSet(TargetSetTreeItem->BoneName, DragDropOp->Item->Set.SetName);
+						// TODO: Use all dragged sets
+						check(!DragDropOp->GetDraggedSets().IsEmpty());
+						const bool bSuccess = SetBinding->AddBoneToSet(TargetSetTreeItem->BoneName, DragDropOp->GetDraggedSets()[0].SetName);
 						check(bSuccess);
+						
 						RepopulateTreeData();
 
 						return FReply::Handled();
diff --git a/Engine/Plugins/Experimental/UAF/UAF/Source/UAFEditor/Private/UAF/AbstractSkeleton/Sets/SetCollectionAssetDefinition.cpp b/Engine/Plugins/Experimental/UAF/UAF/Source/UAFEditor/Private/UAF/AbstractSkeleton/Sets/SetCollectionAssetDefinition.cpp
--- a/Engine/Plugins/Experimental/UAF/UAF/Source/UAFEditor/Private/UAF/AbstractSkeleton/Sets/SetCollectionAssetDefinition.cpp
+++ b/Engine/Plugins/Experimental/UAF/UAF/Source/UAFEditor/Private/UAF/AbstractSkeleton/Sets/SetCollectionAssetDefinition.cpp
@@ -27,7 +27,7 @@ EAssetCommandResult UAssetDefinition_AbstractSkeletonSetCollection::OpenAssets(c
 	if (OpenArgs.OpenMethod == EAssetOpenMethod::Edit)
 	{
 		TArray<UObject*> Assets = OpenArgs.LoadObjects<UObject>();
-		MakeShared<UE::UAF::FSetCollectionEditorToolkit>()->InitEditor(Assets);
+		MakeShared<UE::UAF::Editor::FSetCollectionEditorToolkit>()->InitEditor(Assets);
 	}
 
 	return EAssetCommandResult::Handled;
diff --git a/Engine/Plugins/Experimental/UAF/UAF/Source/UAFEditor/Private/UAF/AbstractSkeleton/Sets/SetCollectionEditor.cpp b/Engine/Plugins/Experimental/UAF/UAF/Source/UAFEditor/Private/UAF/AbstractSkeleton/Sets/SetCollectionEditor.cpp
--- a/Engine/Plugins/Experimental/UAF/UAF/Source/UAFEditor/Private/UAF/AbstractSkeleton/Sets/SetCollectionEditor.cpp
+++ b/Engine/Plugins/Experimental/UAF/UAF/Source/UAFEditor/Private/UAF/AbstractSkeleton/Sets/SetCollectionEditor.cpp
@@ -2,20 +2,13 @@
 
 #include "UAF/AbstractSkeleton/Sets/SetCollectionEditor.h"
 
-#include "ToolMenus.h"
-#include "Widgets/Docking/SDockTab.h"
-#include "Modules/ModuleManager.h"
-#include "PropertyEditorModule.h"
-#include "SPositiveActionButton.h"
-#include "Widgets/Input/SSearchBox.h"
-#include "Framework/Commands/GenericCommands.h"
 #include "DragAndDrop/DecoratedDragDropOp.h"
-#include "ScopedTransaction.h"
-#include "Widgets/Text/SInlineEditableTextBlock.h"
+#include "UAF/AbstractSkeleton/Sets/SSetCollection.h"
+#include "Widgets/Docking/SDockTab.h"
 
-#define LOCTEXT_NAMESPACE "UE::UAF::FSetCollectionEditorToolkit"
+#define LOCTEXT_NAMESPACE "UE::UAF::Editor::FSetCollectionEditorToolkit"
 
-namespace UE::UAF
+namespace UE::UAF::Editor
 {
 	void FSetCollectionEditorToolkit::InitEditor(const TArray<UObject*>& InObjects)
 	{
@@ -28,18 +21,12 @@ namespace UE::UAF
 				->Split
 				(
 					FTabManager::NewStack()
-					->SetSizeCoefficient(0.7f)
+					->SetHideTabWell(true)
 					->AddTab("AbstractSkeletonSetCollection_SetTab", ETabState::OpenedTab)
 				)
-				->Split
-				(
-					FTabManager::NewStack()
-					->SetSizeCoefficient(0.3f)
-					->AddTab("AbstractSkeletonSetCollection_PropertiesTab", ETabState::OpenedTab)
-				)
 			);
 
-		FAssetEditorToolkit::InitAssetEditor(EToolkitMode::Standalone, {}, "AbstractSkeletonSetCollectionEditorToolkit", Layout, true, true, InObjects);
+		InitAssetEditor(EToolkitMode::Standalone, {}, "AbstractSkeletonSetCollectionEditorToolkit", Layout, true, true, InObjects);
 	}
 
 	void FSetCollectionEditorToolkit::RegisterTabSpawners(const TSharedRef<class FTabManager>& InTabManager)
@@ -47,483 +34,24 @@ namespace UE::UAF
 		FAssetEditorToolkit::RegisterTabSpawners(InTabManager);
 
 		WorkspaceMenuCategory = InTabManager->AddLocalWorkspaceMenuCategory(LOCTEXT("AbstractSkeletonSetCollectionEditor", "Abstract Skeleton Set Editor"));
-
-		// Set tab
-		{
-			InTabManager->RegisterTabSpawner("AbstractSkeletonSetCollection_SetTab", FOnSpawnTab::CreateLambda([this](const FSpawnTabArgs&)
-				{
-					return SNew(SDockTab)
-						[
-							SNew(SSetCollection, SetCollection)
-						];
-				}))
-				.SetDisplayName(LOCTEXT("SetTab_DisplayName", "Abstract Skeleton Set"))
-				.SetGroup(WorkspaceMenuCategory.ToSharedRef());
-		}
-
-		// Property tab
-		{
-			InTabManager->RegisterTabSpawner("AbstractSkeletonSetCollection_PropertiesTab", FOnSpawnTab::CreateSP(this, &FSetCollectionEditorToolkit::SpawnPropertiesTab))
-				.SetDisplayName(LOCTEXT("PropertiesTab_DisplayName", "Details"))
-				.SetGroup(WorkspaceMenuCategory.ToSharedRef());
-		}
+		
+		InTabManager->RegisterTabSpawner("AbstractSkeletonSetCollection_SetTab", FOnSpawnTab::CreateLambda([this](const FSpawnTabArgs&)
+			{
+				return SNew(SDockTab)
+					.CanEverClose(false)
+					[
+						SNew(SSetCollection, SetCollection)
+					];
+			}))
+			.SetDisplayName(LOCTEXT("SetTab_DisplayName", "Abstract Skeleton Set"))
+			.SetGroup(WorkspaceMenuCategory.ToSharedRef());
 	}
 
 	void FSetCollectionEditorToolkit::UnregisterTabSpawners(const TSharedRef<class FTabManager>& InTabManager)
 	{
 		FAssetEditorToolkit::UnregisterTabSpawners(InTabManager);
 		InTabManager->UnregisterTabSpawner("AbstractSkeletonSetCollection_SetTab");
-		InTabManager->UnregisterTabSpawner("AbstractSkeletonSetCollection_PropertiesTab");
-	}
-
-	TSharedRef<SDockTab> FSetCollectionEditorToolkit::SpawnPropertiesTab(const FSpawnTabArgs& Args)
-	{
-		FPropertyEditorModule& PropertyEditorModule = FModuleManager::GetModuleChecked<FPropertyEditorModule>("PropertyEditor");
-
-		FDetailsViewArgs DetailsViewArgs;
-		DetailsViewArgs.NameAreaSettings = FDetailsViewArgs::HideNameArea;
-		DetailsViewArgs.bAllowSearch = false;
-		DetailsViewArgs.bUpdatesFromSelection = false;
-		DetailsViewArgs.bShowOptions = false;
-		DetailsViewArgs.bShowObjectLabel = false;
-		DetailsViewArgs.bAllowFavoriteSystem = false;
-
-		TSharedRef<IDetailsView> DetailsView = PropertyEditorModule.CreateDetailView(DetailsViewArgs);
-		DetailsView->SetObjects(TArray<UObject*>{ SetCollection.Get() });
-
-		return SNew(SDockTab)
-			[
-				DetailsView
-			];
-	}
-
-	class FSetCollectionDragDropOp : public FDecoratedDragDropOp
-	{
-	public:
-		DRAG_DROP_OPERATOR_TYPE(FSetCollectionDragDropOp, FDecoratedDragDropOp)
-
-		static TSharedRef<FSetCollectionDragDropOp> New(const FName InSetName, const FName InParentSetName)
-		{
-			TSharedRef<FSetCollectionDragDropOp> Operation = MakeShared<FSetCollectionDragDropOp>();
-			Operation->SetName = InSetName;
-			Operation->ParentSetName = InParentSetName;
-			Operation->Construct();
-			return Operation;
-		}
-
-		TSharedPtr<SWidget> GetDefaultDecorator() const override
-		{
-			return SNew(SBorder)
-				.Visibility(EVisibility::Visible)
-				.BorderImage(FAppStyle::GetBrush("Menu.Background"))
-				[
-					SNew(STextBlock)
-						.Text(FText::FromName(SetName))
-				];
-		}
-
-		FName SetName;
-		FName ParentSetName;
-	};
-
-	void SSetCollection::Construct(const FArguments& InArgs, TWeakObjectPtr<UAbstractSkeletonSetCollection> InSetCollection)
-	{
-		SetCollection = InSetCollection;
-
-		ChildSlot
-		[
-			SNew(SVerticalBox)
-			+ SVerticalBox::Slot()
-			.AutoHeight()
-			.Padding(FMargin(1.0f, 1.0f))
-			[
-				SNew(SHorizontalBox)
-				+ SHorizontalBox::Slot()
-				.AutoWidth()
-				.Padding(FMargin(1.0f, 1.0f))
-				[
-					SNew(SPositiveActionButton)
-					.OnClicked(this, &SSetCollection::OnAddSet)
-					.Icon(FAppStyle::Get().GetBrush("Icons.Plus"))
-					.Text(LOCTEXT("AddButton_Text", "Add Set"))
-				]
-				+ SHorizontalBox::Slot()
-				.FillWidth(1.0f)
-				.Padding(FMargin(1.0f, 1.0f))
-				[
-					SAssignNew(SearchBox, SSearchBox)
-					.SelectAllTextWhenFocused(true)
-					.OnTextChanged_Lambda([](const FText& InText) { /* TODO: Implement */ })
-					.HintText(LOCTEXT("SearchBox_Hint", "Search Sets..."))
-				]
-			]
-			+ SVerticalBox::Slot()
-			.FillHeight(1.0f)
-			.Padding(FMargin(1.0f, 1.0f))
-			[
-				SAssignNew(TreeView, STreeView<TSharedPtr<FTreeItem>>)
-				.SelectionMode(ESelectionMode::Single)
-				.HighlightParentNodesForSelection(true)
-				.TreeItemsSource(&RootItems)
-				.OnGenerateRow(this, &SSetCollection::TreeView_OnGenerateRow)
-				.OnGetChildren(this, &SSetCollection::TreeView_OnGetChildren)
-				.OnContextMenuOpening(this, &SSetCollection::TreeView_OnContextMenuOpening)
-				.OnItemScrolledIntoView(this, &SSetCollection::TreeView_OnItemScrolledIntoView)
-				.OnMouseButtonDoubleClick(this, &SSetCollection::TreeView_OnMouseButtonDoubleClick)
-				.HeaderRow
-				(
-					SNew(SHeaderRow)
-					+ SHeaderRow::Column(FName("SetHierarchy"))
-					.FillWidth(0.5f)
-					.DefaultLabel(LOCTEXT("TreeView_Header", "Set Hierarchy"))
-				)
-			]
-		];
-
-		RepopulateTreeData();
-		ExpandAllTreeItems();
-		BindCommands();
-	}
-
-	FReply SSetCollection::OnKeyDown(const FGeometry& MyGeometry, const FKeyEvent& InKeyEvent)
-	{
-		if (CommandList->ProcessCommandBindings(InKeyEvent))
-		{
-			return FReply::Handled();
-		}
-
-		return FReply::Unhandled();
-	}
-
-	TSharedRef<ITableRow> SSetCollection::TreeView_OnGenerateRow(TSharedPtr<FTreeItem> InItem, const TSharedRef<STableViewBase>& OwnerTable)
-	{
-		TSharedPtr<SInlineEditableTextBlock> InlineWidget;
-
-		TSharedRef<STableRow<TSharedPtr<FTreeItem>>> RowWidget = SNew(STableRow<TSharedPtr<FTreeItem>>, OwnerTable)
-			.Style(&FAppStyle().GetWidgetStyle<FTableRowStyle>("TableView.AlternatingRow"))
-			.OnDragDetected_Lambda([InItem](const FGeometry& MyGeometry, const FPointerEvent& MouseEvent)
-				{
-					if (MouseEvent.IsMouseButtonDown(EKeys::LeftMouseButton))
-					{
-						const TSharedRef<FSetCollectionDragDropOp> DragDropOp = FSetCollectionDragDropOp::New(InItem->Name, InItem->ParentName);
-						return FReply::Handled().BeginDragDrop(DragDropOp);
-					}
-
-					return FReply::Unhandled();
-				})
-			.OnCanAcceptDrop_Lambda([](const FDragDropEvent& DragDropEvent, EItemDropZone DropZone, TSharedPtr<FTreeItem> TargetTreeItem) -> TOptional<EItemDropZone>
-				{
-					using namespace UE::UAF;
-
-					TOptional<EItemDropZone> ReturnedDropZone;
-
-					const TSharedPtr<FSetCollectionDragDropOp> DragDropOp = DragDropEvent.GetOperationAs<FSetCollectionDragDropOp>();
-					if (DragDropOp.IsValid())
-					{
-						ReturnedDropZone = EItemDropZone::OntoItem;
-					}
-
-					return ReturnedDropZone;
-				})
-			.OnAcceptDrop_Lambda([this](const FDragDropEvent& DragDropEvent, EItemDropZone DropZone, TSharedPtr<FTreeItem> TargetTreeItem) -> FReply
-				{
-					const TSharedPtr<FSetCollectionDragDropOp> DragDropOp = DragDropEvent.GetOperationAs<FSetCollectionDragDropOp>();
-					if (DragDropOp.IsValid())
-					{
-						const FScopedTransaction Transaction(LOCTEXT("ReparentSet", "Reparent Set"));
-						SetCollection->Modify();
-
-						SetCollection->ReparentSet(DragDropOp->SetName, TargetTreeItem->Name);
-						RepopulateTreeData();
-
-						return FReply::Handled();
-					}
-
-					return FReply::Unhandled();
-				})
-			.ShowWires(true)
-			[
-				SNew(SHorizontalBox)
-					+ SHorizontalBox::Slot()
-					.AutoWidth()
-					.Padding(3.0f)
-					[
-						SNew(SImage)
-							.Image(FAppStyle::Get().GetBrush("ClassIcon.GroupActor"))
-					]
-					+ SHorizontalBox::Slot()
-					.FillWidth(1.0)
-					.AutoWidth()
-					.Padding(3.0f)
-					[
-						SAssignNew(InlineWidget, SInlineEditableTextBlock)
-							.Text(FText::FromName(InItem->Name))
-							.OnVerifyTextChanged_Lambda([InItem](const FText& InText, FText& OutErrorMessage)
-							{
-								const FName OldName = InItem->Name;
-								const FName NewName = FName(InText.ToString());
-
-								if (InItem->Name == NewName)
-								{
-									return true;
-								}
-
-								check(InItem->CanRenameTo.IsBound());
-								if (InItem->CanRenameTo.Execute(OldName, NewName))
-								{
-									return true;
-								}
-								else
-								{
-									OutErrorMessage = LOCTEXT("SetNameTaken", "Name already taken");
-									return false;
-								}
-							})
-							.OnTextCommitted_Lambda([InItem](const FText& InText, ETextCommit::Type CommitInfo)
-							{
-								const FName OldName = InItem->Name;
-								const FName NewName = FName(InText.ToString());
-
-								InItem->OnRenamed.ExecuteIfBound(OldName, NewName);
-							})
-					]
-
-			];
-
-		InItem->OnRequestRename.BindSP(InlineWidget.Get(), &SInlineEditableTextBlock::EnterEditingMode);
-
-		return RowWidget;
-	}
-
-	void SSetCollection::TreeView_OnGetChildren(TSharedPtr<FTreeItem> InItem, TArray<TSharedPtr<FTreeItem>>& OutChildren)
-	{
-		OutChildren = InItem->Children;
-	}
-
-	TSharedPtr<SWidget> SSetCollection::TreeView_OnContextMenuOpening()
-	{
-		FMenuBuilder MenuBuilder(true, CommandList.ToSharedRef());
-
-		MenuBuilder.AddMenuEntry(FGenericCommands::Get().Rename);
-		MenuBuilder.AddMenuEntry(FGenericCommands::Get().Delete);
-
-		return MenuBuilder.MakeWidget();
-	}
-
-	void SSetCollection::TreeView_OnItemScrolledIntoView(TSharedPtr<FTreeItem> InItem, const TSharedPtr<ITableRow>& InWidget)
-	{
-		if (ItemToRename)
-		{
-			ItemToRename->OnRequestRename.ExecuteIfBound();
-			ItemToRename = nullptr;
-		}
-	}
-
-	void SSetCollection::TreeView_OnMouseButtonDoubleClick(TSharedPtr<FTreeItem> InItem)
-	{
-		InItem->OnRequestRename.ExecuteIfBound();
-	}
-
-	void SSetCollection::PostUndo(bool bSuccess)
-	{
-		RepopulateTreeData();
-	}
-
-	void SSetCollection::PostRedo(bool bSuccess)
-	{
-		RepopulateTreeData();
-	}
-
-	void SSetCollection::OnRemoveSet()
-	{
-		const TArray<TSharedPtr<FTreeItem>> SelectedItems = TreeView->GetSelectedItems();
-		if (!SelectedItems.IsEmpty())
-		{
-			const FScopedTransaction Transaction(LOCTEXT("RemoveSet", "Remove Set"));
-			SetCollection->Modify();
-
-			const FName SetToRemove = SelectedItems[0]->Name;
-			const bool bSuccess = SetCollection->RemoveSet(SetToRemove);
-			if (bSuccess)
-			{
-				RepopulateTreeData();
-			}
-		}
-	}
-
-	void SSetCollection::RepopulateTreeData()
-	{
-		// Make note of the currently selected item, if any
-		TArray<TSharedPtr<FTreeItem>> Selected = TreeView->GetSelectedItems();
-
-		// Make note of all tree items expanded
-		TArray<FName> ExpandedSetNames;
-
-		for (TSharedPtr<FTreeItem> TreeItem : GetAllTreeItems())
-		{
-			if (TreeView->IsItemExpanded(TreeItem))
-			{
-				ExpandedSetNames.Add(TreeItem->Name);
-			}
-		}
-
-		// Rebuild tree data
-		RootItems.Reset();
-
-		TMap<FName, TSharedPtr<FTreeItem>> TreeItemMap;
-
-		for (const FAbstractSkeletonSet& Set : SetCollection->GetSetHierarchy())
-		{
-			TSharedPtr<FTreeItem> SetItem = MakeShared<FTreeItem>();
-			SetItem->Name = Set.SetName;
-			SetItem->ParentName = Set.ParentSetName;
-			SetItem->OnRenamed = FTreeItem::FOnRenamed::CreateSP(this, &SSetCollection::OnItemRenamed);
-			SetItem->CanRenameTo = FTreeItem::FCanRenameTo::CreateSP(this, &SSetCollection::CanRenameItem);
-
-			if (Set.ParentSetName == NAME_None)
-			{
-				RootItems.Add(SetItem);
-			}
-			else
-			{
-				if (!TreeItemMap.Contains(Set.ParentSetName))
-				{
-					TSharedPtr<FTreeItem> ParentItem = MakeShared<FTreeItem>();
-					ParentItem->Name = Set.ParentSetName;
-					ParentItem->OnRenamed = FTreeItem::FOnRenamed::CreateSP(this, &SSetCollection::OnItemRenamed);
-					ParentItem->CanRenameTo = FTreeItem::FCanRenameTo::CreateSP(this, &SSetCollection::CanRenameItem);
-
-					TreeItemMap.Add(Set.ParentSetName, ParentItem);
-				}
-
-				TSharedPtr<FTreeItem> ParentItem = TreeItemMap[Set.ParentSetName];
-				ParentItem->Children.Add(SetItem);
-
-			}
-
-			TreeItemMap.Add(Set.SetName, SetItem);
-		}
-
-		// Restore expanded state of tree items and selection
-		const FName SelectedSetName = Selected.IsEmpty() ? NAME_None : Selected[0]->Name;
-		for (TSharedPtr<FTreeItem> TreeItem : GetAllTreeItems())
-		{
-			if (ExpandedSetNames.Contains(TreeItem->Name))
-			{
-				TreeView->SetItemExpansion(TreeItem, true);
-			}
-
-			if (TreeItem->Name == SelectedSetName)
-			{
-				TreeView->SetSelection(TreeItem);
-			}
-		}
-
-		TreeView->RequestTreeRefresh();
-	}
-
-	void SSetCollection::ExpandAllTreeItems()
-	{
-		for (TSharedPtr<FTreeItem> TreeItem : GetAllTreeItems())
-		{
-			TreeView->SetItemExpansion(TreeItem, true);
-		}
-	}
-
-	void SSetCollection::BindCommands()
-	{
-		CommandList = MakeShareable(new FUICommandList);
-
-		CommandList->MapAction(
-			FGenericCommands::Get().Rename,
-			FExecuteAction::CreateSP(this, &SSetCollection::OnRenameSet)
-		);
-	
-		CommandList->MapAction(
-			FGenericCommands::Get().Delete,
-			FExecuteAction::CreateSP(this, &SSetCollection::OnRemoveSet)
-		);
-	}
-
-	TArray<TSharedPtr<SSetCollection::FTreeItem>> SSetCollection::GetAllTreeItems()
-	{
-		TArray<TSharedPtr<FTreeItem>> AllItems;
-		AllItems.Append(RootItems);
-
-		for (int32 Index = 0; Index < AllItems.Num(); ++Index)
-		{
-			AllItems.Append(AllItems[Index]->Children);
-		}
-
-		return AllItems;
-	}
-
-	FReply SSetCollection::OnAddSet()
-	{
-		const FScopedTransaction Transaction(LOCTEXT("AddSet", "Set Set"));
-		SetCollection->Modify();
-
-		TArray<TSharedPtr<FTreeItem>> SelectedItems = TreeView->GetSelectedItems();
-
-		const FName ParentSetName = SelectedItems.IsEmpty() ? NAME_None : SelectedItems[0]->Name;
-		
-		const FText DefaultSetNameFormat = LOCTEXT("DefaultSetName", "NewSet_{0}");
-
-		int32 Suffix = 0;
-		FName SetName = FName(FText::Format(DefaultSetNameFormat, Suffix).ToString());
-		while (SetCollection->HasSet(SetName))
-		{
-			++Suffix;
-			SetName = FName(FText::Format(DefaultSetNameFormat, Suffix).ToString());
-		}
-
-		SetCollection->AddSet(SetName, ParentSetName);
-		RepopulateTreeData();
-
-		// Expand the parent set tree item so we can see the new set in the tree
-		TreeView->SetItemExpansion(GetTreeItem(ParentSetName), true);
-
-		TSharedPtr<FTreeItem> SetTreeItem = GetTreeItem(SetName);
-		ItemToRename = SetTreeItem;
-		TreeView->RequestScrollIntoView(SetTreeItem);
-
-		return FReply::Handled();
-	}
-
-	void SSetCollection::OnItemRenamed(const FName OldSetName, const FName NewSetName)
-	{
-		SetCollection->RenameSet(OldSetName, NewSetName);
-		RepopulateTreeData();
-	}
-
-	bool SSetCollection::CanRenameItem(const FName OldSetName, const FName NewSetName) const
-	{
-		return !SetCollection->HasSet(NewSetName);
-	}
-
-	TSharedPtr<SSetCollection::FTreeItem> SSetCollection::GetTreeItem(const FName SetName)
-	{
-		for (TSharedPtr<FTreeItem> TreeItem : GetAllTreeItems())
-		{
-			if (TreeItem->Name == SetName)
-			{
-				return TreeItem;
-			}
-		}
-
-		return nullptr;
-	}
-	
-	void SSetCollection::OnRenameSet()
-	{
-		TArray<TSharedPtr<FTreeItem>> Selected = TreeView->GetSelectedItems();
-		if (!Selected.IsEmpty())
-		{		
-			ItemToRename = Selected[0];
-        	TreeView->RequestScrollIntoView(ItemToRename);
-        }
-	}
-	
+	}	
 }
 
 #undef LOCTEXT_NAMESPACE
\ No newline at end of file
diff --git a/Engine/Plugins/Experimental/UAF/UAF/Source/UAFEditor/Private/UAF/AbstractSkeleton/Sets/SetCollectionEditor.h b/Engine/Plugins/Experimental/UAF/UAF/Source/UAFEditor/Private/UAF/AbstractSkeleton/Sets/SetCollectionEditor.h
--- a/Engine/Plugins/Experimental/UAF/UAF/Source/UAFEditor/Private/UAF/AbstractSkeleton/Sets/SetCollectionEditor.h
+++ b/Engine/Plugins/Experimental/UAF/UAF/Source/UAFEditor/Private/UAF/AbstractSkeleton/Sets/SetCollectionEditor.h
@@ -3,13 +3,10 @@
 #pragma once
 
 #include "CoreMinimal.h"
-#include "EditorUndoClient.h"
 #include "UAF/AbstractSkeleton/AbstractSkeletonSetCollection.h"
 #include "Toolkits/AssetEditorToolkit.h"
-#include "Widgets/SCompoundWidget.h"
-#include "Widgets/Views/STreeView.h"
 
-namespace UE::UAF
+namespace UE::UAF::Editor
 {
 	class FSetCollectionEditorToolkit : public FAssetEditorToolkit
 	{
@@ -27,83 +24,6 @@ namespace UE::UAF
 		// End FAssetEditorToolkit
 
 	private:
-		TSharedRef<SDockTab> SpawnPropertiesTab(const FSpawnTabArgs& Args);
-
-		TWeakObjectPtr<UAbstractSkeletonSetCollection> SetCollection;
-	};
-
-	class SSetCollection : public SCompoundWidget, public FSelfRegisteringEditorUndoClient
-	{
-		SLATE_BEGIN_ARGS(SSetCollection) {}
-		SLATE_END_ARGS()
-
-	public:
-		struct FTreeItem
-		{
-			FName Name;
-			FName ParentName;
-			TArray<TSharedPtr<FTreeItem>> Children;
-
-			DECLARE_DELEGATE(FOnRequestRename);
-			FOnRequestRename OnRequestRename;
-
-			DECLARE_DELEGATE_RetVal_TwoParams(bool, FCanRenameTo, const FName /* Old Name */, const FName /* New Name */);
-			FCanRenameTo CanRenameTo;
-
-			DECLARE_DELEGATE_TwoParams(FOnRenamed, const FName /* Old Name */, const FName /* New Name */);
-			FOnRenamed OnRenamed;
-		};
-
-		void Construct(const FArguments& InArgs, TWeakObjectPtr<UAbstractSkeletonSetCollection> InSetCollection);
-
-		/** FSelfRegisteringEditorUndoClient interface */
-		virtual void PostUndo(bool bSuccess) override;
-		virtual void PostRedo(bool bSuccess) override;
-
-	private:
-		virtual FReply OnKeyDown( const FGeometry& MyGeometry, const FKeyEvent& InKeyEvent) override;
-
-		TSharedRef<ITableRow> TreeView_OnGenerateRow(TSharedPtr<FTreeItem> InItem, const TSharedRef<STableViewBase>& OwnerTable);
-
-		void TreeView_OnGetChildren(TSharedPtr<FTreeItem> InItem, TArray<TSharedPtr<FTreeItem>>& OutChildren);
-
-		TSharedPtr<SWidget> TreeView_OnContextMenuOpening();
-
-		void TreeView_OnItemScrolledIntoView(TSharedPtr<FTreeItem> InItem, const TSharedPtr<ITableRow>& InWidget);
-		
-		void TreeView_OnMouseButtonDoubleClick(TSharedPtr<FTreeItem> InItem);
-
-		void OnRenameSet();
-		
-		void OnRemoveSet();
-
-		void RepopulateTreeData();
-
-		void ExpandAllTreeItems();
-
-		void BindCommands();
-
-		TArray<TSharedPtr<FTreeItem>> GetAllTreeItems();
-
-		FReply OnAddSet();
-
-		void OnItemRenamed(const FName OldSetName, const FName NewSetName);
-
-		bool CanRenameItem(const FName OldSetName, const FName NewSetName) const;
-
-	private:
-		TSharedPtr<FTreeItem> GetTreeItem(const FName SetName);
-
 		TWeakObjectPtr<UAbstractSkeletonSetCollection> SetCollection;
-
-		TSharedPtr<STreeView<TSharedPtr<FTreeItem>>> TreeView;
-
-		TArray<TSharedPtr<FTreeItem>> RootItems;
-
-		TSharedPtr<SSearchBox> SearchBox;
-
-		TSharedPtr<FUICommandList> CommandList;
-
-		TSharedPtr<FTreeItem> ItemToRename;
 	};
 }
diff --git a/Engine/Plugins/Experimental/UAF/UAF/Source/UAFEditor/Private/UAF/AbstractSkeleton/Sets/SetDragDrop.cpp b/Engine/Plugins/Experimental/UAF/UAF/Source/UAFEditor/Private/UAF/AbstractSkeleton/Sets/SetDragDrop.cpp
--- a/Engine/Plugins/Experimental/UAF/UAF/Source/UAFEditor/Private/UAF/AbstractSkeleton/Sets/SetDragDrop.cpp
+++ b/Engine/Plugins/Experimental/UAF/UAF/Source/UAFEditor/Private/UAF/AbstractSkeleton/Sets/SetDragDrop.cpp
@@ -2,26 +2,88 @@
 
 #include "UAF/AbstractSkeleton/Sets/SetDragDrop.h"
 
-namespace UE::UAF::Sets
+#include "Widgets/Text/SRichTextBlock.h"
+
+#define LOCTEXT_NAMESPACE "UE::UAF::Editor::FSetDragDropOp"
+
+namespace UE::UAF::Editor
 {
-	TSharedRef<FSetDragDropOp> FSetDragDropOp::New(SSetBinding::FTreeItem_Set* InItem)
+	TSharedRef<FSetDragDropOp> FSetDragDropOp::New(TArray<FAbstractSkeletonSet>&& InSets)
 	{
 		TSharedRef<FSetDragDropOp> Operation = MakeShared<FSetDragDropOp>();
-		Operation->Item = InItem;
+		Operation->Sets = InSets;
 
 		Operation->Construct();
 
 		return Operation;
 	}
 
+	void FSetDragDropOp::Construct()
+	{
+		FDecoratedDragDropOp::Construct();
+
+		check(!Sets.IsEmpty());
+		CachedSetsText = Sets.Num() == 1
+			? FText::FromName(Sets[0].SetName)
+			: FText::Format(LOCTEXT("MultipleSetsTooltip", "{0} and {1} more"), FText::FromName(Sets[0].SetName), FText::AsNumber(Sets.Num() - 1));
+	}
+
 	TSharedPtr<SWidget> FSetDragDropOp::GetDefaultDecorator() const
 	{
 		return SNew(SBorder)
 			.Visibility(EVisibility::Visible)
 			.BorderImage(FAppStyle::GetBrush("Menu.Background"))
+			.Padding(FMargin(8.0f, 4.0f))
 			[
-				SNew(STextBlock)
-					.Text(FText::FromName(Item->Set.SetName))
+				SNew(SHorizontalBox)
+				+ SHorizontalBox::Slot()
+				.AutoWidth()
+				.Padding(FMargin(2.0f))
+				[
+					SNew(SImage)
+					.Image(FAppStyle::Get().GetBrush("ClassIcon.GroupActor"))
+				]
+				+ SHorizontalBox::Slot()
+				.AutoWidth()
+				.Padding(FMargin(2.0f))
+				[
+					SNew(SRichTextBlock)
+					.DecoratorStyleSet(&FAppStyle::Get())
+					.Text_Lambda([this]()
+					{
+						switch (Operation)
+						{
+						case EOperation::None:
+							return CachedSetsText;
+						case EOperation::Bind:
+							return FText::Format(LOCTEXT("BindSet", "Bind <RichTextBlock.Bold>{0}</> to <RichTextBlock.Bold>{1}</>"), CachedSetsText, FText::FromName(OperationSetName));
+						case EOperation::Reparent:
+							return OperationSetName != NAME_None
+								? FText::Format(LOCTEXT("ParentSet", "Parent <RichTextBlock.Bold>{0}</> to <RichTextBlock.Bold>{1}</>"), CachedSetsText, FText::FromName(OperationSetName))
+								: FText::Format(LOCTEXT("Unparent", "Unparent <RichTextBlock.Bold>{0}</>"), CachedSetsText);
+						}
+						return FText::GetEmpty();
+					})
+				]
 			];
 	}
-}
\ No newline at end of file
+
+	void FSetDragDropOp::SetOperation(const EOperation InOperation, const FName InOperationSetName)
+	{
+		Operation = InOperation;
+		OperationSetName = InOperationSetName;
+	}
+	
+	void FSetDragDropOp::ClearOperation()
+	{
+		Operation = EOperation::None;
+		OperationSetName = NAME_None;
+	}
+
+	const TArray<FAbstractSkeletonSet>& FSetDragDropOp::GetDraggedSets() const
+	{
+		return Sets;
+	}
+}
+
+#undef LOCTEXT_NAMESPACE
\ No newline at end of file
diff --git a/Engine/Plugins/Experimental/UAF/UAF/Source/UAFEditor/Private/UAF/AbstractSkeleton/Sets/SetDragDrop.h b/Engine/Plugins/Experimental/UAF/UAF/Source/UAFEditor/Private/UAF/AbstractSkeleton/Sets/SetDragDrop.h
--- a/Engine/Plugins/Experimental/UAF/UAF/Source/UAFEditor/Private/UAF/AbstractSkeleton/Sets/SetDragDrop.h
+++ b/Engine/Plugins/Experimental/UAF/UAF/Source/UAFEditor/Private/UAF/AbstractSkeleton/Sets/SetDragDrop.h
@@ -6,17 +6,37 @@
 #include "DragAndDrop/DecoratedDragDropOp.h"
 #include "UAF/AbstractSkeleton/Sets/SSetBinding.h"
 
-namespace UE::UAF::Sets
+namespace UE::UAF::Editor
 {
 	class FSetDragDropOp : public FDecoratedDragDropOp
 	{
 	public:
+		enum class EOperation : uint8
+		{
+			None,
+			Reparent,
+			Bind
+		};
+		
 		DRAG_DROP_OPERATOR_TYPE(FSetDragDropOp, FDecoratedDragDropOp)
 
-		static TSharedRef<FSetDragDropOp> New(SSetBinding::FTreeItem_Set* InItem);
+		static TSharedRef<FSetDragDropOp> New(TArray<FAbstractSkeletonSet>&& InSets);
+
+		virtual void Construct() override;
 
 		TSharedPtr<SWidget> GetDefaultDecorator() const override;
 
-		SSetBinding::FTreeItem_Set* Item;
+		void SetOperation(const EOperation InOperation, const FName InOperationSetName);
+		void ClearOperation();
+
+		const TArray<FAbstractSkeletonSet>& GetDraggedSets() const;
+		
+	private:
+		TArray<FAbstractSkeletonSet> Sets;
+
+		EOperation Operation = EOperation::None;
+		FName OperationSetName;
+
+		FText CachedSetsText;
 	};
-}
+}
\ No newline at end of file
