Regression: Fix set collection tree view not building the tree items on construction

#rnx

diff --git a/Engine/Plugins/Experimental/UAF/UAF/Source/UAFEditor/Private/UAF/AbstractSkeleton/Sets/SSetCollectionTreeView.cpp b/Engine/Plugins/Experimental/UAF/UAF/Source/UAFEditor/Private/UAF/AbstractSkeleton/Sets/SSetCollectionTreeView.cpp
--- a/Engine/Plugins/Experimental/UAF/UAF/Source/UAFEditor/Private/UAF/AbstractSkeleton/Sets/SSetCollectionTreeView.cpp
+++ b/Engine/Plugins/Experimental/UAF/UAF/Source/UAFEditor/Private/UAF/AbstractSkeleton/Sets/SSetCollectionTreeView.cpp
@@ -38,6 +38,9 @@ namespace UE::UAF::Editor
 		STreeView<FSetCollectionTreeItemPtr>::Construct(SuperArgs);
 		
 		BindCommands();
+
+		// Build the tree items on construction
+		bNeedsRebuilding = true;
 	}
 
 	SSetCollectionTreeView::~SSetCollectionTreeView()
