Add hint text in set binding editor when a set collection hasn't been set

#FailedValidations Preflight

diff --git a/Engine/Plugins/Experimental/UAF/UAF/Source/UAFEditor/Private/UAF/AbstractSkeleton/Sets/SSetBinding.cpp b/Engine/Plugins/Experimental/UAF/UAF/Source/UAFEditor/Private/UAF/AbstractSkeleton/Sets/SSetBinding.cpp
--- a/Engine/Plugins/Experimental/UAF/UAF/Source/UAFEditor/Private/UAF/AbstractSkeleton/Sets/SSetBinding.cpp
+++ b/Engine/Plugins/Experimental/UAF/UAF/Source/UAFEditor/Private/UAF/AbstractSkeleton/Sets/SSetBinding.cpp
@@ -11,6 +11,7 @@
 #include "UAF/AbstractSkeleton/AbstractSkeletonSetBinding.h"
 #include "UAF/AbstractSkeleton/AbstractSkeletonSetCollection.h"
 #include "Widgets/Input/SSearchBox.h"
+#include "Widgets/Layout/SWidgetSwitcher.h"
 
 #define LOCTEXT_NAMESPACE "UE::UAF::Sets::SSetBinding"
 
@@ -24,7 +25,30 @@ void SSetBinding::Construct(const FArguments& InArgs, TWeakObjectPtr<UAbstractSk
 
 		ChildSlot
 			[
-				SNew(SVerticalBox)
+				SNew(SWidgetSwitcher)
+				.WidgetIndex_Lambda([this]()
+				{
+					return SetBinding->GetSetCollection() == nullptr ? 0 : 1;
+				})
+				+ SWidgetSwitcher::Slot()
+				[
+					SNew(SBorder)
+					.BorderImage(FAppStyle::Get().GetBrush("Brushes.Recessed"))
+					[
+						SNew(SOverlay)
+						+ SOverlay::Slot()
+						.HAlign(HAlign_Center)
+						.VAlign(VAlign_Center)
+						[
+							SNew(STextBlock)
+							.ColorAndOpacity(FSlateColor::UseSubduedForeground())
+							.Text(LOCTEXT("NoSetCollectionSelected", "No Set Collection Selected"))
+						]
+					]
+				]
+				+ SWidgetSwitcher::Slot()
+				[
+					SNew(SVerticalBox)
 					+ SVerticalBox::Slot()
 					.AutoHeight()
 					.Padding(FMargin(1.0f, 1.0f))
@@ -59,6 +83,7 @@ void SSetBinding::Construct(const FArguments& InArgs, TWeakObjectPtr<UAbstractSk
 								.DefaultLabel(LOCTEXT("TreeView_Header", "Set Bindings"))
 							)
 					]
+				]
 			];
 
 		RepopulateTreeData();
