Show attribute bindings in the set binding editor for deleted sets
* Previously if attributes are bound to a set in a set binding and later that set is deleted in the set collection then the binding disappears in the editor.
* With these changes, the bindings are shown under an orphaned set so that the attributes bound to it can be unbound.

diff --git a/Engine/Plugins/Experimental/UAF/UAF/Source/UAFEditor/Private/UAF/AbstractSkeleton/Sets/SSetBinding.cpp b/Engine/Plugins/Experimental/UAF/UAF/Source/UAFEditor/Private/UAF/AbstractSkeleton/Sets/SSetBinding.cpp
--- a/Engine/Plugins/Experimental/UAF/UAF/Source/UAFEditor/Private/UAF/AbstractSkeleton/Sets/SSetBinding.cpp
+++ b/Engine/Plugins/Experimental/UAF/UAF/Source/UAFEditor/Private/UAF/AbstractSkeleton/Sets/SSetBinding.cpp
@@ -6,10 +6,11 @@
 #include "Framework/Commands/UICommandList.h"
 #include "Framework/MultiBox/MultiBoxBuilder.h"
 #include "ScopedTransaction.h"
-#include "UAF/AbstractSkeleton/Sets/SetDragDrop.h"
 #include "SkeletonTreeDragDrop.h"
+#include "Styling/StyleColors.h" 
 #include "UAF/AbstractSkeleton/AbstractSkeletonSetBinding.h"
 #include "UAF/AbstractSkeleton/AbstractSkeletonSetCollection.h"
+#include "UAF/AbstractSkeleton/Sets/SetDragDrop.h"
 #include "Widgets/Input/SSearchBox.h"
 #include "Widgets/Layout/SWidgetSwitcher.h"
 
@@ -178,7 +179,7 @@ void SSetBinding::Construct(const FArguments& InArgs, TWeakObjectPtr<UAbstractSk
 				{
 					for (const FAbstractSkeletonSet& Set : SetBinding->GetSetCollection()->GetSetHierarchy())
 					{
-						TSharedPtr<FTreeItem_Set> SetItem = MakeShared<FTreeItem_Set>(*this, Set);
+						TSharedPtr<FTreeItem_Set> SetItem = MakeShared<FTreeItem_Set>(*this, Set, false);
 
 						if (Set.ParentSetName == NAME_None)
 						{
@@ -207,7 +208,16 @@ void SSetBinding::Construct(const FArguments& InArgs, TWeakObjectPtr<UAbstractSk
 					}
 					else
 					{
-						check(false);
+						// This binding binds to a set that no longer exists in the set collection
+						// These bindings are not auto-deleted to avoid implicit destructive edits
+						// This binding can be fixed by re-adding the set in the set collection or
+						// by removing this binding. Show this binding as an orphan.
+
+						TSharedPtr<FTreeItem_Set> DeletedSet = MakeShared<FTreeItem_Set>(*this, FAbstractSkeletonSet(Binding.SetName, NAME_None), true);
+						DeletedSet->Children.Insert(BindingItem, 0);
+
+						TreeItemMap.Add(Binding.SetName, DeletedSet);
+						RootItems.Add(DeletedSet);
 					}
 				}
 
@@ -275,6 +285,34 @@ void SSetBinding::Construct(const FArguments& InArgs, TWeakObjectPtr<UAbstractSk
 
 	TSharedRef<ITableRow> SSetBinding::FTreeItem_Set::GenerateRow(const TSharedRef<STableViewBase>& OwnerTable)
 	{
+		if (bIsDeletedSet)
+		{
+			return SNew(STableRow<TSharedPtr<ITreeItem>>, OwnerTable)
+				.ToolTipText(LOCTEXT("OrphanSetLabel", "This set no longer exists in the Set Collection and will not be used. Attributes bound to this set should be removed or the set should be recreated."))
+				.Style(FAppStyle::Get(), "TableView.AlternatingRow")
+				.ShowWires(true)
+				[
+					SNew(SHorizontalBox)
+					+ SHorizontalBox::Slot()
+					.AutoWidth()
+					.Padding(FMargin(3.0f))
+					[
+						SNew(SImage)
+						.ColorAndOpacity(EStyleColor::Error)
+						.Image(FAppStyle::Get().GetBrush("ClassIcon.GroupActor"))
+					]
+					+ SHorizontalBox::Slot()
+					.FillWidth(1.0)
+					.AutoWidth()
+					.Padding(3.0f)
+					[
+						SNew(STextBlock)
+						.ColorAndOpacity(EStyleColor::Error)
+						.Text(FText::FromName(Set.SetName))
+					]
+				];
+		}
+	
 		return SNew(STableRow<TSharedPtr<ITreeItem>>, OwnerTable)
 			.Style(FAppStyle::Get(), "TableView.AlternatingRow")
 			.OnDragDetected_Lambda([this](const FGeometry& MyGeometry, const FPointerEvent& MouseEvent)
diff --git a/Engine/Plugins/Experimental/UAF/UAF/Source/UAFEditor/Private/UAF/AbstractSkeleton/Sets/SSetBinding.h b/Engine/Plugins/Experimental/UAF/UAF/Source/UAFEditor/Private/UAF/AbstractSkeleton/Sets/SSetBinding.h
--- a/Engine/Plugins/Experimental/UAF/UAF/Source/UAFEditor/Private/UAF/AbstractSkeleton/Sets/SSetBinding.h
+++ b/Engine/Plugins/Experimental/UAF/UAF/Source/UAFEditor/Private/UAF/AbstractSkeleton/Sets/SSetBinding.h
@@ -42,9 +42,10 @@ namespace UE::UAF::Sets
 
 		struct FTreeItem_Set : public ITreeItem
 		{
-			FTreeItem_Set(SSetBinding& InSetBindingWidget, FAbstractSkeletonSet InSet)
+			FTreeItem_Set(SSetBinding& InSetBindingWidget, FAbstractSkeletonSet InSet, bool bInIsDeletedSet)
 				: ITreeItem(InSetBindingWidget)
 				, Set(InSet)
+				, bIsDeletedSet(bInIsDeletedSet)
 			{
 			}
 			
@@ -59,6 +60,7 @@ namespace UE::UAF::Sets
 
 			FAbstractSkeletonSet Set;
 			TArray<TSharedPtr<ITreeItem>> Children;
+			bool bIsDeletedSet;
 		};
 
 		struct FTreeItem_Bone : public ITreeItem
