UAF: Fix crash resolving weak references in module instance constructor call-chain

diff --git a/Engine/Plugins/Experimental/UAF/UAF/Source/UAF/Internal/AnimNextPool.h b/Engine/Plugins/Experimental/UAF/UAF/Source/UAF/Internal/AnimNextPool.h
--- a/Engine/Plugins/Experimental/UAF/UAF/Source/UAF/Internal/AnimNextPool.h
+++ b/Engine/Plugins/Experimental/UAF/UAF/Source/UAF/Internal/AnimNextPool.h
@@ -28,33 +28,6 @@ private:
 	static inline uint32 SerialNumber = 0;
 
 public:
-	// Retrieve the next handle that the pool will allocate when calling Emplace
-	// Handle is potentially invalid once Emplace/Release is called
-	TPoolHandle<ElementType> NextHandle() const
-	{
-		uint32 NextSerialNumber = SerialNumber;
-		NextSerialNumber++;
-		if(NextSerialNumber == 0)
-		{
-			// Serial number wrapped, we cant have any allocated entries or we could have duplicate handles
-			check(Entries.Num() == 0);
-			// Skip zero serial number as this is taken to mean 'invalid'
-			NextSerialNumber++;
-		}
-
-		TPoolHandle<ElementType> Handle;
-		Handle.SerialNumber = NextSerialNumber;
-		if(HeadFreeIndex != MAX_uint32)
-		{
-			Handle.Index = HeadFreeIndex;
-		}
-		else
-		{
-			Handle.Index = Entries.Num();
-		}
-		return Handle;
-	}
-
 	template<typename... ArgsType>
 	TPoolHandle<ElementType> Emplace(ArgsType&&... InArgs)
 	{
diff --git a/Engine/Plugins/Experimental/UAF/UAF/Source/UAF/Internal/Module/AnimNextModuleInstance.h b/Engine/Plugins/Experimental/UAF/UAF/Source/UAF/Internal/Module/AnimNextModuleInstance.h
--- a/Engine/Plugins/Experimental/UAF/UAF/Source/UAF/Internal/Module/AnimNextModuleInstance.h
+++ b/Engine/Plugins/Experimental/UAF/UAF/Source/UAF/Internal/Module/AnimNextModuleInstance.h
@@ -57,22 +57,15 @@ struct FAnimNextModuleInstance : public FUAFAssetInstance
 
 	// Standalone constructor
 	UAF_API FAnimNextModuleInstance(
-		const UObject* InAsset,
+	const UAnimNextModule* InModule,
 		EAnimNextModuleInitMethod InInitMethod);
 
 	// Standalone constructor, UObject host
 	UAF_API FAnimNextModuleInstance(
-		const UObject* InAsset,
+	const UAnimNextModule* InModule,
 		UObject* InObject,
 		EAnimNextModuleInitMethod InInitMethod);
 
-	// Pooled constructor
-	UAF_API FAnimNextModuleInstance(
-		const UObject* InAsset,
-		UObject* InObject,
-		UE::UAF::TPool<FAnimNextModuleInstance>* InPool,
-		UE::UAF::FModuleHandle InHandle,
-		EAnimNextModuleInitMethod InInitMethod);
 	UAF_API ~FAnimNextModuleInstance();
 
 	// Checks to see if this entry is ticking
@@ -142,7 +135,8 @@ struct FAnimNextModuleInstance : public FUAFAssetInstance
 
 private:
 	// Setup the entry
-	UAF_API void Initialize(const FUAFSystemFactoryParams& InFactoryParams);
+	// Note that we perform initialization in two phases because we may need to access weak references to 'this' during initialization
+	UAF_API void Initialize(const FUAFSystemFactoryParams& InFactoryParams, UE::UAF::TPool<FAnimNextModuleInstance>* InPool = nullptr);
 
 	// Clear data that binds the schedule to a runtime (e.g. tick functions) and any instance data
 	void ResetBindingsAndInstanceData();
diff --git a/Engine/Plugins/Experimental/UAF/UAF/Source/UAF/Private/Component/AnimNextWorldSubsystem.cpp b/Engine/Plugins/Experimental/UAF/UAF/Source/UAF/Private/Component/AnimNextWorldSubsystem.cpp
--- a/Engine/Plugins/Experimental/UAF/UAF/Source/UAF/Private/Component/AnimNextWorldSubsystem.cpp
+++ b/Engine/Plugins/Experimental/UAF/UAF/Source/UAF/Private/Component/AnimNextWorldSubsystem.cpp
@@ -5,6 +5,7 @@
 #include "AnimNextDebugDraw.h"
 #include "AnimNextRigVMAsset.h"
 #include "Engine/World.h"
+#include "Factory/SystemFactory.h"
 #include "Logging/StructuredLog.h"
 #include "Module/AnimNextModuleInstance.h"
 #include "Module/ModuleTaskContext.h"
@@ -161,12 +162,23 @@ void UAnimNextWorldSubsystem::RegisterHandle(UE::UAF::FModuleHandle& InOutHandle
 {
 	using namespace UE::UAF;
 	check(IsInGameThread());
+	check(!IsValidHandle(InOutHandle));	// Shouldnt br registering without unregistering first
 	{
 		FWriteScopeLock InstancesLockScope(InstancesLock);
 		UE_MT_SCOPED_WRITE_ACCESS(InstancesAccessDetector);
 
-		InOutHandle = Instances.NextHandle();
-		Instances.Emplace(InAsset, InObject, &Instances, InOutHandle, InitMethod);
+		const UAnimNextModule* Module = Cast<UAnimNextModule>(InAsset);
+		FUAFSystemFactoryParams FactoryParams;
+		if (Module == nullptr)
+		{
+			FactoryParams = FSystemFactory::GetDefaultParamsForObject(InAsset);
+			Module = FSystemFactory::BuildSystem(FactoryParams.GetBuilder());
+		}
+
+		InOutHandle = Instances.Emplace(Module, InObject, InitMethod);
+		FAnimNextModuleInstance& Instance = Instances.Get(InOutHandle);
+		Instance.Handle = InOutHandle;
+		Instance.Initialize(FactoryParams, &Instances);
 	}
 }
 
diff --git a/Engine/Plugins/Experimental/UAF/UAF/Source/UAF/Private/Module/AnimNextModuleInstance.cpp b/Engine/Plugins/Experimental/UAF/UAF/Source/UAF/Private/Module/AnimNextModuleInstance.cpp
--- a/Engine/Plugins/Experimental/UAF/UAF/Source/UAF/Private/Module/AnimNextModuleInstance.cpp
+++ b/Engine/Plugins/Experimental/UAF/UAF/Source/UAF/Private/Module/AnimNextModuleInstance.cpp
@@ -30,48 +30,27 @@ DEFINE_STAT(STAT_AnimNext_InitializeInstance);
 FAnimNextModuleInstance::FAnimNextModuleInstance() = default;
 
 FAnimNextModuleInstance::FAnimNextModuleInstance(
-		const UObject* InAsset,
+		const UAnimNextModule* InModule,
 		EAnimNextModuleInitMethod InInitMethod)
-	: FAnimNextModuleInstance(InAsset, nullptr, nullptr, UE::UAF::FModuleHandle(), InInitMethod)
+	: FAnimNextModuleInstance(InModule, nullptr, InInitMethod)
 {
 }
 
 FAnimNextModuleInstance::FAnimNextModuleInstance(
-		const UObject* InAsset,
+		const UAnimNextModule* InModule,
 		UObject* InObject,
 		EAnimNextModuleInitMethod InInitMethod)
-	: FAnimNextModuleInstance(InAsset, InObject, nullptr, UE::UAF::FModuleHandle(), InInitMethod)
-{
-	check(InObject);
-}
-
-FAnimNextModuleInstance::FAnimNextModuleInstance(
-		const UObject* InAsset,
-		UObject* InObject,
-		UE::UAF::TPool<FAnimNextModuleInstance>* InPool,
-		UE::UAF::FModuleHandle InHandle,
-		EAnimNextModuleInitMethod InInitMethod)
 	: FUAFAssetInstance(StaticStruct())
 	, Object(InObject)
-	, Pool(InPool)
-	, Handle(InHandle)
 	, InitState(EInitState::NotInitialized)
 	, RunState(ERunState::NotInitialized)
 	, InitMethod(InInitMethod)
 {
 	using namespace UE::UAF;
 
-	const UAnimNextModule* Module = Cast<UAnimNextModule>(InAsset);
-	FUAFSystemFactoryParams FactoryParams;
-	if (Module == nullptr)
-	{
-		FactoryParams = FSystemFactory::GetDefaultParamsForObject(InAsset);
-		Module = FSystemFactory::BuildSystem(FactoryParams.GetBuilder());
-	}
+	ensure(InModule);
 
-	ensure(Module);
-
-	Asset = Module;
+	Asset = InModule;
 
 #if UE_ENABLE_DEBUG_DRAWING
 	if(Object && Object->GetWorld())
@@ -79,8 +58,6 @@ FAnimNextModuleInstance::FAnimNextModuleInstance(
 		DebugDraw = MakeUnique<UE::UAF::Debug::FDebugDraw>(Object);
 	}
 #endif
-
-	Initialize(FactoryParams);
 }
 
 FAnimNextModuleInstance::~FAnimNextModuleInstance()
@@ -96,7 +73,7 @@ FAnimNextModuleInstance::~FAnimNextModuleInstance()
 	Handle.Reset();
 }
 
-void FAnimNextModuleInstance::Initialize(const FUAFSystemFactoryParams& InFactoryParams)
+void FAnimNextModuleInstance::Initialize(const FUAFSystemFactoryParams& InFactoryParams, UE::UAF::TPool<FAnimNextModuleInstance>* InPool)
 {
 	SCOPE_CYCLE_COUNTER(STAT_AnimNext_InitializeInstance);
 	
