UAF: Fix pool assignment I missed in previous changelist

Found by bughawk

diff --git a/Engine/Plugins/Experimental/UAF/UAF/Source/UAF/Private/Module/AnimNextModuleInstance.cpp b/Engine/Plugins/Experimental/UAF/UAF/Source/UAF/Private/Module/AnimNextModuleInstance.cpp
--- a/Engine/Plugins/Experimental/UAF/UAF/Source/UAF/Private/Module/AnimNextModuleInstance.cpp
+++ b/Engine/Plugins/Experimental/UAF/UAF/Source/UAF/Private/Module/AnimNextModuleInstance.cpp
@@ -81,6 +81,8 @@ void FAnimNextModuleInstance::Initialize(const FUAFSystemFactoryParams& InFactor
 
 	check(IsInGameThread());
 
+	Pool = InPool;
+
 	const UAnimNextModule* Module = GetModule();
 	if (Module)
 	{
