Variable Watch 0.1
* All very subject to change.
* Foundation work for the workspace interfacing with rewind debugger.
* When a rewind debugger session is active and a UAF workspace is active, debug objects are selectable in the workspace for objects alive at the rewind debugger playhead.
* The variable watch tab currently will display the hierarchy of graph instances and upon selecting an instance will show the variable values of variables in that instance.
* This tab will morph into a more traditional variable watch window similar to visual studio once the design has evolved.

#rnx

diff --git a/Engine/Plugins/Experimental/UAF/UAF/Source/UAFEditor/Private/AnimNextEditorModule.cpp b/Engine/Plugins/Experimental/UAF/UAF/Source/UAFEditor/Private/AnimNextEditorModule.cpp
--- a/Engine/Plugins/Experimental/UAF/UAF/Source/UAFEditor/Private/AnimNextEditorModule.cpp
+++ b/Engine/Plugins/Experimental/UAF/UAF/Source/UAFEditor/Private/AnimNextEditorModule.cpp
@@ -81,6 +81,7 @@
 #include "Wizard/AssetWizard.h"
 #include "Interfaces/IPluginManager.h"
 #include "AssetToolsModule.h"
+#include "UAF/Debug/VariableWatchTabSummoner.h"
 
 #define LOCTEXT_NAMESPACE "AnimNextEditorModule"
 
@@ -216,6 +217,10 @@ void FAnimNextEditorModule::StartupModule()
 			TSharedRef<FAnimNextFindAndReplaceTabSummoner> FindAndReplaceTabSummoner = MakeShared<FAnimNextFindAndReplaceTabSummoner>(InEditorPtr);
 			TabFactories.RegisterFactory(FindAndReplaceTabSummoner);
 			FindAndReplaceTabSummoner->RegisterTabSpawner(InTabManager, nullptr);
+
+			TSharedRef<FVariableWatchTabSummoner> VariableWatchTabSummoner = MakeShared<FVariableWatchTabSummoner>(InEditorPtr);
+			TabFactories.RegisterFactory(VariableWatchTabSummoner);
+			VariableWatchTabSummoner->RegisterTabSpawner(InTabManager, nullptr);
 		});
 
 	WorkspaceEditorModule.OnExtendTabs().AddLambda([](FLayoutExtender& InLayoutExtender, TSharedPtr<UE::Workspace::IWorkspaceEditor> InEditorPtr)
diff --git a/Engine/Plugins/Experimental/UAF/UAF/Source/UAFEditor/Private/RewindDebugger/AnimNextProvider.cpp b/Engine/Plugins/Experimental/UAF/UAF/Source/UAFEditor/Private/RewindDebugger/AnimNextProvider.cpp
--- a/Engine/Plugins/Experimental/UAF/UAF/Source/UAFEditor/Private/RewindDebugger/AnimNextProvider.cpp
+++ b/Engine/Plugins/Experimental/UAF/UAF/Source/UAFEditor/Private/RewindDebugger/AnimNextProvider.cpp
@@ -17,9 +17,9 @@ void FAnimNextProvider::AppendInstance(uint64 InstanceId, uint64 HostInstanceId,
 {
 	Session.WriteAccessCheck();
 
-	if (HostInstanceId == 0)
+	if (HostInstanceId == OuterObjectId)
 	{
-		ComponentIdToModuleId.Add(OuterObjectId, InstanceId);
+		ComponentIdToSystemId.Add(OuterObjectId, InstanceId);
 	}
 
 	TSharedRef<FDataInterfaceData> Data = MakeShared<FDataInterfaceData>(Session);
@@ -70,7 +70,7 @@ bool FAnimNextProvider::GetModuleId(uint64 ComponentId, uint64& OutModuleId) con
 {
 	Session.ReadAccessCheck();
 	
-	if (const uint64* Id = ComponentIdToModuleId.Find(ComponentId))
+	if (const uint64* Id = ComponentIdToSystemId.Find(ComponentId))
 	{
 		OutModuleId = *Id;
 		return true;
@@ -103,5 +103,18 @@ void FAnimNextProvider::EnumerateChildInstances(uint64 InstanceId, TFunctionRef<
 
 }
 
+void FAnimNextProvider::EnumerateDataInterfaces(const double InStartTime, TFunctionRef<void(const TSharedRef<const FDataInterfaceData>)> Callback) const
+{
+	Session.ReadAccessCheck();
+
+	for (const auto& [InstanceId, InterfaceData] : DataInterfaceData)
+	{
+		if (InterfaceData->StartTime <= InStartTime && InterfaceData->EndTime >= InStartTime)
+		{
+			Callback(InterfaceData);
+		}
+	}
+}
+
 #undef LOCTEXT_NAMESPACE
 #endif // ANIMNEXT_TRACE_ENABLED
diff --git a/Engine/Plugins/Experimental/UAF/UAF/Source/UAFEditor/Private/RewindDebugger/AnimNextProvider.h b/Engine/Plugins/Experimental/UAF/UAF/Source/UAFEditor/Private/RewindDebugger/AnimNextProvider.h
--- a/Engine/Plugins/Experimental/UAF/UAF/Source/UAFEditor/Private/RewindDebugger/AnimNextProvider.h
+++ b/Engine/Plugins/Experimental/UAF/UAF/Source/UAFEditor/Private/RewindDebugger/AnimNextProvider.h
@@ -74,11 +74,12 @@ public:
 	const FPropertyDescriptionData* GetPropertyDescriptionData(uint32 PropertyDescriptionHash) const { return HashToPropertyDescriptionDataMap.Find(PropertyDescriptionHash);}
 
 	void EnumerateChildInstances(uint64 InstanceId, TFunctionRef<void(const FDataInterfaceData&)> Callback) const;
+	void EnumerateDataInterfaces(const double StartTime, TFunctionRef<void(const TSharedRef<const FDataInterfaceData>)> Callback) const;
 
 private:
 	TraceServices::IAnalysisSession& Session;
 
-	TMap<uint64, uint64> ComponentIdToModuleId;
+	TMap<uint64, uint64> ComponentIdToSystemId;
 	TMap<uint64, TSharedRef<FDataInterfaceData>> DataInterfaceData;
 	TMap<uint64, TArray<TSharedRef<FDataInterfaceData>>> HostToChildDataMap;
 	TMap<uint32, FPropertyDescriptionData> HashToPropertyDescriptionDataMap;
diff --git a/Engine/Plugins/Experimental/UAF/UAF/Source/UAFEditor/Private/UAF/Debug/SVariableWatch.cpp b/Engine/Plugins/Experimental/UAF/UAF/Source/UAFEditor/Private/UAF/Debug/SVariableWatch.cpp
new file mode 100644
--- /dev/null
+++ b/Engine/Plugins/Experimental/UAF/UAF/Source/UAFEditor/Private/UAF/Debug/SVariableWatch.cpp
@@ -0,0 +1,461 @@
+// Copyright Epic Games, Inc. All Rights Reserved.
+
+#include "SVariableWatch.h"
+
+#include "EditorModeManager.h"
+#include "Component/AnimNextComponent.h"
+#include "IGameplayProvider.h"
+#include "InstancedPropertyBagStructureDataProvider.h"
+#include "IRewindDebugger.h"
+#include "ISinglePropertyView.h"
+#include "IWorkspaceEditor.h"
+#include "IWorkspaceEditorModule.h"
+#include "ObjectAsTraceIdProxyArchiveReader.h"
+#include "ObjectTrace.h"
+#include "RewindDebugger/AnimNextProvider.h"
+#include "SceneOutlinerPublicTypes.h"
+#include "SSceneOutliner.h"
+#include "TraceServices/Model/AnalysisSession.h"
+#include "UAF/Workspace/IDebugObjectSelector.h"
+#include "UAF/Workspace/RewindDebuggerExtension.h"
+#include "Variables/Outliner/VariablesOutlinerColumns.h"
+#include "Variables/Outliner/VariablesOutlinerMode.h"
+#include "Variables/SVariablesView.h"
+#include "Workspace/AnimNextWorkspaceEditorMode.h"
+
+static const FName ColumnId_VariableName(TEXT("VariableName"));
+static const FName ColumnId_VariableValue(TEXT("VariableValue"));
+
+#define LOCTEXT_NAMESPACE "UE::UAF::Editor::SDebugInstanceTreeRow"
+
+namespace UE::UAF::Editor
+{
+	class SDebugInstanceTreeRow : public SMultiColumnTableRow<SVariableWatch::FDebugInstanceTreeItemPtr>
+    {
+    	SLATE_BEGIN_ARGS(SDebugInstanceTreeRow)
+    	{}
+    	SLATE_END_ARGS()
+
+    public:
+    	void Construct(const FArguments& InArgs, SVariableWatch::FDebugInstanceTreeItemPtr InItem, const TSharedRef<STableViewBase>& InOwnerTableView, const bool bInHighlight)
+    	{
+    		Item = InItem;
+    		bHighlight = bInHighlight;
+
+    		FSuperRowType::Construct(FSuperRowType::FArguments(), InOwnerTableView);
+    	}
+
+    	virtual TSharedRef<SWidget> GenerateWidgetForColumn(const FName& InColumnId) override
+    	{
+			return SNew(SHorizontalBox)
+			+ SHorizontalBox::Slot()
+			.AutoWidth()
+			[
+				SNew(SExpanderArrow, SharedThis(this))
+			]
+			+ SHorizontalBox::Slot()
+			.AutoWidth()
+			[
+				SNew(STextBlock)
+				.ColorAndOpacity(bHighlight ? FAppStyle::Get().GetSlateColor("Colors.AccentBlue") : FSlateColor::UseForeground())
+				.Text(FText::FromString(Item->Asset ? Item->Asset->GetName() : FString()))
+			];
+    	}
+
+    private:
+    	SVariableWatch::FDebugInstanceTreeItemPtr Item;
+    	bool bHighlight = false;
+    };
+    	
+    class SVariableWatchRow : public SMultiColumnTableRow<SVariableWatch::FVariablesBoxItemPtr>
+	{
+		SLATE_BEGIN_ARGS(SVariableWatchRow)
+		{}
+		SLATE_END_ARGS()
+                
+	public:
+		void Construct(const FArguments& InArgs, SVariableWatch::FVariablesBoxItemPtr InItem, const TSharedRef<STableViewBase>& InOwnerTableView, TWeakPtr<SVariableWatch> InVariableWatch)
+		{
+			Item = InItem;
+			VariableWatch = InVariableWatch;
+
+			FSuperRowType::Construct(FSuperRowType::FArguments(), InOwnerTableView);
+		}
+
+		virtual TSharedRef<SWidget> GenerateWidgetForColumn(const FName& InColumnId) override
+		{
+			if (InColumnId == ColumnId_VariableName)
+			{
+				return SNew(SHorizontalBox)
+				+ SHorizontalBox::Slot()
+				.AutoWidth()
+				.VAlign(VAlign_Center)
+				.Padding(FMargin(12.0f, 2.0f, 0.0f, 2.0f))
+				[
+					SNew(STextBlock)
+					.TextStyle(FAppStyle::Get(), TEXT("SmallText"))
+					.Text(FText::FromName(*Item))
+				];
+			}
+			
+			if (InColumnId == ColumnId_VariableValue)
+			{
+				FSinglePropertyParams SinglePropertyArgs;
+				SinglePropertyArgs.NamePlacement = EPropertyNamePlacement::Hidden;
+				SinglePropertyArgs.bHideResetToDefault = true;
+				
+				FInstancedPropertyBag& PropertyBag = VariableWatch.Pin()->Properties;
+				const FName PropName = *Item;
+				
+				FPropertyEditorModule& PropertyEditorModule = FModuleManager::Get().LoadModuleChecked<FPropertyEditorModule>("PropertyEditor");
+				const TSharedPtr<ISinglePropertyView> SingleStructPropertyView = PropertyEditorModule.CreateSingleProperty(MakeShared<FInstancePropertyBagStructureDataProvider>(PropertyBag), PropName, SinglePropertyArgs);
+				SingleStructPropertyView->SetEnabled(false);
+                    		
+				return SNew(SHorizontalBox)
+				+ SHorizontalBox::Slot()
+				.AutoWidth()
+				.VAlign(VAlign_Center)
+				[
+					SingleStructPropertyView.ToSharedRef()
+				];
+			}
+
+			return SNullWidget::NullWidget;
+		}
+
+    private:
+    	SVariableWatch::FVariablesBoxItemPtr Item;
+    	TWeakPtr<SVariableWatch> VariableWatch;
+	};
+
+	void SVariableWatch::Construct(const FArguments& InArgs, const TWeakPtr<UE::Workspace::IWorkspaceEditor>& InWorkspaceEditor)
+	{
+		WorkspaceEditor = InWorkspaceEditor;
+
+		// This widget is constructed before the UAF workspace editor mode is activated so we listen to
+		// when its activated and bind to the events we care about
+		InWorkspaceEditor.Pin()->GetEditorModeManager().OnEditorModeIDChanged().AddSP(this, &SVariableWatch::HandleWorkspaceEditorModeChanged);
+		
+		InWorkspaceEditor.Pin()->OnFocusedDocumentChanged().AddSP(this, &SVariableWatch::HandleFocusedDocumentChanged);
+
+		DebugInstanceTreeView = SNew(STreeView<FDebugInstanceTreeItemPtr>)
+			.TreeItemsSource(&RootDebugInstances)
+			.HeaderRow(
+				SNew(SHeaderRow)
+				+ SHeaderRow::Column("Instance")
+				.DefaultLabel(INVTEXT("Instance"))
+			)
+			.OnGetChildren(this, &SVariableWatch::InstanceTreeView_OnGetChildren)
+			.OnGenerateRow(this, &SVariableWatch::InstanceTreeView_OnGenerateRow)
+			.OnSelectionChanged(this, &SVariableWatch::InstanceTreeView_OnSelectionChanged);
+
+		VariablesBox = SNew(SListView<FVariablesBoxItemPtr>)
+			.ListItemsSource(&VariablesBoxSource)
+			.HeaderRow(
+				SNew(SHeaderRow)
+				+ SHeaderRow::Column(ColumnId_VariableName)
+				.DefaultLabel(LOCTEXT("ColumnLabel_VariableName", "Variable"))
+				+ SHeaderRow::Column(ColumnId_VariableValue)
+				.DefaultLabel(LOCTEXT("ColumnLabel_VariableValue", "Value"))
+			)
+			.OnGenerateRow(this, &SVariableWatch::VariablesListView_OnGenerateRow);
+		
+		ChildSlot
+		[
+			SNew(SVerticalBox)
+			+ SVerticalBox::Slot()
+			.AutoHeight()
+			[
+				DebugInstanceTreeView.ToSharedRef()
+			]
+			+ SVerticalBox::Slot()
+			.AutoHeight()
+			[
+				VariablesBox.ToSharedRef()
+			]
+		];
+	}
+	
+	void SVariableWatch::RegenerateDebugInstances(const bool bForceRegenerate)
+	{
+		if (!RewindDebugger)
+		{
+			return;
+		}
+		
+		const TraceServices::IAnalysisSession* AnalysisSession = RewindDebugger->GetAnalysisSession();
+		TraceServices::FAnalysisSessionReadScope SessionReadScope(*AnalysisSession);
+
+		const double CurrentScrubTime = RewindDebugger->GetScrubTime();
+		const double CurrentTraceTime  = RewindDebugger->CurrentTraceTime();
+
+		RootDebugInstances.Empty();
+		DebugInstances.Empty();
+
+		const TraceServices::IFrameProvider& FrameProvider = TraceServices::ReadFrameProvider(*RewindDebugger->GetAnalysisSession());
+
+		TraceServices::FFrame Frame;
+		if (!FrameProvider.GetFrameFromTime(ETraceFrameType::TraceFrameType_Game, CurrentTraceTime, Frame))
+		{
+			return;
+		}
+
+		const IGameplayProvider* GameplayProvider = AnalysisSession->ReadProvider<IGameplayProvider>("GameplayProvider");
+		const FAnimNextProvider* UAFProvider = AnalysisSession->ReadProvider<FAnimNextProvider>("AnimNextProvider");
+		
+		if (!GameplayProvider || !UAFProvider)
+		{
+			return;
+		}
+
+		const UE::Workspace::FWorkspaceDocument FocusedDocument = WorkspaceEditor.Pin()->GetFocusedWorkspaceDocument();
+		const FSoftObjectPath FocusedAssetPath = FocusedDocument.Export.GetFirstAssetPath();
+
+		Properties.Reset();
+		NativeProperties.Reset();
+
+		const UObject* WorkspaceDebugObject = nullptr;
+		UAnimNextWorkspaceEditorMode* EditorMode = Cast<UAnimNextWorkspaceEditorMode>(WorkspaceEditor.Pin()->GetEditorModeManager().GetActiveScriptableMode(UAnimNextWorkspaceEditorMode::EM_AnimNextWorkspace));
+		if (EditorMode)
+		{
+			WorkspaceDebugObject = EditorMode->GetDebugObjectSelector()->GetDebugObject().Get();
+		}
+
+		UAFProvider->EnumerateDataInterfaces(CurrentScrubTime, [this, FocusedAssetPath, UAFProvider, WorkspaceDebugObject](const TSharedRef<const FDataInterfaceData> DataInterfaceData)
+		{
+			UObject* Asset = FObjectTrace::GetObjectFromId(DataInterfaceData->AssetId);
+			if (FSoftObjectPath(Asset) != FocusedAssetPath)
+			{
+				return;
+			}
+			
+			const UObject* OwnerObject = FObjectTrace::GetObjectFromId(DataInterfaceData->OuterObjectId);
+			if (!WorkspaceDebugObject || (WorkspaceDebugObject != OwnerObject && !OwnerObject->IsInOuter(WorkspaceDebugObject)))
+			{
+				return;
+			}
+
+			TSharedPtr<FDebugInstance> NewDebugInstance = MakeShared<FDebugInstance>();
+			NewDebugInstance->InstanceId = DataInterfaceData->InstanceId;
+			NewDebugInstance->HostInstanceId = DataInterfaceData->HostInstanceId;
+			NewDebugInstance->Asset = FObjectTrace::GetObjectFromId(DataInterfaceData->AssetId);
+			
+			DebugInstances.Add(NewDebugInstance->InstanceId, NewDebugInstance);
+			
+			uint64 HostInstanceId = NewDebugInstance->HostInstanceId;
+			const FDataInterfaceData* HostInstanceData = UAFProvider->GetDataInterfaceData(HostInstanceId);
+			while (HostInstanceData)
+			{
+				if (!DebugInstances.Contains(HostInstanceId))
+				{
+					TSharedPtr<FDebugInstance> NewHostDebugInstance = MakeShared<FDebugInstance>();
+					NewHostDebugInstance->InstanceId = HostInstanceData->InstanceId;
+					NewHostDebugInstance->HostInstanceId = HostInstanceData->HostInstanceId;
+					NewHostDebugInstance->Asset = FObjectTrace::GetObjectFromId(HostInstanceData->AssetId);
+					
+					DebugInstances.Add(NewHostDebugInstance->InstanceId, NewHostDebugInstance);
+					
+					HostInstanceId = NewHostDebugInstance->HostInstanceId;
+                    HostInstanceData = UAFProvider->GetDataInterfaceData(HostInstanceId);
+				}
+				else
+				{
+					break;
+				}
+			}
+			
+			if (!DebugInstances.Contains(HostInstanceId))
+			{
+				TSharedPtr<FDebugInstance> NewHostDebugInstance = MakeShared<FDebugInstance>();
+				NewHostDebugInstance->InstanceId = HostInstanceId;
+				NewHostDebugInstance->HostInstanceId = 0;
+				NewHostDebugInstance->Asset = FObjectTrace::GetObjectFromId(HostInstanceId);
+			
+				DebugInstances.Add(NewHostDebugInstance->InstanceId, NewHostDebugInstance);
+				RootDebugInstances.Add(NewHostDebugInstance);
+			}
+		});
+				
+		for (const auto& [InstanceId, DebugTreeItem] : DebugInstances)
+		{
+			DebugInstanceTreeView->SetItemExpansion(DebugTreeItem, true);
+		}
+		
+		DebugInstanceTreeView->RequestTreeRefresh();
+		
+		if (SelectedDebugInstanceId != 0)
+		{
+			const auto DebugInstanceTreeItem = DebugInstances.Find(SelectedDebugInstanceId);
+			if (DebugInstanceTreeItem)
+			{
+				DebugInstanceTreeView->SetSelection(*DebugInstanceTreeItem);
+			}
+			
+		}
+		
+		OnDebugInstanceChanged();
+	}
+
+	void SVariableWatch::HandleFocusedDocumentChanged(const UE::Workspace::FWorkspaceDocument& InDocument)
+	{
+		SelectedDebugInstanceId = 0;
+	
+		RegenerateDebugInstances(true);
+	}
+	
+	void SVariableWatch::HandleWorkspaceDebugObjectChanged(TWeakObjectPtr<UObject> InObject)
+	{
+		RegenerateDebugInstances(true);
+	}
+	
+	void SVariableWatch::HandleRewindDebuggerUpdate(float DeltaTime, IRewindDebugger* InRewindDebugger)
+	{
+		RewindDebugger = InRewindDebugger;
+		
+		RegenerateDebugInstances(false);
+	}
+	
+	void SVariableWatch::OnDebugInstanceChanged()
+	{
+		VariablesBoxSource.Empty();
+	
+		const TraceServices::IAnalysisSession* AnalysisSession = RewindDebugger->GetAnalysisSession();
+		TraceServices::FAnalysisSessionReadScope SessionReadScope(*AnalysisSession);
+
+		const double CurrentScrubTime = RewindDebugger->GetScrubTime();
+		const double CurrentTraceTime  = RewindDebugger->CurrentTraceTime();
+				
+		const IGameplayProvider* GameplayProvider = AnalysisSession->ReadProvider<IGameplayProvider>("GameplayProvider");
+		const FAnimNextProvider* UAFProvider = AnalysisSession->ReadProvider<FAnimNextProvider>("AnimNextProvider");
+        if (!GameplayProvider || !UAFProvider)
+        {
+	        return;
+        }  	
+        
+		const FDataInterfaceData* DataInterfaceData = UAFProvider->GetDataInterfaceData(SelectedDebugInstanceId);
+		if (!DataInterfaceData)
+		{
+			return;
+		}
+		
+		const TraceServices::IFrameProvider& FrameProvider = TraceServices::ReadFrameProvider(*AnalysisSession);
+
+		TraceServices::FFrame Frame;
+		if (!FrameProvider.GetFrameFromTime(ETraceFrameType::TraceFrameType_Game, CurrentTraceTime, Frame))
+		{
+			return;
+		}
+			
+		DataInterfaceData->VariablesTimeline.EnumerateEvents(Frame.StartTime, Frame.EndTime, [UAFProvider, GameplayProvider, AnalysisSession, this](double InStartTime, double InEndTime, uint32 InDepth, const FPropertyVariableData& VariableListData)
+		{
+			// Most of this reading code is ripped from FAnimNextModuleTrack
+			if (VariableListData.ValueType == EPropertyVariableDataType::PropertyBag)
+			{
+				// First look up property description
+				const FPropertyDescriptionData* DescriptionData = UAFProvider->GetPropertyDescriptionData(VariableListData.DescriptionHash);
+				if (DescriptionData == nullptr)
+				{
+					// Can't do anything without a description
+					return TraceServices::EEventEnumerate::Continue;
+				}
+
+				// Load the property descriptions
+				TArray<FPropertyBagPropertyDesc> PropertyDescriptions;
+				FMemoryReader DescriptionReader(DescriptionData->Data);
+				FObjectAsTraceIdProxyArchiveReader DescriptionReaderProxy(DescriptionReader, GameplayProvider);
+				DescriptionReaderProxy.UsingCustomVersion(UE::UAF::FAnimNextTrace::CustomVersionGUID);
+				DescriptionReaderProxy << PropertyDescriptions;
+				
+				//Properties->Properties.SetNum(Properties->Properties.Num() + 1);
+				Properties.AddProperties(PropertyDescriptions);
+
+				// Load the property values
+				FMemoryReader Reader(VariableListData.ValueData);
+				FObjectAsTraceIdProxyArchiveReader ReaderProxy(Reader, GameplayProvider);
+
+				UPropertyBag* PropertyBagStruct = const_cast<UPropertyBag*>(Properties.GetPropertyBagStruct());
+				if (PropertyBagStruct != nullptr)
+				{
+					PropertyBagStruct->SerializeItem(ReaderProxy, Properties.GetMutableValue().GetMemory(), nullptr);
+				}
+			}
+			else
+			{
+				FMemoryReader Reader(VariableListData.ValueData);
+				FObjectAsTraceIdProxyArchiveReader ReaderProxy(Reader, GameplayProvider);
+
+				FInstancedStruct::StaticStruct()->SerializeItem(ReaderProxy, &NativeProperties, nullptr);
+			}
+
+			return TraceServices::EEventEnumerate::Continue;
+		});
+		        		
+		// This gives you the schema (like a struct definition)
+		const UPropertyBag* BagStruct = Properties.GetPropertyBagStruct();
+		if (!BagStruct)
+		{
+			return;
+		}
+		
+		// Iterate over defined properties
+		for (const FPropertyBagPropertyDesc& PropDesc : BagStruct->GetPropertyDescs())
+		{
+			const FName PropName = PropDesc.Name;
+			VariablesBoxSource.Add(MakeShared<FName>(PropName));
+		}
+		
+		VariablesBox->RequestListRefresh();
+	}
+
+	void SVariableWatch::HandleWorkspaceEditorModeChanged(const FEditorModeID& EditorModeId, bool bIsEnteringMode)
+	{
+		if (EditorModeId == UAnimNextWorkspaceEditorMode::EM_AnimNextWorkspace)
+		{
+			if (bIsEnteringMode)
+			{
+				UAnimNextWorkspaceEditorMode* EditorMode = Cast<UAnimNextWorkspaceEditorMode>(WorkspaceEditor.Pin()->GetEditorModeManager().GetActiveScriptableMode(UAnimNextWorkspaceEditorMode::EM_AnimNextWorkspace));
+				if (ensure(EditorMode))
+				{
+					EditorMode->GetRewindDebuggerExtension()->RegisterOnRewindDebuggerUpdate(
+						IUAFWorkspaceRewindDebugger::FOnRewindDebuggerUpdate::FDelegate::CreateSP(this, &SVariableWatch::HandleRewindDebuggerUpdate));
+
+					EditorMode->GetDebugObjectSelector()->RegisterOnDebugObjectChanged(
+						IDebugObjectSelector::FOnDebugObjectChanged::FDelegate::CreateSP(this, &SVariableWatch::HandleWorkspaceDebugObjectChanged));
+				}
+			}
+		}
+	}
+
+	void SVariableWatch::InstanceTreeView_OnGetChildren(FDebugInstanceTreeItemPtr InItem, TArray<FDebugInstanceTreeItemPtr>& OutChildren)
+	{
+		for (const auto& [InstanceId, DebugInstance] : DebugInstances)
+		{
+			if (DebugInstance->HostInstanceId == InItem->InstanceId)
+			{
+				OutChildren.Add(DebugInstance);
+			}
+		}
+	}
+
+	TSharedRef<ITableRow> SVariableWatch::InstanceTreeView_OnGenerateRow(FDebugInstanceTreeItemPtr InItem, const TSharedRef<STableViewBase>& OwnerTable)
+	{
+		const bool bHighlight = FSoftObjectPath(InItem->Asset) == WorkspaceEditor.Pin()->GetFocusedWorkspaceDocument().Export.GetFirstAssetPath();
+		return SNew(SDebugInstanceTreeRow, InItem, OwnerTable, bHighlight);
+	}
+	
+	void SVariableWatch::InstanceTreeView_OnSelectionChanged(FDebugInstanceTreeItemPtr NewSelection, ESelectInfo::Type)
+	{
+		if (NewSelection.IsValid() && NewSelection->InstanceId != SelectedDebugInstanceId)
+		{
+			SelectedDebugInstanceId = NewSelection->InstanceId;
+		}
+		OnDebugInstanceChanged();
+	}
+
+	TSharedRef<ITableRow> SVariableWatch::VariablesListView_OnGenerateRow(const FVariablesBoxItemPtr InItem, const TSharedRef<STableViewBase>& OwnerTable)
+	{
+		return SNew(SVariableWatchRow, InItem, OwnerTable, SharedThis(this));
+	}
+}
+
+#undef LOCTEXT_NAMESPACE
\ No newline at end of file
diff --git a/Engine/Plugins/Experimental/UAF/UAF/Source/UAFEditor/Private/UAF/Debug/SVariableWatch.h b/Engine/Plugins/Experimental/UAF/UAF/Source/UAFEditor/Private/UAF/Debug/SVariableWatch.h
new file mode 100644
--- /dev/null
+++ b/Engine/Plugins/Experimental/UAF/UAF/Source/UAFEditor/Private/UAF/Debug/SVariableWatch.h
@@ -0,0 +1,99 @@
+// Copyright Epic Games, Inc. All Rights Reserved.
+
+#pragma once
+
+#include "CoreMinimal.h"
+#include "Editor.h"
+#include "Widgets/Views/STreeView.h"
+#include "StructUtils/PropertyBag.h"
+
+class IRewindDebugger;
+
+class UAnimNextRigVMAssetEntry;
+
+namespace UE::Workspace
+{
+	class IWorkspaceEditor;
+	struct FWorkspaceDocument;
+}
+
+namespace UE::UAF::Editor
+{
+	class SVariablesOutliner;
+
+	class SVariableWatch : public SCompoundWidget
+	{
+	public:
+		friend class SVariableWatchRow;
+		
+		SLATE_BEGIN_ARGS(SVariableWatch)
+		{}
+		SLATE_END_ARGS()
+
+		void Construct(const FArguments& InArgs, const TWeakPtr<UE::Workspace::IWorkspaceEditor>& InWorkspaceEditor);
+
+		struct FDebugInstance
+		{
+			uint64 InstanceId;
+			uint64 HostInstanceId;
+			UObject* Asset;
+		};
+
+		using FDebugInstanceTreeItemPtr = TSharedPtr<FDebugInstance>;
+
+		using FVariablesBoxItemPtr = TSharedPtr<FName>;
+	
+	private:
+		void HandleWorkspaceEditorModeChanged(const FEditorModeID& EditorModeId, bool bIsEnteringMode);
+		
+		void HandleRewindDebuggerUpdate(float DeltaTime, IRewindDebugger* RewindDebugger);
+		
+		void RegenerateDebugInstances(const bool bForceRegenerate);
+		void OnDebugInstanceChanged();
+
+		void HandleFocusedDocumentChanged(const UE::Workspace::FWorkspaceDocument& InDocument);
+		void HandleWorkspaceDebugObjectChanged(TWeakObjectPtr<UObject> InObject);
+
+		void InstanceTreeView_OnGetChildren(FDebugInstanceTreeItemPtr InItem, TArray<FDebugInstanceTreeItemPtr>& OutChildren);
+		TSharedRef<ITableRow> InstanceTreeView_OnGenerateRow(FDebugInstanceTreeItemPtr InItem, const TSharedRef<STableViewBase>& OwnerTable);
+		void InstanceTreeView_OnSelectionChanged(FDebugInstanceTreeItemPtr NewSelection, ESelectInfo::Type);
+
+		TSharedRef<ITableRow> VariablesListView_OnGenerateRow(const FVariablesBoxItemPtr InItem, const TSharedRef<STableViewBase>& OwnerTable);
+		
+	private:
+		IRewindDebugger* RewindDebugger = nullptr;
+	
+		FInstancedPropertyBag Properties;
+		FInstancedStruct NativeProperties;
+
+	private:
+		TWeakPtr<UE::Workspace::IWorkspaceEditor> WorkspaceEditor;
+
+		TArray<FDebugInstanceTreeItemPtr> RootDebugInstances;
+		TMap<uint64, FDebugInstanceTreeItemPtr> DebugInstances;
+
+		TSharedPtr<STreeView<FDebugInstanceTreeItemPtr>> DebugInstanceTreeView;
+		
+		FDebugInstanceTreeItemPtr SelectedDebugInstance;
+
+	public:
+		class FTreeItem
+		{
+		public:
+			TStrongObjectPtr<UAnimNextRigVMAssetEntry> AssetEntry;
+		};
+
+		using FTreeItemPtr = TSharedPtr<FTreeItem>;
+		
+		TSharedPtr<STreeView<FTreeItemPtr>> TreeView;
+		
+	private:
+		TSharedPtr<SListView<FVariablesBoxItemPtr>> VariablesBox;
+		
+		TArray<FVariablesBoxItemPtr> VariablesBoxSource;
+
+		TArray<FTreeItemPtr> TreeItems;
+		
+		uint64 SelectedDebugInstanceId = 0;
+	};
+}
\ No newline at end of file
diff --git a/Engine/Plugins/Experimental/UAF/UAF/Source/UAFEditor/Private/UAF/Debug/VariableWatchTabSummoner.cpp b/Engine/Plugins/Experimental/UAF/UAF/Source/UAFEditor/Private/UAF/Debug/VariableWatchTabSummoner.cpp
new file mode 100644
--- /dev/null
+++ b/Engine/Plugins/Experimental/UAF/UAF/Source/UAFEditor/Private/UAF/Debug/VariableWatchTabSummoner.cpp
@@ -0,0 +1,33 @@
+// Copyright Epic Games, Inc. All Rights Reserved.
+
+#include "VariableWatchTabSummoner.h"
+#include "SVariableWatch.h"
+
+#define LOCTEXT_NAMESPACE "UE::UAF::Editor::FVariableWatchTabSummoner"
+
+namespace UE::UAF::Editor
+{
+	FVariableWatchTabSummoner::FVariableWatchTabSummoner(TSharedPtr<UE::Workspace::IWorkspaceEditor> InHostingApp)
+		: FWorkflowTabFactory(TEXT("VariableWatch"), StaticCastSharedPtr<FAssetEditorToolkit>(InHostingApp))
+	{
+		TabLabel = LOCTEXT("TabLabel", "Variable Watch");
+		TabIcon = FSlateIcon(FAppStyle::GetAppStyleSetName(), "Kismet.Tabs.FindResults");
+		ViewMenuDescription = LOCTEXT("TabMenuDescription", "Variable Watch");
+		ViewMenuTooltip = LOCTEXT("TabMenuTooltip", "Track the values of variables.");
+		bIsSingleton = true;
+
+		Widget = SNew(SVariableWatch, InHostingApp);
+	}
+
+	TSharedRef<SWidget> FVariableWatchTabSummoner::CreateTabBody(const FWorkflowTabSpawnInfo& Info) const
+	{
+		return Widget.ToSharedRef();
+	}
+
+	FText FVariableWatchTabSummoner::GetTabToolTipText(const FWorkflowTabSpawnInfo& Info) const
+	{
+		return ViewMenuTooltip;
+	}
+}
+
+#undef LOCTEXT_NAMESPACE
\ No newline at end of file
diff --git a/Engine/Plugins/Experimental/UAF/UAF/Source/UAFEditor/Private/UAF/Debug/VariableWatchTabSummoner.h b/Engine/Plugins/Experimental/UAF/UAF/Source/UAFEditor/Private/UAF/Debug/VariableWatchTabSummoner.h
new file mode 100644
--- /dev/null
+++ b/Engine/Plugins/Experimental/UAF/UAF/Source/UAFEditor/Private/UAF/Debug/VariableWatchTabSummoner.h
@@ -0,0 +1,26 @@
+// Copyright Epic Games, Inc. All Rights Reserved.
+
+#pragma once
+
+#include "CoreMinimal.h"
+#include "IWorkspaceEditor.h"
+#include "WorkflowOrientedApp/WorkflowTabFactory.h"
+
+namespace UE::UAF::Editor
+{
+	class SVariableWatch;
+	
+	struct FVariableWatchTabSummoner : public FWorkflowTabFactory
+	{
+	public:
+		FVariableWatchTabSummoner(TSharedPtr<UE::Workspace::IWorkspaceEditor> InHostingApp);
+
+	private:
+		// FWorkflowTabFactory interface
+		virtual TSharedRef<SWidget> CreateTabBody(const FWorkflowTabSpawnInfo& Info) const override;
+		virtual FText GetTabToolTipText(const FWorkflowTabSpawnInfo& Info) const override;
+
+		// The widget this tab spawner wraps
+		TSharedPtr<SVariableWatch> Widget;
+	};
+}
diff --git a/Engine/Plugins/Experimental/UAF/UAF/Source/UAFEditor/Private/UAF/Workspace/IDebugObjectSelector.h b/Engine/Plugins/Experimental/UAF/UAF/Source/UAFEditor/Private/UAF/Workspace/IDebugObjectSelector.h
new file mode 100644
--- /dev/null
+++ b/Engine/Plugins/Experimental/UAF/UAF/Source/UAFEditor/Private/UAF/Workspace/IDebugObjectSelector.h
@@ -0,0 +1,22 @@
+// Copyright Epic Games, Inc. All Rights Reserved.
+
+#pragma once
+
+#include "CoreMinimal.h"
+
+namespace UE::UAF::Editor
+{
+	class IDebugObjectSelector
+	{
+	public:
+		virtual ~IDebugObjectSelector() = default;
+		
+		using FOnDebugObjectChanged = TMulticastDelegate<void(TWeakObjectPtr<UObject>)>;
+
+		virtual TWeakObjectPtr<UObject> GetDebugObject() const = 0;
+
+		virtual FDelegateHandle RegisterOnDebugObjectChanged(FOnDebugObjectChanged::FDelegate InDelegate) = 0;
+
+		virtual bool UnregisterOnDebugObjectChanged(const FDelegateHandle InHandle) = 0;
+	};
+}
\ No newline at end of file
diff --git a/Engine/Plugins/Experimental/UAF/UAF/Source/UAFEditor/Private/UAF/Workspace/IUAFWorkspaceRewindDebugger.h b/Engine/Plugins/Experimental/UAF/UAF/Source/UAFEditor/Private/UAF/Workspace/IUAFWorkspaceRewindDebugger.h
new file mode 100644
--- /dev/null
+++ b/Engine/Plugins/Experimental/UAF/UAF/Source/UAFEditor/Private/UAF/Workspace/IUAFWorkspaceRewindDebugger.h
@@ -0,0 +1,19 @@
+// Copyright Epic Games, Inc. All Rights Reserved.
+
+#pragma once
+
+#include "CoreMinimal.h"
+
+namespace UE::UAF::Editor
+{
+	class IUAFWorkspaceRewindDebugger
+	{
+	public:
+		virtual ~IUAFWorkspaceRewindDebugger() = default;
+		
+		using FOnRewindDebuggerUpdate = TMulticastDelegate<void(float, class IRewindDebugger*)>;
+
+		virtual FDelegateHandle RegisterOnRewindDebuggerUpdate(FOnRewindDebuggerUpdate::FDelegate InDelegate) = 0;
+		virtual bool UnregisterOnRewindDebuggerUpdate(const FDelegateHandle InHandle) = 0;
+	};
+}
\ No newline at end of file
diff --git a/Engine/Plugins/Experimental/UAF/UAF/Source/UAFEditor/Private/UAF/Workspace/RewindDebuggerExtension.cpp b/Engine/Plugins/Experimental/UAF/UAF/Source/UAFEditor/Private/UAF/Workspace/RewindDebuggerExtension.cpp
new file mode 100644
--- /dev/null
+++ b/Engine/Plugins/Experimental/UAF/UAF/Source/UAFEditor/Private/UAF/Workspace/RewindDebuggerExtension.cpp
@@ -0,0 +1,34 @@
+// Copyright Epic Games, Inc. All Rights Reserved.
+
+#include "UAF/Workspace/RewindDebuggerExtension.h"
+
+#include "IRewindDebugger.h"
+
+namespace UE::UAF::Editor
+{
+	FString FRewindDebuggerWorkspaceExtension::GetName()
+	{
+		return TEXT("UAF Workspace Extension");
+	}
+	
+	void FRewindDebuggerWorkspaceExtension::Update(float DeltaTime, IRewindDebugger* RewindDebugger)
+	{
+		if (RewindDebugger->GetScrubTime() != LastScrubTime)
+		{
+			LastScrubTime = RewindDebugger->GetScrubTime();
+
+			// Forward this update to workspace components
+			OnRewindDebuggerUpdate.Broadcast(DeltaTime, RewindDebugger);
+		}
+	}
+
+	FDelegateHandle FRewindDebuggerWorkspaceExtension::RegisterOnRewindDebuggerUpdate(FOnRewindDebuggerUpdate::FDelegate InDelegate)
+	{
+		return OnRewindDebuggerUpdate.Add(InDelegate);
+	}
+	
+	bool FRewindDebuggerWorkspaceExtension::UnregisterOnRewindDebuggerUpdate(const FDelegateHandle InHandle)
+	{
+		return OnRewindDebuggerUpdate.Remove(InHandle);
+	}
+}
\ No newline at end of file
diff --git a/Engine/Plugins/Experimental/UAF/UAF/Source/UAFEditor/Private/UAF/Workspace/RewindDebuggerExtension.h b/Engine/Plugins/Experimental/UAF/UAF/Source/UAFEditor/Private/UAF/Workspace/RewindDebuggerExtension.h
new file mode 100644
--- /dev/null
+++ b/Engine/Plugins/Experimental/UAF/UAF/Source/UAFEditor/Private/UAF/Workspace/RewindDebuggerExtension.h
@@ -0,0 +1,36 @@
+// Copyright Epic Games, Inc. All Rights Reserved.
+
+#pragma once
+
+#include "CoreMinimal.h"
+#include "IRewindDebuggerExtension.h"
+#include "UAF/Workspace/IUAFWorkspaceRewindDebugger.h"
+
+class IRewindDebugger;
+
+namespace UE::UAF::Editor
+{
+	class IWorkspaceEditor;
+
+	// FRewindDebuggerWorkspaceExtension exists to allow any systems interacting
+	// with the workspace editor to query a centralized rewind debugger extension
+	// instead of having to define and register their own.
+	class FRewindDebuggerWorkspaceExtension : public IRewindDebuggerExtension, public IUAFWorkspaceRewindDebugger
+	{
+	public:
+		virtual ~FRewindDebuggerWorkspaceExtension() override = default;
+
+		// IUAFWorkspaceRewindDebugger
+		virtual FDelegateHandle RegisterOnRewindDebuggerUpdate(FOnRewindDebuggerUpdate::FDelegate InDelegate) override;
+		virtual bool UnregisterOnRewindDebuggerUpdate(const FDelegateHandle InHandle) override;
+
+	private:
+		// IRewindDebuggerExtension
+		virtual FString GetName() override;
+		virtual void Update(float DeltaTime, IRewindDebugger* RewindDebugger) override;
+		
+	private:
+		FOnRewindDebuggerUpdate OnRewindDebuggerUpdate;
+		float LastScrubTime = 0.0f;
+	};
+}
diff --git a/Engine/Plugins/Experimental/UAF/UAF/Source/UAFEditor/Private/UAF/Workspace/SDebugObjectSelector.cpp b/Engine/Plugins/Experimental/UAF/UAF/Source/UAFEditor/Private/UAF/Workspace/SDebugObjectSelector.cpp
new file mode 100644
--- /dev/null
+++ b/Engine/Plugins/Experimental/UAF/UAF/Source/UAFEditor/Private/UAF/Workspace/SDebugObjectSelector.cpp
@@ -0,0 +1,91 @@
+// Copyright Epic Games, Inc. All Rights Reserved.
+
+#include "UAF/Workspace/SDebugObjectSelector.h"
+
+#include "EditorModeManager.h"
+#include "Widgets/Input/STextComboBox.h"
+
+#define LOCTEXT_NAMESPACE "UE::UAF::Editor::SDebugObjectSelector"
+
+namespace UE::UAF::Editor
+{
+	void SDebugObjectSelector::Construct(const FArguments& InArgs)
+	{
+		GetDebugObjectsDelegate = InArgs._GetDebugObjects;
+		
+		ChildSlot
+		[
+			SNew(SComboBox<TSharedPtr<FBlueprintDebugObjectInstance>>)
+			.OptionsSource(&DebugObjects)
+			.OnComboBoxOpening(this, &SDebugObjectSelector::ComboBox_OnComboBoxOpening)
+			.OnSelectionChanged(this, &SDebugObjectSelector::ComboBox_OnSelectionChanged)
+			.OnGenerateWidget(this, &SDebugObjectSelector::ComboBox_OnGenerateWidget)
+			.ContentPadding(FMargin(0.f, 4.f))
+			[
+				SNew(STextBlock)
+				.Text_Lambda([this]()
+				{
+					return Selected.IsValid()
+						? FText::FromString(Selected->ObjectLabel)
+						: LOCTEXT("NoDebugObjectSelected", "No debug object selected");
+				})
+			]
+		];
+	}
+
+	TSharedRef<SWidget> SDebugObjectSelector::ComboBox_OnGenerateWidget(TSharedPtr<FBlueprintDebugObjectInstance> InItem)
+	{
+		FString ItemString;
+		FString ItemTooltip;
+
+		if (InItem.IsValid())
+		{
+			ItemString = InItem->ObjectLabel;
+			ItemTooltip = InItem->ObjectPath;
+		}
+
+		return SNew(STextBlock)
+			.Text(FText::FromString(*ItemString))
+			.ToolTipText(FText::FromString(*ItemTooltip));
+	}
+
+	void SDebugObjectSelector::ComboBox_OnSelectionChanged(TSharedPtr<FBlueprintDebugObjectInstance> NewSelection, ESelectInfo::Type SelectInfo)
+	{
+		Selected = NewSelection;
+		OnDebugObjectChangedDelegate.Broadcast(Selected ? Selected->ObjectPtr : TWeakObjectPtr<UObject>());
+	}
+	
+	void SDebugObjectSelector::ComboBox_OnComboBoxOpening()
+	{
+		RegenerateDebugObjects();
+	}
+	
+	void SDebugObjectSelector::RegenerateDebugObjects()
+	{
+		TArray<UObject*> DebugObjectPtrs;
+		GetDebugObjectsDelegate.ExecuteIfBound(DebugObjectPtrs);
+		
+		DebugObjects.Empty();
+		for (UObject* DebugObject : DebugObjectPtrs)
+		{
+			DebugObjects.Add(MakeShared<FBlueprintDebugObjectInstance>(DebugObject, DebugObject->GetFullName()));
+		}
+	}
+	
+	TWeakObjectPtr<UObject> SDebugObjectSelector::GetDebugObject() const
+	{
+		return Selected.IsValid() ? Selected->ObjectPtr : nullptr;
+	}
+
+	FDelegateHandle SDebugObjectSelector::RegisterOnDebugObjectChanged(FOnDebugObjectChanged::FDelegate InDelegate)
+	{
+		return OnDebugObjectChangedDelegate.Add(InDelegate);
+	}
+	
+	bool SDebugObjectSelector::UnregisterOnDebugObjectChanged(const FDelegateHandle InHandle)
+	{
+		return OnDebugObjectChangedDelegate.Remove(InHandle);
+	}
+}
+
+#undef LOCTEXT_NAMESPACE
\ No newline at end of file
diff --git a/Engine/Plugins/Experimental/UAF/UAF/Source/UAFEditor/Private/UAF/Workspace/SDebugObjectSelector.h b/Engine/Plugins/Experimental/UAF/UAF/Source/UAFEditor/Private/UAF/Workspace/SDebugObjectSelector.h
new file mode 100644
--- /dev/null
+++ b/Engine/Plugins/Experimental/UAF/UAF/Source/UAFEditor/Private/UAF/Workspace/SDebugObjectSelector.h
@@ -0,0 +1,41 @@
+// Copyright Epic Games, Inc. All Rights Reserved.
+
+#pragma once
+
+#include "CoreMinimal.h"
+#include "SBlueprintEditorSelectedDebugObjectWidget.h"
+#include "UAF/Workspace/IDebugObjectSelector.h"
+
+namespace UE::UAF::Editor
+{
+	class SDebugObjectSelector : public SCompoundWidget, public IDebugObjectSelector
+	{
+	public:
+		DECLARE_DELEGATE_OneParam(FGetDebugObjects, TArray<UObject*>&)
+		
+		SLATE_BEGIN_ARGS(SDebugObjectSelector)
+		{}
+			SLATE_EVENT(FGetDebugObjects, GetDebugObjects)
+		SLATE_END_ARGS()
+    
+		void Construct(const FArguments& InArgs);
+
+		// IDebugObjectSelector
+    	virtual TWeakObjectPtr<UObject> GetDebugObject() const override;
+		virtual FDelegateHandle RegisterOnDebugObjectChanged(FOnDebugObjectChanged::FDelegate InDelegate) override;
+		virtual bool UnregisterOnDebugObjectChanged(const FDelegateHandle InHandle) override;
+    		
+	private:
+		TSharedRef<SWidget> ComboBox_OnGenerateWidget(TSharedPtr<FBlueprintDebugObjectInstance> InItem);
+		void ComboBox_OnSelectionChanged(TSharedPtr<FBlueprintDebugObjectInstance> NewSelection, ESelectInfo::Type SelectInfo);
+		void ComboBox_OnComboBoxOpening();
+		
+		void RegenerateDebugObjects();
+
+		FOnDebugObjectChanged OnDebugObjectChangedDelegate;
+		FGetDebugObjects GetDebugObjectsDelegate;
+		
+		TArray<TSharedPtr<FBlueprintDebugObjectInstance>> DebugObjects;
+		TSharedPtr<FBlueprintDebugObjectInstance> Selected;
+	};
+}
\ No newline at end of file
diff --git a/Engine/Plugins/Experimental/UAF/UAF/Source/UAFEditor/Private/Workspace/AnimNextWorkspaceEditorMode.cpp b/Engine/Plugins/Experimental/UAF/UAF/Source/UAFEditor/Private/Workspace/AnimNextWorkspaceEditorMode.cpp
--- a/Engine/Plugins/Experimental/UAF/UAF/Source/UAFEditor/Private/Workspace/AnimNextWorkspaceEditorMode.cpp
+++ b/Engine/Plugins/Experimental/UAF/UAF/Source/UAFEditor/Private/Workspace/AnimNextWorkspaceEditorMode.cpp
@@ -6,21 +6,25 @@
 #include "AnimNextRigVMAsset.h"
 #include "AnimNextRigVMAssetEditorData.h"
 #include "AnimNextScopedCompileJob.h"
+#include "Component/AnimNextComponent.h"
 #include "ContextObjectStore.h"
 #include "EditorModeManager.h"
 #include "IAssetCompilationHandler.h"
 #include "InteractiveToolManager.h"
 #include "IWorkspaceEditor.h"
 #include "IWorkspaceEditorModule.h"
+#include "Modules/ModuleManager.h"
 #include "RigVMCommands.h"
 #include "RigVMCore/RigVMGraphFunctionHost.h"
 #include "RigVMModel/Nodes/RigVMFunctionReferenceNode.h"
-#include "UncookedOnlyUtils.h"
-#include "Logging/MessageLog.h"
-#include "Modules/ModuleManager.h"
 #include "Toolkits/AssetEditorToolkitMenuContext.h"
 #include "Toolkits/BaseToolkit.h"
+#include "ToolMenus.h"
 #include "Tools/EdModeInteractiveToolsContext.h"
+#include "UAF/Workspace/RewindDebuggerExtension.h"
+#include "UAF/Workspace/SDebugObjectSelector.h"
+#include "UObject/UObjectIterator.h"
+#include "UncookedOnlyUtils.h"
 
 #include UE_INLINE_GENERATED_CPP_BY_NAME(AnimNextWorkspaceEditorMode)
 
@@ -64,6 +68,21 @@ void UAnimNextWorkspaceEditorMode::Enter()
 	}
 
 	UpdateCompileStatus();
+	
+	RewindDebuggerExtension = MakeShared<UE::UAF::Editor::FRewindDebuggerWorkspaceExtension>();
+	IModularFeatures::Get().RegisterModularFeature(IRewindDebuggerExtension::ModularFeatureName, RewindDebuggerExtension.Get());
+	
+	if (UToolMenu* Menu = UToolMenus::Get()->ExtendMenu(WorkspaceEditor->GetToolMenuToolbarName()))
+	{
+		FToolMenuSection& DebugSection = Menu->AddSection("Debug", FText::GetEmpty(), FToolMenuInsert("WorkspaceOperations", EToolMenuInsertType::After));
+
+		TSharedRef<UE::UAF::Editor::SDebugObjectSelector> DebugObjectSelectorWidget = SNew(UE::UAF::Editor::SDebugObjectSelector)
+			.GetDebugObjects_UObject(this, &UAnimNextWorkspaceEditorMode::GetDebugObjects);
+
+		DebugObjectSelector = DebugObjectSelectorWidget;
+
+		DebugSection.AddEntry(FToolMenuEntry::InitWidget("DebugObjectSelector", DebugObjectSelectorWidget, FText::GetEmpty(), true, false, true));
+	}
 }
 
 void UAnimNextWorkspaceEditorMode::Exit()
@@ -77,6 +96,11 @@ void UAnimNextWorkspaceEditorMode::Exit()
 		CommandList->UnmapAction(RigVMCommands.AutoCompile);
 		CommandList->UnmapAction(RigVMCommands.CompileWholeWorkspace);
 		CommandList->UnmapAction(RigVMCommands.CompileDirtyFiles);
+
+		if (UToolMenu* Menu = UToolMenus::Get()->ExtendMenu(WorkspaceEditor->GetToolMenuToolbarName()))
+		{
+			Menu->RemoveSection("Debug");
+		}
 	}
 
 	Super::Exit();
@@ -104,6 +128,8 @@ void UAnimNextWorkspaceEditorMode::Exit()
 	}
 
 	WeakAssets.Reset();
+
+	IModularFeatures::Get().UnregisterModularFeature(IRewindDebuggerExtension::ModularFeatureName, RewindDebuggerExtension.Get());
 }
 
 void UAnimNextWorkspaceEditorMode::BindCommands()
@@ -437,4 +463,49 @@ TSharedPtr<UE::UAF::Editor::IAssetCompilationHandler> UAnimNextWorkspaceEditorMo
 	return nullptr;
 }
 
+void UAnimNextWorkspaceEditorMode::GetDebugObjects(TArray<UObject*>& OutObjects)
+{
+	for (TObjectIterator<UAnimNextComponent> It; It; ++It)
+	{
+		UAnimNextComponent* Object = *It;
+
+		UWorld* World = Object->GetWorld();
+		if (!World || Object->HasAnyFlags(RF_ClassDefaultObject))
+		{
+			continue;
+		}
+
+		bool bPendingKill = false;
+		UObject* Outer = Object;
+		do
+		{
+			bPendingKill = !IsValid(Outer);
+			Outer = Outer->GetOuter();
+		}
+		while (!bPendingKill && Outer != nullptr);
+
+		if (bPendingKill)
+		{
+			continue;
+		}
+
+		if (World->WorldType != EWorldType::PIE)
+		{
+			continue;
+		}
+		
+		OutObjects.Add(Object);
+	}
+}
+
+TSharedPtr<UE::UAF::Editor::FRewindDebuggerWorkspaceExtension> UAnimNextWorkspaceEditorMode::GetRewindDebuggerExtension()
+{
+	return RewindDebuggerExtension;
+}
+
+TSharedPtr<UE::UAF::Editor::IDebugObjectSelector> UAnimNextWorkspaceEditorMode::GetDebugObjectSelector()
+{
+	return DebugObjectSelector;
+}
+
 #undef LOCTEXT_NAMESPACE
diff --git a/Engine/Plugins/Experimental/UAF/UAF/Source/UAFEditor/Private/Workspace/AnimNextWorkspaceEditorMode.h b/Engine/Plugins/Experimental/UAF/UAF/Source/UAFEditor/Private/Workspace/AnimNextWorkspaceEditorMode.h
--- a/Engine/Plugins/Experimental/UAF/UAF/Source/UAFEditor/Private/Workspace/AnimNextWorkspaceEditorMode.h
+++ b/Engine/Plugins/Experimental/UAF/UAF/Source/UAFEditor/Private/Workspace/AnimNextWorkspaceEditorMode.h
@@ -5,6 +5,7 @@
 #include "AnimNextWorkspaceState.h"
 #include "EdMode.h"
 #include "IAssetCompilationHandler.h"
+
 #include "AnimNextWorkspaceEditorMode.generated.h"
 
 class UAnimNextWorkspaceSchema;
@@ -22,6 +23,8 @@ class IWorkspaceEditor;
 namespace UE::UAF::Editor
 {
 	class FAssetWizard;
+	class FRewindDebuggerWorkspaceExtension;
+	class IDebugObjectSelector;
 }
 
 UCLASS(Transient)
@@ -40,6 +43,10 @@ public:
 	// Get the currently stored workspace state
 	const FAnimNextWorkspaceState& GetState() const { return State; } 
 
+	TSharedPtr<UE::UAF::Editor::FRewindDebuggerWorkspaceExtension> GetRewindDebuggerExtension();
+
+	TSharedPtr<UE::UAF::Editor::IDebugObjectSelector> GetDebugObjectSelector();
+
 private:
 	// UEdMode interface
 	virtual bool IsCompatibleWith(FEditorModeID OtherModeID) const override { return true; }
@@ -76,6 +83,8 @@ private:
 	// Lazily construct an asset compiler for the supplied asset
 	TSharedPtr<UE::UAF::Editor::IAssetCompilationHandler> GetAssetCompiler(UObject* InAsset);
 
+	void GetDebugObjects(TArray<UObject*>& OutObjects);
+
 	friend UAnimNextWorkspaceSchema;
 	friend UE::UAF::Editor::FAssetWizard;
 
@@ -93,4 +102,8 @@ private:
 
 	// Flag to update the latest compile status
 	bool bUpdateCompileStatus = false;
+
+	TSharedPtr<UE::UAF::Editor::FRewindDebuggerWorkspaceExtension> RewindDebuggerExtension;
+
+	TSharedPtr<UE::UAF::Editor::IDebugObjectSelector> DebugObjectSelector;
 };
\ No newline at end of file
