[Backout] - CL46524649
#rnx
Original CL Desc
-----------------------------------------------------------------
Variable Watch 0.1
* All very subject to change.
* Foundation work for the workspace interfacing with rewind debugger.
* When a rewind debugger session is active and a UAF workspace is active, debug objects are selectable in the workspace for objects alive at the rewind debugger playhead.
* The variable watch tab currently will display the hierarchy of graph instances and upon selecting an instance will show the variable values of variables in that instance.
* This tab will morph into a more traditional variable watch window similar to visual studio once the design has evolved.

#rnx

diff --git a/Engine/Plugins/Experimental/UAF/UAF/Source/UAFEditor/Private/AnimNextEditorModule.cpp b/Engine/Plugins/Experimental/UAF/UAF/Source/UAFEditor/Private/AnimNextEditorModule.cpp
--- a/Engine/Plugins/Experimental/UAF/UAF/Source/UAFEditor/Private/AnimNextEditorModule.cpp
+++ b/Engine/Plugins/Experimental/UAF/UAF/Source/UAFEditor/Private/AnimNextEditorModule.cpp
@@ -81,7 +81,6 @@
 #include "Wizard/AssetWizard.h"
 #include "Interfaces/IPluginManager.h"
 #include "AssetToolsModule.h"
-#include "UAF/Debug/VariableWatchTabSummoner.h"
 
 #define LOCTEXT_NAMESPACE "AnimNextEditorModule"
 
@@ -217,10 +216,6 @@ void FAnimNextEditorModule::StartupModule()
 			TSharedRef<FAnimNextFindAndReplaceTabSummoner> FindAndReplaceTabSummoner = MakeShared<FAnimNextFindAndReplaceTabSummoner>(InEditorPtr);
 			TabFactories.RegisterFactory(FindAndReplaceTabSummoner);
 			FindAndReplaceTabSummoner->RegisterTabSpawner(InTabManager, nullptr);
-
-			TSharedRef<FVariableWatchTabSummoner> VariableWatchTabSummoner = MakeShared<FVariableWatchTabSummoner>(InEditorPtr);
-			TabFactories.RegisterFactory(VariableWatchTabSummoner);
-			VariableWatchTabSummoner->RegisterTabSpawner(InTabManager, nullptr);
 		});
 
 	WorkspaceEditorModule.OnExtendTabs().AddLambda([](FLayoutExtender& InLayoutExtender, TSharedPtr<UE::Workspace::IWorkspaceEditor> InEditorPtr)
diff --git a/Engine/Plugins/Experimental/UAF/UAF/Source/UAFEditor/Private/RewindDebugger/AnimNextProvider.cpp b/Engine/Plugins/Experimental/UAF/UAF/Source/UAFEditor/Private/RewindDebugger/AnimNextProvider.cpp
--- a/Engine/Plugins/Experimental/UAF/UAF/Source/UAFEditor/Private/RewindDebugger/AnimNextProvider.cpp
+++ b/Engine/Plugins/Experimental/UAF/UAF/Source/UAFEditor/Private/RewindDebugger/AnimNextProvider.cpp
@@ -17,9 +17,9 @@ void FAnimNextProvider::AppendInstance(uint64 InstanceId, uint64 HostInstanceId,
 {
 	Session.WriteAccessCheck();
 
-	if (HostInstanceId == OuterObjectId)
+	if (HostInstanceId == 0)
 	{
-		ComponentIdToSystemId.Add(OuterObjectId, InstanceId);
+		ComponentIdToModuleId.Add(OuterObjectId, InstanceId);
 	}
 
 	TSharedRef<FDataInterfaceData> Data = MakeShared<FDataInterfaceData>(Session);
@@ -70,7 +70,7 @@ bool FAnimNextProvider::GetModuleId(uint64 ComponentId, uint64& OutModuleId) con
 {
 	Session.ReadAccessCheck();
 	
-	if (const uint64* Id = ComponentIdToSystemId.Find(ComponentId))
+	if (const uint64* Id = ComponentIdToModuleId.Find(ComponentId))
 	{
 		OutModuleId = *Id;
 		return true;
@@ -103,18 +103,5 @@ void FAnimNextProvider::EnumerateChildInstances(uint64 InstanceId, TFunctionRef<
 
 }
 
-void FAnimNextProvider::EnumerateDataInterfaces(const double InStartTime, TFunctionRef<void(const TSharedRef<const FDataInterfaceData>)> Callback) const
-{
-	Session.ReadAccessCheck();
-
-	for (const auto& [InstanceId, InterfaceData] : DataInterfaceData)
-	{
-		if (InterfaceData->StartTime <= InStartTime && InterfaceData->EndTime >= InStartTime)
-		{
-			Callback(InterfaceData);
-		}
-	}
-}
-
 #undef LOCTEXT_NAMESPACE
 #endif // ANIMNEXT_TRACE_ENABLED
diff --git a/Engine/Plugins/Experimental/UAF/UAF/Source/UAFEditor/Private/RewindDebugger/AnimNextProvider.h b/Engine/Plugins/Experimental/UAF/UAF/Source/UAFEditor/Private/RewindDebugger/AnimNextProvider.h
--- a/Engine/Plugins/Experimental/UAF/UAF/Source/UAFEditor/Private/RewindDebugger/AnimNextProvider.h
+++ b/Engine/Plugins/Experimental/UAF/UAF/Source/UAFEditor/Private/RewindDebugger/AnimNextProvider.h
@@ -74,12 +74,11 @@ public:
 	const FPropertyDescriptionData* GetPropertyDescriptionData(uint32 PropertyDescriptionHash) const { return HashToPropertyDescriptionDataMap.Find(PropertyDescriptionHash);}
 
 	void EnumerateChildInstances(uint64 InstanceId, TFunctionRef<void(const FDataInterfaceData&)> Callback) const;
-	void EnumerateDataInterfaces(const double StartTime, TFunctionRef<void(const TSharedRef<const FDataInterfaceData>)> Callback) const;
 
 private:
 	TraceServices::IAnalysisSession& Session;
 
-	TMap<uint64, uint64> ComponentIdToSystemId;
+	TMap<uint64, uint64> ComponentIdToModuleId;
 	TMap<uint64, TSharedRef<FDataInterfaceData>> DataInterfaceData;
 	TMap<uint64, TArray<TSharedRef<FDataInterfaceData>>> HostToChildDataMap;
 	TMap<uint32, FPropertyDescriptionData> HashToPropertyDescriptionDataMap;
diff --git a/Engine/Plugins/Experimental/UAF/UAF/Source/UAFEditor/Private/UAF/Debug/SVariableWatch.cpp b/Engine/Plugins/Experimental/UAF/UAF/Source/UAFEditor/Private/UAF/Debug/SVariableWatch.cpp
deleted file mode 100644
diff --git a/Engine/Plugins/Experimental/UAF/UAF/Source/UAFEditor/Private/UAF/Debug/SVariableWatch.h b/Engine/Plugins/Experimental/UAF/UAF/Source/UAFEditor/Private/UAF/Debug/SVariableWatch.h
deleted file mode 100644
diff --git a/Engine/Plugins/Experimental/UAF/UAF/Source/UAFEditor/Private/UAF/Debug/VariableWatchTabSummoner.cpp b/Engine/Plugins/Experimental/UAF/UAF/Source/UAFEditor/Private/UAF/Debug/VariableWatchTabSummoner.cpp
deleted file mode 100644
diff --git a/Engine/Plugins/Experimental/UAF/UAF/Source/UAFEditor/Private/UAF/Debug/VariableWatchTabSummoner.h b/Engine/Plugins/Experimental/UAF/UAF/Source/UAFEditor/Private/UAF/Debug/VariableWatchTabSummoner.h
deleted file mode 100644
diff --git a/Engine/Plugins/Experimental/UAF/UAF/Source/UAFEditor/Private/UAF/Workspace/IDebugObjectSelector.h b/Engine/Plugins/Experimental/UAF/UAF/Source/UAFEditor/Private/UAF/Workspace/IDebugObjectSelector.h
deleted file mode 100644
diff --git a/Engine/Plugins/Experimental/UAF/UAF/Source/UAFEditor/Private/UAF/Workspace/IUAFWorkspaceRewindDebugger.h b/Engine/Plugins/Experimental/UAF/UAF/Source/UAFEditor/Private/UAF/Workspace/IUAFWorkspaceRewindDebugger.h
deleted file mode 100644
diff --git a/Engine/Plugins/Experimental/UAF/UAF/Source/UAFEditor/Private/UAF/Workspace/RewindDebuggerExtension.cpp b/Engine/Plugins/Experimental/UAF/UAF/Source/UAFEditor/Private/UAF/Workspace/RewindDebuggerExtension.cpp
deleted file mode 100644
diff --git a/Engine/Plugins/Experimental/UAF/UAF/Source/UAFEditor/Private/UAF/Workspace/RewindDebuggerExtension.h b/Engine/Plugins/Experimental/UAF/UAF/Source/UAFEditor/Private/UAF/Workspace/RewindDebuggerExtension.h
deleted file mode 100644
diff --git a/Engine/Plugins/Experimental/UAF/UAF/Source/UAFEditor/Private/UAF/Workspace/SDebugObjectSelector.cpp b/Engine/Plugins/Experimental/UAF/UAF/Source/UAFEditor/Private/UAF/Workspace/SDebugObjectSelector.cpp
deleted file mode 100644
diff --git a/Engine/Plugins/Experimental/UAF/UAF/Source/UAFEditor/Private/UAF/Workspace/SDebugObjectSelector.h b/Engine/Plugins/Experimental/UAF/UAF/Source/UAFEditor/Private/UAF/Workspace/SDebugObjectSelector.h
deleted file mode 100644
diff --git a/Engine/Plugins/Experimental/UAF/UAF/Source/UAFEditor/Private/Workspace/AnimNextWorkspaceEditorMode.cpp b/Engine/Plugins/Experimental/UAF/UAF/Source/UAFEditor/Private/Workspace/AnimNextWorkspaceEditorMode.cpp
--- a/Engine/Plugins/Experimental/UAF/UAF/Source/UAFEditor/Private/Workspace/AnimNextWorkspaceEditorMode.cpp
+++ b/Engine/Plugins/Experimental/UAF/UAF/Source/UAFEditor/Private/Workspace/AnimNextWorkspaceEditorMode.cpp
@@ -6,25 +6,21 @@
 #include "AnimNextRigVMAsset.h"
 #include "AnimNextRigVMAssetEditorData.h"
 #include "AnimNextScopedCompileJob.h"
-#include "Component/AnimNextComponent.h"
 #include "ContextObjectStore.h"
 #include "EditorModeManager.h"
 #include "IAssetCompilationHandler.h"
 #include "InteractiveToolManager.h"
 #include "IWorkspaceEditor.h"
 #include "IWorkspaceEditorModule.h"
-#include "Modules/ModuleManager.h"
 #include "RigVMCommands.h"
 #include "RigVMCore/RigVMGraphFunctionHost.h"
 #include "RigVMModel/Nodes/RigVMFunctionReferenceNode.h"
+#include "UncookedOnlyUtils.h"
+#include "Logging/MessageLog.h"
+#include "Modules/ModuleManager.h"
 #include "Toolkits/AssetEditorToolkitMenuContext.h"
 #include "Toolkits/BaseToolkit.h"
-#include "ToolMenus.h"
 #include "Tools/EdModeInteractiveToolsContext.h"
-#include "UAF/Workspace/RewindDebuggerExtension.h"
-#include "UAF/Workspace/SDebugObjectSelector.h"
-#include "UObject/UObjectIterator.h"
-#include "UncookedOnlyUtils.h"
 
 #include UE_INLINE_GENERATED_CPP_BY_NAME(AnimNextWorkspaceEditorMode)
 
@@ -68,21 +64,6 @@ void UAnimNextWorkspaceEditorMode::Enter()
 	}
 
 	UpdateCompileStatus();
-	
-	RewindDebuggerExtension = MakeShared<UE::UAF::Editor::FRewindDebuggerWorkspaceExtension>();
-	IModularFeatures::Get().RegisterModularFeature(IRewindDebuggerExtension::ModularFeatureName, RewindDebuggerExtension.Get());
-	
-	if (UToolMenu* Menu = UToolMenus::Get()->ExtendMenu(WorkspaceEditor->GetToolMenuToolbarName()))
-	{
-		FToolMenuSection& DebugSection = Menu->AddSection("Debug", FText::GetEmpty(), FToolMenuInsert("WorkspaceOperations", EToolMenuInsertType::After));
-
-		TSharedRef<UE::UAF::Editor::SDebugObjectSelector> DebugObjectSelectorWidget = SNew(UE::UAF::Editor::SDebugObjectSelector)
-			.GetDebugObjects_UObject(this, &UAnimNextWorkspaceEditorMode::GetDebugObjects);
-
-		DebugObjectSelector = DebugObjectSelectorWidget;
-
-		DebugSection.AddEntry(FToolMenuEntry::InitWidget("DebugObjectSelector", DebugObjectSelectorWidget, FText::GetEmpty(), true, false, true));
-	}
 }
 
 void UAnimNextWorkspaceEditorMode::Exit()
@@ -96,11 +77,6 @@ void UAnimNextWorkspaceEditorMode::Exit()
 		CommandList->UnmapAction(RigVMCommands.AutoCompile);
 		CommandList->UnmapAction(RigVMCommands.CompileWholeWorkspace);
 		CommandList->UnmapAction(RigVMCommands.CompileDirtyFiles);
-
-		if (UToolMenu* Menu = UToolMenus::Get()->ExtendMenu(WorkspaceEditor->GetToolMenuToolbarName()))
-		{
-			Menu->RemoveSection("Debug");
-		}
 	}
 
 	Super::Exit();
@@ -128,8 +104,6 @@ void UAnimNextWorkspaceEditorMode::Exit()
 	}
 
 	WeakAssets.Reset();
-
-	IModularFeatures::Get().UnregisterModularFeature(IRewindDebuggerExtension::ModularFeatureName, RewindDebuggerExtension.Get());
 }
 
 void UAnimNextWorkspaceEditorMode::BindCommands()
@@ -463,49 +437,4 @@ TSharedPtr<UE::UAF::Editor::IAssetCompilationHandler> UAnimNextWorkspaceEditorMo
 	return nullptr;
 }
 
-void UAnimNextWorkspaceEditorMode::GetDebugObjects(TArray<UObject*>& OutObjects)
-{
-	for (TObjectIterator<UAnimNextComponent> It; It; ++It)
-	{
-		UAnimNextComponent* Object = *It;
-
-		UWorld* World = Object->GetWorld();
-		if (!World || Object->HasAnyFlags(RF_ClassDefaultObject))
-		{
-			continue;
-		}
-
-		bool bPendingKill = false;
-		UObject* Outer = Object;
-		do
-		{
-			bPendingKill = !IsValid(Outer);
-			Outer = Outer->GetOuter();
-		}
-		while (!bPendingKill && Outer != nullptr);
-
-		if (bPendingKill)
-		{
-			continue;
-		}
-
-		if (World->WorldType != EWorldType::PIE)
-		{
-			continue;
-		}
-		
-		OutObjects.Add(Object);
-	}
-}
-
-TSharedPtr<UE::UAF::Editor::FRewindDebuggerWorkspaceExtension> UAnimNextWorkspaceEditorMode::GetRewindDebuggerExtension()
-{
-	return RewindDebuggerExtension;
-}
-
-TSharedPtr<UE::UAF::Editor::IDebugObjectSelector> UAnimNextWorkspaceEditorMode::GetDebugObjectSelector()
-{
-	return DebugObjectSelector;
-}
-
 #undef LOCTEXT_NAMESPACE
diff --git a/Engine/Plugins/Experimental/UAF/UAF/Source/UAFEditor/Private/Workspace/AnimNextWorkspaceEditorMode.h b/Engine/Plugins/Experimental/UAF/UAF/Source/UAFEditor/Private/Workspace/AnimNextWorkspaceEditorMode.h
--- a/Engine/Plugins/Experimental/UAF/UAF/Source/UAFEditor/Private/Workspace/AnimNextWorkspaceEditorMode.h
+++ b/Engine/Plugins/Experimental/UAF/UAF/Source/UAFEditor/Private/Workspace/AnimNextWorkspaceEditorMode.h
@@ -5,7 +5,6 @@
 #include "AnimNextWorkspaceState.h"
 #include "EdMode.h"
 #include "IAssetCompilationHandler.h"
-
 #include "AnimNextWorkspaceEditorMode.generated.h"
 
 class UAnimNextWorkspaceSchema;
@@ -23,8 +22,6 @@ class IWorkspaceEditor;
 namespace UE::UAF::Editor
 {
 	class FAssetWizard;
-	class FRewindDebuggerWorkspaceExtension;
-	class IDebugObjectSelector;
 }
 
 UCLASS(Transient)
@@ -43,10 +40,6 @@ public:
 	// Get the currently stored workspace state
 	const FAnimNextWorkspaceState& GetState() const { return State; } 
 
-	TSharedPtr<UE::UAF::Editor::FRewindDebuggerWorkspaceExtension> GetRewindDebuggerExtension();
-
-	TSharedPtr<UE::UAF::Editor::IDebugObjectSelector> GetDebugObjectSelector();
-
 private:
 	// UEdMode interface
 	virtual bool IsCompatibleWith(FEditorModeID OtherModeID) const override { return true; }
@@ -83,8 +76,6 @@ private:
 	// Lazily construct an asset compiler for the supplied asset
 	TSharedPtr<UE::UAF::Editor::IAssetCompilationHandler> GetAssetCompiler(UObject* InAsset);
 
-	void GetDebugObjects(TArray<UObject*>& OutObjects);
-
 	friend UAnimNextWorkspaceSchema;
 	friend UE::UAF::Editor::FAssetWizard;
 
@@ -102,8 +93,4 @@ private:
 
 	// Flag to update the latest compile status
 	bool bUpdateCompileStatus = false;
-
-	TSharedPtr<UE::UAF::Editor::FRewindDebuggerWorkspaceExtension> RewindDebuggerExtension;
-
-	TSharedPtr<UE::UAF::Editor::IDebugObjectSelector> DebugObjectSelector;
 };
\ No newline at end of file
