Patchup issues found by bug-hawk

diff --git a/Engine/Plugins/Experimental/UAF/UAF/Source/UAFEditor/Private/Common/Outliner/CommonOutlinerMode.cpp b/Engine/Plugins/Experimental/UAF/UAF/Source/UAFEditor/Private/Common/Outliner/CommonOutlinerMode.cpp
--- a/Engine/Plugins/Experimental/UAF/UAF/Source/UAFEditor/Private/Common/Outliner/CommonOutlinerMode.cpp
+++ b/Engine/Plugins/Experimental/UAF/UAF/Source/UAFEditor/Private/Common/Outliner/CommonOutlinerMode.cpp
@@ -99,7 +99,7 @@ FCommonOutlinerMode::FCommonOutlinerMode(SSceneOutliner* InOutliner, const TShar
 			{
 				const FText EntryCategoryText = FText::FromStringView(ItemCategoryPath);
 				const FText DropCategoryText = FText::FromString(DropCategoryItem->ParentCategoryName);
-				const FString NewCategoryName = FString(ItemCategoryPath) + TEXT("") + DropCategoryItem->CategoryName;
+				const FString NewCategoryName = FString(ItemCategoryPath) + TEXT("|") + DropCategoryItem->CategoryName;
 
 				UAnimNextRigVMAsset* ThisAsset = DroppedOnItem->WeakOwner.Get();
 				UAnimNextRigVMAsset* OtherAsset = DropCategoryItem->WeakOwner.Get();
diff --git a/Engine/Plugins/Experimental/UAF/UAF/Source/UAFEditor/Private/Variables/Outliner/VariablesOutlinerEntryItem.cpp b/Engine/Plugins/Experimental/UAF/UAF/Source/UAFEditor/Private/Variables/Outliner/VariablesOutlinerEntryItem.cpp
--- a/Engine/Plugins/Experimental/UAF/UAF/Source/UAFEditor/Private/Variables/Outliner/VariablesOutlinerEntryItem.cpp
+++ b/Engine/Plugins/Experimental/UAF/UAF/Source/UAFEditor/Private/Variables/Outliner/VariablesOutlinerEntryItem.cpp
@@ -138,7 +138,6 @@ EAnimNextExportAccessSpecifier FVariablesOutlinerEntryItem::GetAccessSpecifier()
 		return EAnimNextExportAccessSpecifier::Private;
 	}
 
-	FScopedTransaction Transaction(LOCTEXT("SetAccessSpecifierTransaction", "Set Access Specifier"));
 	return Export->GetExportAccessSpecifier();
 }
 
@@ -172,4 +171,4 @@ bool FVariablesOutlinerEntryItem::IsReadOnly() const
 }
 }
 
-#undef LOCTEXT_NAMESPACE // "VariablesOutlinerTreeItem"
\ No newline at end of file
+#undef LOCTEXT_NAMESPACE // "VariablesOutlinerTreeItem"
