Adding a new shared variable within a System asset doesn't show the new variable until you reload or reopen a graph
#fix update outliner assets for any RigVM change as well

diff --git a/Engine/Plugins/Experimental/UAF/UAF/Source/UAFEditor/Private/Variables/SVariablesView.cpp b/Engine/Plugins/Experimental/UAF/UAF/Source/UAFEditor/Private/Variables/SVariablesView.cpp
--- a/Engine/Plugins/Experimental/UAF/UAF/Source/UAFEditor/Private/Variables/SVariablesView.cpp
+++ b/Engine/Plugins/Experimental/UAF/UAF/Source/UAFEditor/Private/Variables/SVariablesView.cpp
@@ -93,11 +93,10 @@ void SVariablesView::OnEditorDataModified(UAnimNextRigVMAssetEditorData* InEdito
 	case EAnimNextEditorDataNotifType::VariableCategoryChanged:
 	case EAnimNextEditorDataNotifType::CategoryAdded:
 	case EAnimNextEditorDataNotifType::CategoryChanged:
-	case EAnimNextEditorDataNotifType::CategoryRemoved:
-		VariablesOutliner->FullRefresh();
-		break;
+	case EAnimNextEditorDataNotifType::CategoryRemoved:	
 	case EAnimNextEditorDataNotifType::EntryRemoved:
 		VariablesOutliner->UpdateAssets();
+		VariablesOutliner->FullRefresh();
 		break;
 	default:
 		break;
