Fix add/remove function notify catching from RigVM

diff --git a/Engine/Plugins/Experimental/UAF/UAF/Source/UAFEditor/Private/Functions/Outliner/FunctionsOutlinerView.cpp b/Engine/Plugins/Experimental/UAF/UAF/Source/UAFEditor/Private/Functions/Outliner/FunctionsOutlinerView.cpp
--- a/Engine/Plugins/Experimental/UAF/UAF/Source/UAFEditor/Private/Functions/Outliner/FunctionsOutlinerView.cpp
+++ b/Engine/Plugins/Experimental/UAF/UAF/Source/UAFEditor/Private/Functions/Outliner/FunctionsOutlinerView.cpp
@@ -95,6 +95,15 @@ void SFunctionsView::OnRigVMModified(ERigVMGraphNotifType NotifType, URigVMGraph
 	case ERigVMGraphNotifType::FunctionVariantGuidChanged:
 		FunctionsOutliner->FullRefresh();
 		break;
+	case ERigVMGraphNotifType::NodeAdded:
+	case ERigVMGraphNotifType::NodeRemoved:
+		{
+			if (Cast<URigVMFunctionLibrary>(Graph))
+			{
+				FunctionsOutliner->FullRefresh();
+			}
+			break;
+		}
 	default:
 		break;
 	}
