Crash dragging bound variable into new asset for shared variable
#fix changed check to error log

diff --git a/Engine/Plugins/Experimental/UAF/UAF/Source/UAFUncookedOnly/Private/UncookedOnlyUtils.cpp b/Engine/Plugins/Experimental/UAF/UAF/Source/UAFUncookedOnly/Private/UncookedOnlyUtils.cpp
--- a/Engine/Plugins/Experimental/UAF/UAF/Source/UAFUncookedOnly/Private/UncookedOnlyUtils.cpp
+++ b/Engine/Plugins/Experimental/UAF/UAF/Source/UAFUncookedOnly/Private/UncookedOnlyUtils.cpp
@@ -1899,10 +1899,15 @@ void FUtils::ReplaceVariableReferences(UAnimNextRigVMAssetEditorData* InEditorDa
 								{
 									if (bDifferentVariableSource)
 									{
-										check(FindOuterAsset == ThisOuterAsset && ReplaceOuterAsset != ThisOuterAsset);
+										if (FindOuterAsset == ThisOuterAsset && ReplaceOuterAsset != ThisOuterAsset)
 										{
 											Controller->ReplaceVariableNodeWithSharedVariableNode(VariableNode, ReferenceToReplaceWith.GetName(), ReferenceToReplaceWith.GetObject(), bSetupUndoRedo, bPrintPythonCommands);
 										}
+										else
+										{
+											UE_LOGFMT(LogAnimation, Error, "Found matching instance of Variable Type/Name ({0}/{1}) combination on {2} which differs from {3} from which the replace operation is being executed, this instance will be skipped.",
+											FindReferenceType.ToString(), SoftReferenceToFind.GetName().ToString(), FindOuterAsset->GetPathName(), ThisOuterAsset->GetPathName());
+										}
 									}
 									else
 									{					
