Rename various AnimNext trace code to UAF

#rnx

diff --git a/Engine/Plugins/Experimental/UAF/UAF/Source/UAF/Internal/Module/AnimNextModuleInstance.h b/Engine/Plugins/Experimental/UAF/UAF/Source/UAF/Internal/Module/AnimNextModuleInstance.h
--- a/Engine/Plugins/Experimental/UAF/UAF/Source/UAF/Internal/Module/AnimNextModuleInstance.h
+++ b/Engine/Plugins/Experimental/UAF/UAF/Source/UAF/Internal/Module/AnimNextModuleInstance.h
@@ -147,7 +147,7 @@ private:
 	// Flip proxy variable buffers then copy any dirty values
 	void CopyProxyVariables();
 
-#if ANIMNEXT_TRACE_ENABLED
+#if UAF_TRACE_ENABLED
 	// trace this module if it hasn't already been traced this frame
 	UAF_API void Trace();
 #endif
@@ -270,7 +270,7 @@ private:
 #endif
 
 private:
-#if ANIMNEXT_TRACE_ENABLED
+#if UAF_TRACE_ENABLED
 	bool bTracedThisFrame = false;
 #endif
 };
@@ -285,8 +285,8 @@ struct TStructOpsTypeTraits<FAnimNextModuleInstance> : public TStructOpsTypeTrai
 };
 
 
-#if ANIMNEXT_TRACE_ENABLED
-#define TRACE_ANIMNEXT_MODULE(ModuleInstance) ModuleInstance.Trace();
+#if UAF_TRACE_ENABLED
+#define TRACE_UAF_MODULE(ModuleInstance) ModuleInstance.Trace();
 #else
-#define TRACE_ANIMNEXT_MODULE(ModuleInstance)
+#define TRACE_UAF_MODULE(ModuleInstance)
 #endif
diff --git a/Engine/Plugins/Experimental/UAF/UAF/Source/UAF/Internal/Variables/UAFInstanceVariableContainer.h b/Engine/Plugins/Experimental/UAF/UAF/Source/UAF/Internal/Variables/UAFInstanceVariableContainer.h
--- a/Engine/Plugins/Experimental/UAF/UAF/Source/UAF/Internal/Variables/UAFInstanceVariableContainer.h
+++ b/Engine/Plugins/Experimental/UAF/UAF/Source/UAF/Internal/Variables/UAFInstanceVariableContainer.h
@@ -21,7 +21,7 @@ namespace UE::UAF
 	struct FVariableOverridesCollection;
 	struct FVariableOverrides;
 	struct FStructDataCache;
-	struct FAnimNextTrace;
+	struct FUAFTrace;
 }
 
 // A 'shared' variable container, either a user-defined FInstancedPropertyBag or a natively-defined struct wrapped in a FInstancedStruct
@@ -47,7 +47,7 @@ private:
 	friend FUAFInstanceVariableDataProxy;
 	template<typename, ESPMode>
 	friend class SharedPointerInternals::TIntrusiveReferenceController;
-	friend UE::UAF::FAnimNextTrace;
+	friend UE::UAF::FUAFTrace;
 	friend FUAFRigVMComponent;
 
 	// Construct from struct data, copying the defaults. Optional overrides provided.
diff --git a/Engine/Plugins/Experimental/UAF/UAF/Source/UAF/Private/AnimNextModuleImpl.cpp b/Engine/Plugins/Experimental/UAF/UAF/Source/UAF/Private/AnimNextModuleImpl.cpp
--- a/Engine/Plugins/Experimental/UAF/UAF/Source/UAF/Private/AnimNextModuleImpl.cpp
+++ b/Engine/Plugins/Experimental/UAF/UAF/Source/UAF/Private/AnimNextModuleImpl.cpp
@@ -30,14 +30,14 @@
 #include "Factory/SystemFactory.h"
 #include "Features/IModularFeatures.h"
 #include "RewindDebuggerRuntimeInterface/IRewindDebuggerRuntimeExtension.h"
-#include "RewindDebugger/RewindDebuggerAnimNextRuntime.h"
+#include "RewindDebugger/RewindDebuggerUAFRuntime.h"
 
 #define LOCTEXT_NAMESPACE "AnimNextModule"
 
 namespace UE::UAF
 {
-#if ANIMNEXT_TRACE_ENABLED
-	FRewindDebuggerAnimNextRuntime GRewindDebuggerAnimNextRuntime;
+#if UAF_TRACE_ENABLED
+	FRewindDebuggerUAFRuntime GRewindDebuggerUAFRuntime;
 #endif
 	
 	void FAnimNextModuleImpl::StartupModule()
@@ -83,8 +83,8 @@ namespace UE::UAF
 
 		FSystemFactory::Init();
 
-#if ANIMNEXT_TRACE_ENABLED
-		IModularFeatures::Get().RegisterModularFeature(IRewindDebuggerRuntimeExtension::ModularFeatureName, &GRewindDebuggerAnimNextRuntime);
+#if UAF_TRACE_ENABLED
+		IModularFeatures::Get().RegisterModularFeature(IRewindDebuggerRuntimeExtension::ModularFeatureName, &GRewindDebuggerUAFRuntime);
 #endif
 
 		UE::UniversalObjectLocator::IUniversalObjectLocatorModule& UolModule = FModuleManager::Get().LoadModuleChecked<UE::UniversalObjectLocator::IUniversalObjectLocatorModule>("UniversalObjectLocator");
@@ -122,8 +122,8 @@ namespace UE::UAF
 
 	void FAnimNextModuleImpl::ShutdownModule()
 	{
-#if ANIMNEXT_TRACE_ENABLED
-		IModularFeatures::Get().UnregisterModularFeature(IRewindDebuggerRuntimeExtension::ModularFeatureName, &GRewindDebuggerAnimNextRuntime);
+#if UAF_TRACE_ENABLED
+		IModularFeatures::Get().UnregisterModularFeature(IRewindDebuggerRuntimeExtension::ModularFeatureName, &GRewindDebuggerUAFRuntime);
 #endif
 		
 		FRigVMRuntimeDataRegistry::Destroy();
diff --git a/Engine/Plugins/Experimental/UAF/UAF/Source/UAF/Private/Module/AnimNextModuleInstance.cpp b/Engine/Plugins/Experimental/UAF/UAF/Source/UAF/Private/Module/AnimNextModuleInstance.cpp
--- a/Engine/Plugins/Experimental/UAF/UAF/Source/UAF/Private/Module/AnimNextModuleInstance.cpp
+++ b/Engine/Plugins/Experimental/UAF/UAF/Source/UAF/Private/Module/AnimNextModuleInstance.cpp
@@ -19,7 +19,7 @@
 #include "Module/RigUnit_AnimNextModuleEvents.h"
 #include "Module/ModuleTickFunction.h"
 #include "UObject/UObjectIterator.h"
-#include "RewindDebugger/AnimNextTrace.h"
+#include "RewindDebugger/UAFTrace.h"
 #include "Script/UAFScriptComponent.h"
 #include "Variables/AnimNextVariableReference.h"
 
@@ -420,12 +420,12 @@ void FAnimNextModuleInstance::CopyProxyVariables()
 	ProxyVariablesRead.CopyDirty();
 }
 
-#if ANIMNEXT_TRACE_ENABLED
+#if UAF_TRACE_ENABLED
 void FAnimNextModuleInstance::Trace()
 {
 	if (!bTracedThisFrame)
 	{
-		TRACE_ANIMNEXT_VARIABLES(this, Object);
+		TRACE_UAF_VARIABLES(this, Object);
 		bTracedThisFrame = true;
 	}
 }
diff --git a/Engine/Plugins/Experimental/UAF/UAF/Source/UAF/Private/Module/ModuleTickFunction.cpp b/Engine/Plugins/Experimental/UAF/UAF/Source/UAF/Private/Module/ModuleTickFunction.cpp
--- a/Engine/Plugins/Experimental/UAF/UAF/Source/UAF/Private/Module/ModuleTickFunction.cpp
+++ b/Engine/Plugins/Experimental/UAF/UAF/Source/UAF/Private/Module/ModuleTickFunction.cpp
@@ -197,7 +197,7 @@ void FModuleEventTickFunction::Run(float InDeltaTime)
 
 	if (bFirstUserEvent)
 	{
-#if ANIMNEXT_TRACE_ENABLED
+#if UAF_TRACE_ENABLED
 		ModuleInstance->bTracedThisFrame = false;
 #endif
 
diff --git a/Engine/Plugins/Experimental/UAF/UAF/Source/UAF/Private/RewindDebugger/AnimNextTrace.cpp b/Engine/Plugins/Experimental/UAF/UAF/Source/UAF/Private/RewindDebugger/AnimNextTrace.cpp
deleted file mode 100644
diff --git a/Engine/Plugins/Experimental/UAF/UAF/Source/UAF/Private/RewindDebugger/RewindDebuggerAnimNextRuntime.cpp b/Engine/Plugins/Experimental/UAF/UAF/Source/UAF/Private/RewindDebugger/RewindDebuggerAnimNextRuntime.cpp
deleted file mode 100644
diff --git a/Engine/Plugins/Experimental/UAF/UAF/Source/UAF/Private/RewindDebugger/RewindDebuggerAnimNextRuntime.h b/Engine/Plugins/Experimental/UAF/UAF/Source/UAF/Private/RewindDebugger/RewindDebuggerAnimNextRuntime.h
deleted file mode 100644
diff --git a/Engine/Plugins/Experimental/UAF/UAF/Source/UAF/Private/RewindDebugger/RewindDebuggerUAFRuntime.cpp b/Engine/Plugins/Experimental/UAF/UAF/Source/UAF/Private/RewindDebugger/RewindDebuggerUAFRuntime.cpp
new file mode 100644
--- /dev/null
+++ b/Engine/Plugins/Experimental/UAF/UAF/Source/UAF/Private/RewindDebugger/RewindDebuggerUAFRuntime.cpp
@@ -0,0 +1,28 @@
+// Copyright Epic Games, Inc. All Rights Reserved.
+
+#include "RewindDebuggerUAFRuntime.h"
+
+#if UAF_TRACE_ENABLED
+
+namespace UE::UAF
+{
+	
+void FRewindDebuggerUAFRuntime::RecordingStarted()
+{
+	UE::Trace::ToggleChannel(TEXT("UAF"), true);
+}
+
+void FRewindDebuggerUAFRuntime::RecordingStopped()
+{
+	UE::Trace::ToggleChannel(TEXT("UAF"), false);
+}
+
+void FRewindDebuggerUAFRuntime::Clear()
+{
+	FUAFTrace::Reset();
+}
+
+}
+
+#endif
+	
diff --git a/Engine/Plugins/Experimental/UAF/UAF/Source/UAF/Private/RewindDebugger/RewindDebuggerUAFRuntime.h b/Engine/Plugins/Experimental/UAF/UAF/Source/UAF/Private/RewindDebugger/RewindDebuggerUAFRuntime.h
new file mode 100644
--- /dev/null
+++ b/Engine/Plugins/Experimental/UAF/UAF/Source/UAF/Private/RewindDebugger/RewindDebuggerUAFRuntime.h
@@ -0,0 +1,25 @@
+// Copyright Epic Games, Inc. All Rights Reserved.
+#pragma once
+
+#include "CoreMinimal.h"
+#include "RewindDebuggerRuntimeInterface/IRewindDebuggerRuntimeExtension.h"
+#include "RewindDebugger/UAFTrace.h"
+
+#if UAF_TRACE_ENABLED
+// Rewind debugger extension for Chooser support
+
+namespace UE::UAF
+{
+
+class FRewindDebuggerUAFRuntime : public IRewindDebuggerRuntimeExtension
+{
+public:
+	virtual void RecordingStarted() override;
+	virtual void RecordingStopped() override;
+	
+	virtual void Clear() override;
+};
+
+}
+
+#endif
\ No newline at end of file
diff --git a/Engine/Plugins/Experimental/UAF/UAF/Source/UAF/Private/RewindDebugger/UAFTrace.cpp b/Engine/Plugins/Experimental/UAF/UAF/Source/UAF/Private/RewindDebugger/UAFTrace.cpp
new file mode 100644
--- /dev/null
+++ b/Engine/Plugins/Experimental/UAF/UAF/Source/UAF/Private/RewindDebugger/UAFTrace.cpp
@@ -0,0 +1,320 @@
+// Copyright Epic Games, Inc. All Rights Reserved.
+
+#include "RewindDebugger/UAFTrace.h"
+
+#include "ObjectTrace.h"
+#include "Chaos/ParticleDirtyFlags.h"
+#include "UAFAssetInstance.h"
+#include "Misc/HashBuilder.h"
+#include "Module/AnimNextModuleInstance.h"
+#include "Serialization/MemoryWriter.h"
+#include "ObjectAsTraceIdProxyArchive.h"
+#include "Serialization/ObjectWriter.h"
+#include "Variables/VariableOverrides.h"
+#include "Variables/StructDataCache.h"
+
+#if UAF_TRACE_ENABLED
+
+
+UE_TRACE_EVENT_BEGIN(UAF, Instance)
+	UE_TRACE_EVENT_FIELD(uint64, InstanceId)
+	UE_TRACE_EVENT_FIELD(uint64, HostInstanceId)
+	UE_TRACE_EVENT_FIELD(uint64, OuterObjectId)
+	UE_TRACE_EVENT_FIELD(uint64, AssetId)
+UE_TRACE_EVENT_END()
+
+UE_TRACE_EVENT_BEGIN(UAF, InstanceVariables)
+	UE_TRACE_EVENT_FIELD(uint64, Cycle)
+	UE_TRACE_EVENT_FIELD(double, RecordingTime)
+	UE_TRACE_EVENT_FIELD(uint64, InstanceId)
+	UE_TRACE_EVENT_FIELD(uint32, VariableDescriptionHash)
+	UE_TRACE_EVENT_FIELD(uint8[], VariableData)
+UE_TRACE_EVENT_END()
+
+UE_TRACE_EVENT_BEGIN(UAF, InstanceVariablesStruct)
+	UE_TRACE_EVENT_FIELD(uint64, Cycle)
+	UE_TRACE_EVENT_FIELD(double, RecordingTime)
+	UE_TRACE_EVENT_FIELD(uint64, InstanceId)
+	UE_TRACE_EVENT_FIELD(uint8[], VariableData)
+UE_TRACE_EVENT_END()
+
+UE_TRACE_EVENT_BEGIN(UAF, InstanceVariableDescriptions)
+	UE_TRACE_EVENT_FIELD(uint32, VariableDescriptionHash)
+	UE_TRACE_EVENT_FIELD(uint8[], VariableDescriptionData)
+UE_TRACE_EVENT_END()
+
+UE_TRACE_CHANNEL_DEFINE(UAFChannel);
+
+namespace UE::UAF
+{
+
+FRWLock GTracedInstancesRWLock;
+TSet<uint64> GTracedInstances;
+
+FRWLock GTracedPropertiesRWLock;
+TMap<uint32, const UPropertyBag*> GTracedPropertyDescs;
+
+const FGuid FUAFTrace::CustomVersionGUID = FGuid(0x83E9BE13, 0xB1C845DC, 0x86C4D3E5, 0xE66CBE91);
+
+namespace FUAFTraceCustomVersion
+{
+enum CustomVersionIndex
+{
+	// Before any version changes were made in the plugin
+	FirstVersion = 0,
+
+	// -----<new versions can be added above this line>-------------------------------------------------
+	VersionPlusOne,
+	LatestVersion = VersionPlusOne - 1
+};
+}
+
+FCustomVersionRegistration GPropertyBagCustomVersion(FUAFTrace::CustomVersionGUID, FUAFTraceCustomVersion::FirstVersion, TEXT("FUAFTraceCustomVersion"));
+
+void FUAFTrace::Reset()
+{
+	GTracedInstances.Empty();
+	GTracedPropertyDescs.Empty();
+}
+
+void FUAFTrace::OutputUAFInstance(const FUAFAssetInstance* AnimNextInstance, const UObject* OuterObject)
+{
+	bool bChannelEnabled = UE_TRACE_CHANNELEXPR_IS_ENABLED(UAFChannel);
+	if (!bChannelEnabled || AnimNextInstance == nullptr)
+	{
+		return;
+	}
+
+	if (CANNOT_TRACE_OBJECT(OuterObject))
+	{
+		return;
+	}
+	
+	uint64 InstanceId = AnimNextInstance->GetUniqueId();
+
+	bool bTrace = false;
+
+	{
+		FRWScopeLock RWScopeLock(GTracedInstancesRWLock, SLT_ReadOnly);
+		
+		if (!GTracedInstances.Contains(InstanceId))
+		{
+			RWScopeLock.ReleaseReadOnlyLockAndAcquireWriteLock_USE_WITH_CAUTION();
+
+			// Note that we might double add here due to 2x people writing at the same time, but it doesn't matter
+			GTracedInstances.Add(InstanceId);
+			bTrace = true;
+		}
+	}
+
+	if (bTrace)
+	{
+		if (AnimNextInstance->GetHost())
+		{
+			OutputUAFInstance(AnimNextInstance->GetHost(), OuterObject);
+		}
+
+		const FUAFAssetInstance* HostInstance = AnimNextInstance->GetHost();
+		uint64 HostInstanceId = HostInstance ? HostInstance->GetUniqueId() : FObjectTrace::GetObjectId(OuterObject);
+
+		TRACE_OBJECT(OuterObject);
+		TRACE_OBJECT(AnimNextInstance->GetAsset<UAnimNextRigVMAsset>());
+		TRACE_INSTANCE(OuterObject, AnimNextInstance->GetUniqueId(), HostInstanceId, FUAFAssetInstance::StaticStruct(), AnimNextInstance->GetDebugName().ToString());
+
+		UE_TRACE_LOG(UAF, Instance, UAFChannel)
+			<< Instance.InstanceId(InstanceId)
+			<< Instance.OuterObjectId(FObjectTrace::GetObjectId(OuterObject))
+			<< Instance.HostInstanceId(HostInstanceId)
+			<< Instance.AssetId(FObjectTrace::GetObjectId(AnimNextInstance->GetAsset<UAnimNextRigVMAsset>()));
+	}
+}
+
+uint32 GetPropertyDescHash(const TConstArrayView<FPropertyBagPropertyDesc>& Descs)
+{
+	FHashBuilder HashBuilder;
+	for (const FPropertyBagPropertyDesc& Desc : Descs)
+	{
+		HashBuilder.Append(Desc.Name);
+		HashBuilder.Append(Desc.ContainerTypes);
+		HashBuilder.Append(Desc.ID);
+		HashBuilder.Append(Desc.PropertyFlags);
+		HashBuilder.Append(Desc.ValueType);
+
+		// Q: should we care about metadata?
+	}
+	return HashBuilder.GetHash();
+}
+
+void FUAFTrace::OutputUAFVariables(const FUAFAssetInstance* Instance, const UObject* OuterObject)
+{
+	bool bChannelEnabled = UE_TRACE_CHANNELEXPR_IS_ENABLED(UAFChannel);
+	if (!bChannelEnabled || Instance == nullptr || OuterObject == nullptr)
+	{
+		return;
+	}
+
+	uint64 InstanceId = Instance->GetUniqueId();
+
+	if (CANNOT_TRACE_OBJECT(OuterObject))
+	{
+		return;
+	}
+
+	OutputUAFInstance(Instance, OuterObject);
+
+	bool bVariablesOutput = false;
+	if (Instance->Variables.AllSharedVariableContainers.Num() > 0)
+	{
+		for (const TWeakPtr<FUAFInstanceVariableContainer>& OwnedVariableSet : Instance->Variables.AllSharedVariableContainers)
+		{
+			if (OutputUAFVariableSet(OwnedVariableSet.Pin().ToSharedRef(), InstanceId, OuterObject))
+			{
+				bVariablesOutput = true;
+			}
+		}
+	}
+
+	if (!bVariablesOutput)
+	{
+		// trace an empty InstanceVariables message, as it's used for lifetime tracking currently
+		UE_TRACE_LOG(UAF, InstanceVariables, UAFChannel)
+				<< InstanceVariables.Cycle(FPlatformTime::Cycles64())
+				<< InstanceVariables.RecordingTime(FObjectTrace::GetWorldElapsedTime(OuterObject->GetWorld()))
+				<< InstanceVariables.InstanceId(InstanceId);
+	}
+}
+
+bool FUAFTrace::OutputUAFVariableSet(const TSharedRef<FUAFInstanceVariableContainer>& InVariableSet, uint64 InstanceId, const UObject* OuterObject)
+{
+	bool bOutputAnything = false;
+	switch (InVariableSet->VariablesContainer.GetIndex())
+	{
+	case FUAFInstanceVariableContainer::FVariablesContainerType::IndexOfType<FInstancedPropertyBag>():
+		{
+			// Copy property bag data as it will be modified
+			FInstancedPropertyBag InstancedPropertyBag = InVariableSet->VariablesContainer.Get<FInstancedPropertyBag>();
+			if (InstancedPropertyBag.GetNumPropertiesInBag() == 0)
+			{
+				break;
+			}
+
+			TConstArrayView<FPropertyBagPropertyDesc> PropertyDescs = InstancedPropertyBag.GetPropertyBagStruct()->GetPropertyDescs();
+			uint32 PropertyDescHash = GetPropertyDescHash(PropertyDescs);
+			{
+				FRWScopeLock RWScopeLock(GTracedPropertiesRWLock, SLT_ReadOnly);
+				TArray<uint8> ArchiveData;
+				ArchiveData.Reserve(1024 * 10);
+				
+				if (!GTracedPropertyDescs.Contains(PropertyDescHash))
+				{
+					// Upgrade to write lock
+					RWScopeLock.ReleaseReadOnlyLockAndAcquireWriteLock_USE_WITH_CAUTION();
+
+					// Check in case someone just serialized before we aquired the write lock
+					if (!GTracedPropertyDescs.Contains(PropertyDescHash))
+					{
+						GTracedPropertyDescs.Add(PropertyDescHash, InstancedPropertyBag.GetPropertyBagStruct());
+						
+						// Serialize prop descs
+						TArray<FPropertyBagPropertyDesc, TInlineAllocator<64>> PropertyDescriptions;
+						PropertyDescriptions.Reserve(PropertyDescs.Num());
+						PropertyDescriptions = PropertyDescs;	// need to copy since we can't serialize an array view
+
+						FMemoryWriter WriterArchive(ArchiveData);
+						FObjectAsTraceIdProxyArchive Archive(WriterArchive);
+						Archive.UsingCustomVersion(CustomVersionGUID);
+						Archive.UsingCustomVersion(FPropertyBagCustomVersion::GUID);
+						Archive << PropertyDescriptions;
+
+						UE_TRACE_LOG(UAF, InstanceVariableDescriptions, UAFChannel)
+							<< InstanceVariableDescriptions.VariableDescriptionHash(PropertyDescHash)
+							<< InstanceVariableDescriptions.VariableDescriptionData(ArchiveData.GetData(), ArchiveData.Num());
+					}
+				}
+
+				ArchiveData.SetNum(0, EAllowShrinking::No);
+				FMemoryWriter WriterArchive(ArchiveData);
+				FObjectAsTraceIdProxyArchive Archive(WriterArchive);
+				UPropertyBag* PropertyBag = const_cast<UPropertyBag*>(InstancedPropertyBag.GetPropertyBagStruct());
+				if (PropertyBag != nullptr)
+				{
+					// Apply overrides
+					if (!InVariableSet->ResolvedOverrides.IsEmpty())
+					{
+						const int32 VariablesNum = PropertyDescs.Num();
+						for (int32 VariableIndex = 0; VariableIndex < VariablesNum; ++VariableIndex)
+						{
+							if (InVariableSet->ResolvedOverrides[VariableIndex] != nullptr)
+							{
+								const FPropertyBagPropertyDesc& Desc = PropertyDescs[VariableIndex];
+								uint8* Target = Desc.CachedProperty->ContainerPtrToValuePtr<uint8>(InstancedPropertyBag.GetMutableValue().GetMemory());
+								uint8* Source = InVariableSet->ResolvedOverrides[VariableIndex];
+								
+								Desc.CachedProperty->CopyCompleteValue(Target, Source);
+							}
+						}
+					}
+				
+					PropertyBag->SerializeItem(Archive, InstancedPropertyBag.GetMutableValue().GetMemory(), nullptr);
+				}
+
+				bOutputAnything = true;
+				UE_TRACE_LOG(UAF, InstanceVariables, UAFChannel)
+						<< InstanceVariables.Cycle(FPlatformTime::Cycles64())
+						<< InstanceVariables.RecordingTime(FObjectTrace::GetWorldElapsedTime(OuterObject->GetWorld()))
+						<< InstanceVariables.InstanceId(InstanceId)
+						<< InstanceVariables.VariableDescriptionHash(PropertyDescHash)
+						<< InstanceVariables.VariableData(ArchiveData.GetData(), ArchiveData.Num());
+			}
+			break;
+		}
+		case FUAFInstanceVariableContainer::FVariablesContainerType::IndexOfType<FInstancedStruct>():
+		{
+			TArray<uint8> ArchiveData;
+			ArchiveData.Reserve(1024 * 10);
+
+			FMemoryWriter WriterArchive(ArchiveData);
+			FObjectAsTraceIdProxyArchive Archive(WriterArchive);
+
+			// Copy property bag data as it will be modified
+			FInstancedStruct InstancedStruct = InVariableSet->VariablesContainer.Get<FInstancedStruct>();
+
+			// Apply overrides
+			if (!InVariableSet->ResolvedOverrides.IsEmpty())
+			{
+				for (int32 VariableIndex = 0; VariableIndex < InVariableSet->NumVariables; ++VariableIndex)
+				{
+					const FStructDataCache::FPropertyInfo& PropertyInfo = InVariableSet->StructDataCache->GetProperties()[VariableIndex];
+					const FProperty* Property = PropertyInfo.Property;
+					
+					uint8* Target = Property->ContainerPtrToValuePtr<uint8>(InstancedStruct.GetMutableMemory());
+                    uint8* Source = InVariableSet->ResolvedOverrides[VariableIndex];
+					
+					Property->CopyCompleteValue(Target, Source);
+				}
+			}
+			
+			if (InstancedStruct.IsValid())
+			{
+				FInstancedStruct::StaticStruct()->SerializeItem(Archive, &InstancedStruct, nullptr);
+			}
+
+			bOutputAnything = true;
+			UE_TRACE_LOG(UAF, InstanceVariablesStruct, UAFChannel)
+					<< InstanceVariablesStruct.Cycle(FPlatformTime::Cycles64())
+					<< InstanceVariablesStruct.RecordingTime(FObjectTrace::GetWorldElapsedTime(OuterObject->GetWorld()))
+					<< InstanceVariablesStruct.InstanceId(InstanceId)
+					<< InstanceVariablesStruct.VariableData(ArchiveData.GetData(), ArchiveData.Num());
+			break;
+		}
+	default:
+		checkNoEntry();
+		break;
+	}
+
+	return bOutputAnything;
+}
+
+}
+
+#endif
diff --git a/Engine/Plugins/Experimental/UAF/UAF/Source/UAF/Private/UAFAssetInstance.cpp b/Engine/Plugins/Experimental/UAF/UAF/Source/UAF/Private/UAFAssetInstance.cpp
--- a/Engine/Plugins/Experimental/UAF/UAF/Source/UAF/Private/UAFAssetInstance.cpp
+++ b/Engine/Plugins/Experimental/UAF/UAF/Source/UAF/Private/UAFAssetInstance.cpp
@@ -52,7 +52,7 @@ FUAFAssetInstance::FUAFAssetInstance(UScriptStruct* InScriptStruct)
 	: ScriptStruct(InScriptStruct)
 {
 	check(ScriptStruct != nullptr);
-#if ANIMNEXT_TRACE_ENABLED
+#if UAF_TRACE_ENABLED
 	UniqueId = FObjectTrace::AllocateInstanceId();
 #endif
 }
diff --git a/Engine/Plugins/Experimental/UAF/UAF/Source/UAF/Public/RewindDebugger/AnimNextTrace.h b/Engine/Plugins/Experimental/UAF/UAF/Source/UAF/Public/RewindDebugger/AnimNextTrace.h
deleted file mode 100644
diff --git a/Engine/Plugins/Experimental/UAF/UAF/Source/UAF/Public/RewindDebugger/UAFTrace.h b/Engine/Plugins/Experimental/UAF/UAF/Source/UAF/Public/RewindDebugger/UAFTrace.h
new file mode 100644
--- /dev/null
+++ b/Engine/Plugins/Experimental/UAF/UAF/Source/UAF/Public/RewindDebugger/UAFTrace.h
@@ -0,0 +1,52 @@
+// Copyright Epic Games, Inc. All Rights Reserved.
+
+#pragma once
+
+#if UE_ENABLE_INCLUDE_ORDER_DEPRECATED_IN_5_7
+#include "CoreMinimal.h"
+#include "ObjectTrace.h"
+#include "StructUtils/PropertyBag.h"
+#endif // UE_ENABLE_INCLUDE_ORDER_DEPRECATED_IN_5_7
+
+#include "CoreTypes.h"
+#include "Misc/Guid.h"
+#include "ObjectTraceDefines.h"
+#include "Trace/Trace.h"
+#include "Templates/SharedPointer.h"
+
+#define UAF_TRACE_ENABLED (OBJECT_TRACE_ENABLED && !(UE_BUILD_SHIPPING || UE_BUILD_TEST))
+
+#if UAF_TRACE_ENABLED
+
+UE_TRACE_CHANNEL_EXTERN(UAFChannel, UAF_API);
+
+class UObject;
+struct FAnimNextModuleInstance;
+struct FUAFAssetInstance;
+struct FAnimNextGraphInstance;
+struct FUAFInstanceVariableContainer;
+
+namespace UE::UAF
+{
+struct FUAFTrace
+{
+	UAF_API static const FGuid CustomVersionGUID;
+	
+	UAF_API static void Reset();
+	
+	UAF_API static void OutputUAFInstance(const FUAFAssetInstance* DataInterface, const UObject* OuterObject);
+	UAF_API static void OutputUAFVariables(const FUAFAssetInstance* DataInterface, const UObject* OuterObject);
+
+private:
+	UAF_API static bool OutputUAFVariableSet(const TSharedRef<FUAFInstanceVariableContainer>& VariableSet, uint64 InstanceId, const UObject* OuterObject);
+};
+
+}
+
+
+#define TRACE_UAF_INSTANCE(DataInterface, OuterObject) UE::UAF::FUAFTrace::OutputUAFInstance(DataInterface, OuterObject);
+#define TRACE_UAF_VARIABLES(DataInterface, OuterObject) UE::UAF::FUAFTrace::OutputUAFVariables(DataInterface, OuterObject);
+#else
+#define TRACE_UAF_INSTANCE(DataInterface, OuterObject)
+#define TRACE_UAF_VARIABLES(DataInterface, OuterObject)
+#endif
diff --git a/Engine/Plugins/Experimental/UAF/UAF/Source/UAF/Public/UAFAssetInstance.h b/Engine/Plugins/Experimental/UAF/UAF/Source/UAF/Public/UAFAssetInstance.h
--- a/Engine/Plugins/Experimental/UAF/UAF/Source/UAF/Public/UAFAssetInstance.h
+++ b/Engine/Plugins/Experimental/UAF/UAF/Source/UAF/Public/UAFAssetInstance.h
@@ -7,8 +7,9 @@
 #include "Script/UAFScriptComponent.h"
 #include "Param/ParamType.h"
 #include "StructUtils/PropertyBag.h"
-#include "RewindDebugger/AnimNextTrace.h"
+#include "RewindDebugger/UAFTrace.h"
 #include "Variables/UAFInstanceVariableData.h"
+
 #include "UAFAssetInstance.generated.h"
 
 #define UE_API UAF_API
@@ -24,7 +25,7 @@ namespace UE::UAF
 {
 	struct FPoolHandle;
 	struct FInjectionInfo;
-	struct FAnimNextTrace;
+	struct FUAFTrace;
 	struct FInstanceTaskContext;
 	struct FVariableOverridesCollection;
 }
@@ -214,7 +215,7 @@ struct FUAFAssetInstance
 
 	uint64 GetUniqueId() const
 	{
-#if ANIMNEXT_TRACE_ENABLED
+#if UAF_TRACE_ENABLED
 		return UniqueId;
 #else
 		return 0;
@@ -322,20 +323,20 @@ protected:
 	friend FAnimNextModuleInjectionComponent;
 	friend FAnimNextModuleContextData;
 	friend FUAFInstanceVariableData;
-	friend UE::UAF::FAnimNextTrace;
+	friend UE::UAF::FUAFTrace;
 	friend UE::UAF::FInstanceTaskContext;
 	friend FUAFRigVMComponent;
 
 	void SetDebugName(FName Name)
 	{
-	#if ANIMNEXT_TRACE_ENABLED
+	#if UAF_TRACE_ENABLED
 		DebugName = Name;
 	#endif
 	}
 	
 	FName GetDebugName() const
 	{
-#if ANIMNEXT_TRACE_ENABLED
+#if UAF_TRACE_ENABLED
 		if (!DebugName.IsNone())
 		{
 			return DebugName;
@@ -344,7 +345,7 @@ protected:
 		return GetAssetName();
 	}
 
-#if ANIMNEXT_TRACE_ENABLED
+#if UAF_TRACE_ENABLED
 	FName DebugName;
 	uint64 UniqueId;
 	UE_API volatile static int64 NextUniqueId;
diff --git a/Engine/Plugins/Experimental/UAF/UAF/Source/UAF/Public/Variables/UAFInstanceVariableData.h b/Engine/Plugins/Experimental/UAF/UAF/Source/UAF/Public/Variables/UAFInstanceVariableData.h
--- a/Engine/Plugins/Experimental/UAF/UAF/Source/UAF/Public/Variables/UAFInstanceVariableData.h
+++ b/Engine/Plugins/Experimental/UAF/UAF/Source/UAF/Public/Variables/UAFInstanceVariableData.h
@@ -25,7 +25,7 @@ namespace UE::UAF
 {
 	struct FVariableOverrides;
 	struct FVariableOverridesCollection;
-	struct FAnimNextTrace;
+	struct FUAFTrace;
 	struct FInstanceTaskContext;
 }
 
@@ -45,7 +45,7 @@ private:
 	friend FAnimNextModuleInstance;
 	friend FAnimNextGraphInstance;
 	friend FUAFInstanceVariableDataProxy;
-	friend UE::UAF::FAnimNextTrace;
+	friend UE::UAF::FUAFTrace;
 	friend UE::UAF::FInstanceTaskContext;
 	friend FUAFRigVMComponent;
 
diff --git a/Engine/Plugins/Experimental/UAF/UAF/Source/UAFEditor/Private/AnimNextEditorModule.cpp b/Engine/Plugins/Experimental/UAF/UAF/Source/UAFEditor/Private/AnimNextEditorModule.cpp
--- a/Engine/Plugins/Experimental/UAF/UAF/Source/UAFEditor/Private/AnimNextEditorModule.cpp
+++ b/Engine/Plugins/Experimental/UAF/UAF/Source/UAFEditor/Private/AnimNextEditorModule.cpp
@@ -85,14 +85,14 @@
 
 #define LOCTEXT_NAMESPACE "AnimNextEditorModule"
 
-#if ANIMNEXT_TRACE_ENABLED
+#if UAF_TRACE_ENABLED
 namespace
 {
-	FRewindDebuggerAnimNext GRewindDebuggerAnimNext;
-#if ANIMNEXT_TRACE_ENABLED
+	FRewindDebuggerUAF GRewindDebuggerUAF;
+#if UAF_TRACE_ENABLED
 	UE::UAF::Editor::FAnimNextModuleTrackCreator GAnimNextModulesTrackCreator;
 #endif
-	FAnimNextTraceModule GAnimNextTraceModule;
+	FUAFTraceModule GUAFTraceModule;
 }
 #endif
 
@@ -153,13 +153,13 @@ void FAnimNextEditorModule::StartupModule()
 			UAF::Editor::FAssetWizard::Launch();
 		})});
 
-#if ANIMNEXT_TRACE_ENABLED
-	IModularFeatures::Get().RegisterModularFeature(IRewindDebuggerExtension::ModularFeatureName, &GRewindDebuggerAnimNext);
-#if ANIMNEXT_TRACE_ENABLED
+#if UAF_TRACE_ENABLED
+	IModularFeatures::Get().RegisterModularFeature(IRewindDebuggerExtension::ModularFeatureName, &GRewindDebuggerUAF);
+#if UAF_TRACE_ENABLED
 	IModularFeatures::Get().RegisterModularFeature(RewindDebugger::IRewindDebuggerTrackCreator::ModularFeatureName, &GAnimNextModulesTrackCreator);
 #endif
-	IModularFeatures::Get().RegisterModularFeature(TraceServices::ModuleFeatureName, &GAnimNextTraceModule);
-#endif // ANIMNEXT_TRACE_ENABLED
+	IModularFeatures::Get().RegisterModularFeature(TraceServices::ModuleFeatureName, &GUAFTraceModule);
+#endif // UAF_TRACE_ENABLED
 
 	FPropertyEditorModule& PropertyModule = FModuleManager::LoadModuleChecked<FPropertyEditorModule>("PropertyEditor");
 
@@ -608,13 +608,11 @@ void FAnimNextEditorModule::ShutdownModule()
 	}
 	UnregisterAssetCompilationHandler(FTopLevelAssetPath(TEXT("/Script/UAF.AnimNextModule")));
 
-#if ANIMNEXT_TRACE_ENABLED
-	IModularFeatures::Get().UnregisterModularFeature(IRewindDebuggerExtension::ModularFeatureName, &GRewindDebuggerAnimNext);
-#if ANIMNEXT_TRACE_ENABLED
+#if UAF_TRACE_ENABLED
+	IModularFeatures::Get().UnregisterModularFeature(IRewindDebuggerExtension::ModularFeatureName, &GRewindDebuggerUAF);
 	IModularFeatures::Get().UnregisterModularFeature(RewindDebugger::IRewindDebuggerTrackCreator::ModularFeatureName, &GAnimNextModulesTrackCreator);
-#endif
-	IModularFeatures::Get().UnregisterModularFeature(TraceServices::ModuleFeatureName, &GAnimNextTraceModule);
-#endif // ANIMNEXT_TRACE_ENABLED
+	IModularFeatures::Get().UnregisterModularFeature(TraceServices::ModuleFeatureName, &GUAFTraceModule);
+#endif // UAF_TRACE_ENABLED
 
 	FEdGraphUtilities::UnregisterVisualPinFactory(GraphPanelPinFactory);
 
diff --git a/Engine/Plugins/Experimental/UAF/UAF/Source/UAFEditor/Private/RewindDebugger/AnimNextAnalyzer.cpp b/Engine/Plugins/Experimental/UAF/UAF/Source/UAFEditor/Private/RewindDebugger/AnimNextAnalyzer.cpp
--- a/Engine/Plugins/Experimental/UAF/UAF/Source/UAFEditor/Private/RewindDebugger/AnimNextAnalyzer.cpp
+++ b/Engine/Plugins/Experimental/UAF/UAF/Source/UAFEditor/Private/RewindDebugger/AnimNextAnalyzer.cpp
@@ -2,7 +2,7 @@
 
 #include "AnimNextAnalyzer.h"
 
-#if ANIMNEXT_TRACE_ENABLED
+#if UAF_TRACE_ENABLED
 
 #include "HAL/LowLevelMemTracker.h"
 #include "Serialization/MemoryReader.h"
@@ -12,13 +12,13 @@
 #include "StructUtils/PropertyBag.h"
 #include "Serialization/ObjectReader.h"
 
-FAnimNextAnalyzer::FAnimNextAnalyzer(TraceServices::IAnalysisSession& InSession, FAnimNextProvider& InProvider)
+FUAFAnalyzer::FUAFAnalyzer(TraceServices::IAnalysisSession& InSession, FUAFProvider& InProvider)
 	: Session(InSession)
 	, Provider(InProvider)
 {
 }
 
-void FAnimNextAnalyzer::OnAnalysisBegin(const FOnAnalysisContext& Context)
+void FUAFAnalyzer::OnAnalysisBegin(const FOnAnalysisContext& Context)
 {
 	auto& Builder = Context.InterfaceBuilder;
 
@@ -28,9 +28,9 @@ void FAnimNextAnalyzer::OnAnalysisBegin(const FOnAnalysisContext& Context)
 	Builder.RouteEvent(RouteId_InstanceVariableDescriptions, "UAF", "InstanceVariableDescriptions");
 }
 
-bool FAnimNextAnalyzer::OnEvent(uint16 RouteId, EStyle Style, const FOnEventContext& Context)
+bool FUAFAnalyzer::OnEvent(uint16 RouteId, EStyle Style, const FOnEventContext& Context)
 {
-	LLM_SCOPE_BYNAME(TEXT("Insights/FAnimNextAnalyzer"));
+	LLM_SCOPE_BYNAME(TEXT("Insights/FUAFAnalyzer"));
 
 	TraceServices::FAnalysisSessionEditScope _(Session);
 
@@ -79,4 +79,4 @@ bool FAnimNextAnalyzer::OnEvent(uint16 RouteId, EStyle Style, const FOnEventCont
 
 	return true;
 }
-#endif // ANIMNEXT_TRACE_ENABLED
+#endif // UAF_TRACE_ENABLED
diff --git a/Engine/Plugins/Experimental/UAF/UAF/Source/UAFEditor/Private/RewindDebugger/AnimNextAnalyzer.h b/Engine/Plugins/Experimental/UAF/UAF/Source/UAFEditor/Private/RewindDebugger/AnimNextAnalyzer.h
--- a/Engine/Plugins/Experimental/UAF/UAF/Source/UAFEditor/Private/RewindDebugger/AnimNextAnalyzer.h
+++ b/Engine/Plugins/Experimental/UAF/UAF/Source/UAFEditor/Private/RewindDebugger/AnimNextAnalyzer.h
@@ -2,18 +2,18 @@
 
 #pragma once
 
-#include "RewindDebugger/AnimNextTrace.h"
+#include "RewindDebugger/UAFTrace.h"
 
-#if ANIMNEXT_TRACE_ENABLED
+#if UAF_TRACE_ENABLED
 #include "Trace/Analyzer.h"
 
-class FAnimNextProvider;
+class FUAFProvider;
 namespace TraceServices { class IAnalysisSession; }
 
-class FAnimNextAnalyzer : public UE::Trace::IAnalyzer
+class FUAFAnalyzer : public UE::Trace::IAnalyzer
 {
 public:
-	FAnimNextAnalyzer(TraceServices::IAnalysisSession& InSession, FAnimNextProvider& InProvider);
+	FUAFAnalyzer(TraceServices::IAnalysisSession& InSession, FUAFProvider& InProvider);
 
 	virtual void OnAnalysisBegin(const FOnAnalysisContext& Context) override;
 	virtual void OnAnalysisEnd() override {}
@@ -29,6 +29,6 @@ private:
 	};
 
 	TraceServices::IAnalysisSession& Session;
-	FAnimNextProvider& Provider;
+	FUAFProvider& Provider;
 };
-#endif // ANIMNEXT_TRACE_ENABLED
+#endif // UAF_TRACE_ENABLED
diff --git a/Engine/Plugins/Experimental/UAF/UAF/Source/UAFEditor/Private/RewindDebugger/AnimNextModuleTrack.cpp b/Engine/Plugins/Experimental/UAF/UAF/Source/UAFEditor/Private/RewindDebugger/AnimNextModuleTrack.cpp
--- a/Engine/Plugins/Experimental/UAF/UAF/Source/UAFEditor/Private/RewindDebugger/AnimNextModuleTrack.cpp
+++ b/Engine/Plugins/Experimental/UAF/UAF/Source/UAFEditor/Private/RewindDebugger/AnimNextModuleTrack.cpp
@@ -6,7 +6,7 @@
 
 #include UE_INLINE_GENERATED_CPP_BY_NAME(AnimNextModuleTrack)
 
-#if ANIMNEXT_TRACE_ENABLED
+#if UAF_TRACE_ENABLED
 #include "AnimNextProvider.h"
 #include "Editor.h"
 #include "IGameplayProvider.h"
@@ -36,7 +36,7 @@ FText FAnimNextModuleTrack::GetDisplayNameInternal() const
 	if (DisplayNameCache.IsEmpty())
 	{
 		TraceServices::FAnalysisSessionReadScope SessionReadScope(*IRewindDebugger::Instance()->GetAnalysisSession());
-		const FAnimNextProvider* AnimNextProvider = IRewindDebugger::Instance()->GetAnalysisSession()->ReadProvider<FAnimNextProvider>("AnimNextProvider");
+		const FUAFProvider* UAFProvider = IRewindDebugger::Instance()->GetAnalysisSession()->ReadProvider<FUAFProvider>("UAFProvider");
 		const IGameplayProvider* GameplayProvider = IRewindDebugger::Instance()->GetAnalysisSession()->ReadProvider<IGameplayProvider>("GameplayProvider");
 
 		if (InstanceId != 0)
@@ -119,7 +119,7 @@ bool FAnimNextModuleTrack::UpdateInternal()
 
 	bool bChanged = false;
 
-	if (const FAnimNextProvider* AnimNextProvider = IRewindDebugger::Instance()->GetAnalysisSession()->ReadProvider<FAnimNextProvider>("AnimNextProvider"))
+	if (const FUAFProvider* UAFProvider = IRewindDebugger::Instance()->GetAnalysisSession()->ReadProvider<FUAFProvider>("UAFProvider"))
 	{
 		double CurrentScrubTime = IRewindDebugger::Instance()->CurrentTraceTime();
 
@@ -135,7 +135,7 @@ bool FAnimNextModuleTrack::UpdateInternal()
 
 		if (InstanceId != 0)
 		{
-			if (const FDataInterfaceData* Data = AnimNextProvider->GetDataInterfaceData(InstanceId))
+			if (const FDataInterfaceData* Data = UAFProvider->GetDataInterfaceData(InstanceId))
 			{
 				if (PreviousScrubTime != CurrentScrubTime)
 				{
@@ -151,13 +151,13 @@ bool FAnimNextModuleTrack::UpdateInternal()
 					TraceServices::FFrame MarkerFrame;
 					if (FramesProvider.GetFrameFromTime(ETraceFrameType::TraceFrameType_Game, CurrentScrubTime, MarkerFrame))
 					{
-						Data->VariablesTimeline.EnumerateEvents(MarkerFrame.StartTime, MarkerFrame.EndTime, [AnimNextProvider, DetailsObject, AnalysisSession, this](double InStartTime, double InEndTime, uint32 InDepth, const FPropertyVariableData& VariableListData)
+						Data->VariablesTimeline.EnumerateEvents(MarkerFrame.StartTime, MarkerFrame.EndTime, [UAFProvider, DetailsObject, AnalysisSession, this](double InStartTime, double InEndTime, uint32 InDepth, const FPropertyVariableData& VariableListData)
 						{
 							const IGameplayProvider* GameplayProvider = AnalysisSession->ReadProvider<IGameplayProvider>("GameplayProvider");
 							if (VariableListData.ValueType == EPropertyVariableDataType::PropertyBag)
 							{
 								// First look up property description
-								const FPropertyDescriptionData* DescriptionData = AnimNextProvider->GetPropertyDescriptionData(VariableListData.DescriptionHash);
+								const FPropertyDescriptionData* DescriptionData = UAFProvider->GetPropertyDescriptionData(VariableListData.DescriptionHash);
 								if (DescriptionData == nullptr)
 								{
 									// Can't do anything without a description
@@ -167,7 +167,7 @@ bool FAnimNextModuleTrack::UpdateInternal()
 								// Load the property descriptions
 								FMemoryReader DescriptionReader(DescriptionData->Data);
 								FObjectAsTraceIdProxyArchiveReader DescriptionReaderProxy(DescriptionReader, GameplayProvider);
-								DescriptionReaderProxy.UsingCustomVersion(UE::UAF::FAnimNextTrace::CustomVersionGUID);
+								DescriptionReaderProxy.UsingCustomVersion(UE::UAF::FUAFTrace::CustomVersionGUID);
 								DescriptionReaderProxy << PropertyDescriptions;
 								DetailsObject->Properties.SetNum(DetailsObject->Properties.Num() + 1);
 								FInstancedPropertyBag& PropertyBag = DetailsObject->Properties.Last();
@@ -224,7 +224,7 @@ bool FAnimNextModuleTrack::HandleDoubleClickInternal()
 			TraceServices::FAnalysisSessionReadScope SessionReadScope(*AnalysisSession);
 
 			const IGameplayProvider* GameplayProvider = AnalysisSession->ReadProvider<IGameplayProvider>("GameplayProvider");
-			const FAnimNextProvider* AnimNextProvider = IRewindDebugger::Instance()->GetAnalysisSession()->ReadProvider<FAnimNextProvider>("AnimNextProvider");
+			const FUAFProvider* AnimNextProvider = IRewindDebugger::Instance()->GetAnalysisSession()->ReadProvider<FUAFProvider>("UAFProvider");
 
 			if (const FDataInterfaceData* ModuleData = AnimNextProvider->GetDataInterfaceData(InstanceId))
 			{
@@ -254,4 +254,4 @@ bool FAnimNextModuleTrackCreator::HasDebugInfoInternal(const RewindDebugger::FOb
 
 #undef LOCTEXT_NAMESPACE
 
-#endif // ANIMNEXT_TRACE_ENABLED
+#endif // UAF_TRACE_ENABLED
diff --git a/Engine/Plugins/Experimental/UAF/UAF/Source/UAFEditor/Private/RewindDebugger/AnimNextModuleTrack.h b/Engine/Plugins/Experimental/UAF/UAF/Source/UAFEditor/Private/RewindDebugger/AnimNextModuleTrack.h
--- a/Engine/Plugins/Experimental/UAF/UAF/Source/UAFEditor/Private/RewindDebugger/AnimNextModuleTrack.h
+++ b/Engine/Plugins/Experimental/UAF/UAF/Source/UAFEditor/Private/RewindDebugger/AnimNextModuleTrack.h
@@ -1,15 +1,15 @@
 // Copyright Epic Games, Inc. All Rights Reserved.
 #pragma once
 
-#include "RewindDebugger/AnimNextTrace.h"
+#include "RewindDebugger/UAFTrace.h"
 
-#if ANIMNEXT_TRACE_ENABLED
+#if UAF_TRACE_ENABLED
 #include "IRewindDebuggerTrackCreator.h"
 #include "RewindDebuggerTrack.h"
 #include "Textures/SlateIcon.h"
 #include "SEventTimelineView.h"
 #include "IDetailsView.h"
-#endif // ANIMNEXT_TRACE_ENABLED
+#endif // UAF_TRACE_ENABLED
 #include "StructUtils/PropertyBag.h"
 
 #include "AnimNextModuleTrack.generated.h"
@@ -27,7 +27,7 @@ public:
 	TArray<FInstancedStruct> NativeProperties;
 };
 
-#if ANIMNEXT_TRACE_ENABLED
+#if UAF_TRACE_ENABLED
 namespace UE::UAF::Editor
 {
 	
@@ -101,4 +101,4 @@ private:
 };
 
 }
-#endif // ANIMNEXT_TRACE_ENABLED
+#endif // UAF_TRACE_ENABLED
diff --git a/Engine/Plugins/Experimental/UAF/UAF/Source/UAFEditor/Private/RewindDebugger/AnimNextProvider.cpp b/Engine/Plugins/Experimental/UAF/UAF/Source/UAFEditor/Private/RewindDebugger/AnimNextProvider.cpp
--- a/Engine/Plugins/Experimental/UAF/UAF/Source/UAFEditor/Private/RewindDebugger/AnimNextProvider.cpp
+++ b/Engine/Plugins/Experimental/UAF/UAF/Source/UAFEditor/Private/RewindDebugger/AnimNextProvider.cpp
@@ -1,19 +1,19 @@
 // Copyright Epic Games, Inc. All Rights Reserved.
 
 #include "AnimNextProvider.h"
-#if ANIMNEXT_TRACE_ENABLED
+#if UAF_TRACE_ENABLED
 #include "ObjectTrace.h"
 
-FName FAnimNextProvider::ProviderName("AnimNextProvider");
+FName FUAFProvider::ProviderName("AnimNextProvider");
 
 #define LOCTEXT_NAMESPACE "AnimNextProvider"
 
-FAnimNextProvider::FAnimNextProvider(TraceServices::IAnalysisSession& InSession)
+FUAFProvider::FUAFProvider(TraceServices::IAnalysisSession& InSession)
 	: Session(InSession)
 {
 }
 
-void FAnimNextProvider::AppendInstance(uint64 InstanceId, uint64 HostInstanceId, uint64 AssetId, uint64 OuterObjectId)
+void FUAFProvider::AppendInstance(uint64 InstanceId, uint64 HostInstanceId, uint64 AssetId, uint64 OuterObjectId)
 {
 	Session.WriteAccessCheck();
 
@@ -34,7 +34,7 @@ void FAnimNextProvider::AppendInstance(uint64 InstanceId, uint64 HostInstanceId,
 	ChildList.AddUnique(Data);
 }
 
-void FAnimNextProvider::AppendVariables(double ProfileTime, double RecordingTime, uint64 DataInterfaceId, EPropertyVariableDataType PropertyVariableDataType, uint32 PropertyDescriptionHash, const TArrayView<const uint8>& VariableData)
+void FUAFProvider::AppendVariables(double ProfileTime, double RecordingTime, uint64 DataInterfaceId, EPropertyVariableDataType PropertyVariableDataType, uint32 PropertyDescriptionHash, const TArrayView<const uint8>& VariableData)
 {
 	Session.WriteAccessCheck();
 	
@@ -52,7 +52,7 @@ void FAnimNextProvider::AppendVariables(double ProfileTime, double RecordingTime
 	}
 }
 
-void FAnimNextProvider::AppendVariableDescriptions(uint32 PropertyDescriptionHash, const TArrayView<const uint8>& VariableDescriptionData)
+void FUAFProvider::AppendVariableDescriptions(uint32 PropertyDescriptionHash, const TArrayView<const uint8>& VariableDescriptionData)
 {
 	Session.WriteAccessCheck();
 	
@@ -66,7 +66,7 @@ void FAnimNextProvider::AppendVariableDescriptions(uint32 PropertyDescriptionHas
 }
 
 
-bool FAnimNextProvider::GetModuleId(uint64 ComponentId, uint64& OutModuleId) const
+bool FUAFProvider::GetModuleId(uint64 ComponentId, uint64& OutModuleId) const
 {
 	Session.ReadAccessCheck();
 	
@@ -78,7 +78,7 @@ bool FAnimNextProvider::GetModuleId(uint64 ComponentId, uint64& OutModuleId) con
 	return false;
 }
 
-const FDataInterfaceData* FAnimNextProvider::GetDataInterfaceData(uint64 DataInterfaceId) const
+const FDataInterfaceData* FUAFProvider::GetDataInterfaceData(uint64 DataInterfaceId) const
 {
 	const TSharedRef<FDataInterfaceData>* Data = DataInterfaceData.Find(DataInterfaceId);
 	if (Data)
@@ -88,7 +88,7 @@ const FDataInterfaceData* FAnimNextProvider::GetDataInterfaceData(uint64 DataInt
 	return nullptr;
 }
 
-void FAnimNextProvider::EnumerateChildInstances(uint64 InstanceId, TFunctionRef<void(const FDataInterfaceData&)> Callback) const
+void FUAFProvider::EnumerateChildInstances(uint64 InstanceId, TFunctionRef<void(const FDataInterfaceData&)> Callback) const
 {
 	Session.ReadAccessCheck();
 
@@ -103,7 +103,7 @@ void FAnimNextProvider::EnumerateChildInstances(uint64 InstanceId, TFunctionRef<
 
 }
 
-void FAnimNextProvider::EnumerateDataInterfaces(const double InStartTime, TFunctionRef<void(const TSharedRef<const FDataInterfaceData>)> Callback) const
+void FUAFProvider::EnumerateDataInterfaces(const double InStartTime, TFunctionRef<void(const TSharedRef<const FDataInterfaceData>)> Callback) const
 {
 	Session.ReadAccessCheck();
 
@@ -117,4 +117,4 @@ void FAnimNextProvider::EnumerateDataInterfaces(const double InStartTime, TFunct
 }
 
 #undef LOCTEXT_NAMESPACE
-#endif // ANIMNEXT_TRACE_ENABLED
+#endif // UAF_TRACE_ENABLED
diff --git a/Engine/Plugins/Experimental/UAF/UAF/Source/UAFEditor/Private/RewindDebugger/AnimNextProvider.h b/Engine/Plugins/Experimental/UAF/UAF/Source/UAFEditor/Private/RewindDebugger/AnimNextProvider.h
--- a/Engine/Plugins/Experimental/UAF/UAF/Source/UAFEditor/Private/RewindDebugger/AnimNextProvider.h
+++ b/Engine/Plugins/Experimental/UAF/UAF/Source/UAFEditor/Private/RewindDebugger/AnimNextProvider.h
@@ -2,9 +2,9 @@
 
 #pragma once
 
-#include "RewindDebugger/AnimNextTrace.h"
+#include "RewindDebugger/UAFTrace.h"
 
-#if ANIMNEXT_TRACE_ENABLED
+#if UAF_TRACE_ENABLED
 #include "Engine/EngineBaseTypes.h"
 #include "Model/PointTimeline.h"
 #include "Containers/Map.h"
@@ -58,12 +58,12 @@ struct FDataInterfaceData
 	double EndTime;
 };
 
-class FAnimNextProvider : public TraceServices::IProvider
+class FUAFProvider : public TraceServices::IProvider
 {
 public:
 	static FName ProviderName;
 
-	FAnimNextProvider(TraceServices::IAnalysisSession& InSession);
+	FUAFProvider(TraceServices::IAnalysisSession& InSession);
 
 	void AppendInstance(uint64 InstanceId, uint64 HostInstanceId,  uint64 AssetId, uint64 OuterObjectId);
 	void AppendVariables(double ProfileTime, double RecordingTime, uint64 InModuleInstanceId, EPropertyVariableDataType PropertyVariableDataType, uint32 PropertyDescriptionHash, const TArrayView<const uint8>& VariableData);
@@ -84,4 +84,4 @@ private:
 	TMap<uint64, TArray<TSharedRef<FDataInterfaceData>>> HostToChildDataMap;
 	TMap<uint32, FPropertyDescriptionData> HashToPropertyDescriptionDataMap;
 };
-#endif //ANIMNEXT_TRACE_ENABLED
\ No newline at end of file
+#endif //UAF_TRACE_ENABLED
\ No newline at end of file
diff --git a/Engine/Plugins/Experimental/UAF/UAF/Source/UAFEditor/Private/RewindDebugger/AnimNextTraceModule.cpp b/Engine/Plugins/Experimental/UAF/UAF/Source/UAFEditor/Private/RewindDebugger/AnimNextTraceModule.cpp
--- a/Engine/Plugins/Experimental/UAF/UAF/Source/UAFEditor/Private/RewindDebugger/AnimNextTraceModule.cpp
+++ b/Engine/Plugins/Experimental/UAF/UAF/Source/UAFEditor/Private/RewindDebugger/AnimNextTraceModule.cpp
@@ -2,33 +2,33 @@
 
 #include "AnimNextTraceModule.h"
 
-#if ANIMNEXT_TRACE_ENABLED
+#if UAF_TRACE_ENABLED
 #include "AnimNextProvider.h"
 #include "AnimNextAnalyzer.h"
 
-FName FAnimNextTraceModule::ModuleName("UAF");
+FName FUAFTraceModule::ModuleName("UAF");
 
-void FAnimNextTraceModule::GetModuleInfo(TraceServices::FModuleInfo& OutModuleInfo)
+void FUAFTraceModule::GetModuleInfo(TraceServices::FModuleInfo& OutModuleInfo)
 {
 	OutModuleInfo.Name = ModuleName;
 	OutModuleInfo.DisplayName = TEXT("UAF");
 }
 
-void FAnimNextTraceModule::OnAnalysisBegin(TraceServices::IAnalysisSession& InSession)
+void FUAFTraceModule::OnAnalysisBegin(TraceServices::IAnalysisSession& InSession)
 {
-	TSharedPtr<FAnimNextProvider> AnimNextProvider = MakeShared<FAnimNextProvider>(InSession);
-	InSession.AddProvider(FAnimNextProvider::ProviderName, AnimNextProvider);
+	TSharedPtr<FUAFProvider> UAFProvider = MakeShared<FUAFProvider>(InSession);
+	InSession.AddProvider(FUAFProvider::ProviderName, UAFProvider);
 
-	InSession.AddAnalyzer(new FAnimNextAnalyzer(InSession, *AnimNextProvider));
+	InSession.AddAnalyzer(new FUAFAnalyzer(InSession, *UAFProvider));
 }
 
-void FAnimNextTraceModule::GetLoggers(TArray<const TCHAR *>& OutLoggers)
+void FUAFTraceModule::GetLoggers(TArray<const TCHAR *>& OutLoggers)
 {
 	OutLoggers.Add(TEXT("UAF"));
 }
 
-void FAnimNextTraceModule::GenerateReports(const TraceServices::IAnalysisSession& Session, const TCHAR* CmdLine, const TCHAR* OutputDirectory)
+void FUAFTraceModule::GenerateReports(const TraceServices::IAnalysisSession& Session, const TCHAR* CmdLine, const TCHAR* OutputDirectory)
 {
 
 }
-#endif //ANIMNEXT_TRACE_ENABLED
+#endif //UAF_TRACE_ENABLED
diff --git a/Engine/Plugins/Experimental/UAF/UAF/Source/UAFEditor/Private/RewindDebugger/AnimNextTraceModule.h b/Engine/Plugins/Experimental/UAF/UAF/Source/UAFEditor/Private/RewindDebugger/AnimNextTraceModule.h
--- a/Engine/Plugins/Experimental/UAF/UAF/Source/UAFEditor/Private/RewindDebugger/AnimNextTraceModule.h
+++ b/Engine/Plugins/Experimental/UAF/UAF/Source/UAFEditor/Private/RewindDebugger/AnimNextTraceModule.h
@@ -2,15 +2,15 @@
 
 #pragma once
 
-#include "RewindDebugger/AnimNextTrace.h"
+#include "RewindDebugger/UAFTrace.h"
 
-#if ANIMNEXT_TRACE_ENABLED
+#if UAF_TRACE_ENABLED
 #include "CoreMinimal.h"
 #include "TraceServices/ModuleService.h"
 
 class IAnimationProvider;
 
-class FAnimNextTraceModule : public TraceServices::IModule
+class FUAFTraceModule : public TraceServices::IModule
 {
 public:
 	// TraceServices::IModule interface
@@ -23,5 +23,5 @@ public:
 private:
 	static FName ModuleName;
 };
-#endif //ANIMNEXT_TRACE_ENABLED
+#endif //UAF_TRACE_ENABLED
 
diff --git a/Engine/Plugins/Experimental/UAF/UAF/Source/UAFEditor/Private/RewindDebugger/RewindDebuggerAnimNext.cpp b/Engine/Plugins/Experimental/UAF/UAF/Source/UAFEditor/Private/RewindDebugger/RewindDebuggerAnimNext.cpp
--- a/Engine/Plugins/Experimental/UAF/UAF/Source/UAFEditor/Private/RewindDebugger/RewindDebuggerAnimNext.cpp
+++ b/Engine/Plugins/Experimental/UAF/UAF/Source/UAFEditor/Private/RewindDebugger/RewindDebuggerAnimNext.cpp
@@ -2,17 +2,17 @@
 
 #include "RewindDebuggerAnimNext.h"
 
-#if ANIMNEXT_TRACE_ENABLED
+#if UAF_TRACE_ENABLED
 #include "IGameplayProvider.h"
 
-FRewindDebuggerAnimNext::FRewindDebuggerAnimNext()
+FRewindDebuggerUAF::FRewindDebuggerUAF()
 {
 
 }
 
-void FRewindDebuggerAnimNext::Update(float DeltaTime, IRewindDebugger* RewindDebugger)
+void FRewindDebuggerUAF::Update(float DeltaTime, IRewindDebugger* RewindDebugger)
 {
 	// Future place to sync anim next workspace editor debugging when scrubbing
 }
 
-#endif // ANIMNEXT_TRACE_ENABLED
+#endif // UAF_TRACE_ENABLED
diff --git a/Engine/Plugins/Experimental/UAF/UAF/Source/UAFEditor/Private/RewindDebugger/RewindDebuggerAnimNext.h b/Engine/Plugins/Experimental/UAF/UAF/Source/UAFEditor/Private/RewindDebugger/RewindDebuggerAnimNext.h
--- a/Engine/Plugins/Experimental/UAF/UAF/Source/UAFEditor/Private/RewindDebugger/RewindDebuggerAnimNext.h
+++ b/Engine/Plugins/Experimental/UAF/UAF/Source/UAFEditor/Private/RewindDebugger/RewindDebuggerAnimNext.h
@@ -1,21 +1,21 @@
 // Copyright Epic Games, Inc. All Rights Reserved.
 #pragma once
 
-#include "RewindDebugger/AnimNextTrace.h"
+#include "RewindDebugger/UAFTrace.h"
 
-#if ANIMNEXT_TRACE_ENABLED
+#if UAF_TRACE_ENABLED
 #include "CoreMinimal.h"
 #include "IRewindDebugger.h"
 #include "IRewindDebuggerExtension.h"
 
 // Rewind debugger extension for Chooser support
 
-class FRewindDebuggerAnimNext : public IRewindDebuggerExtension
+class FRewindDebuggerUAF : public IRewindDebuggerExtension
 {
 public:
-	FRewindDebuggerAnimNext();
-	virtual ~FRewindDebuggerAnimNext() {};
+	FRewindDebuggerUAF();
+	virtual ~FRewindDebuggerUAF() {};
 
 	virtual void Update(float DeltaTime, IRewindDebugger* RewindDebugger) override;
 };
-#endif //ANIMNEXT_TRACE_ENABLED
+#endif //UAF_TRACE_ENABLED
diff --git a/Engine/Plugins/Experimental/UAF/UAF/Source/UAFEditor/Private/UAF/Debug/SVariableWatch.cpp b/Engine/Plugins/Experimental/UAF/UAF/Source/UAFEditor/Private/UAF/Debug/SVariableWatch.cpp
--- a/Engine/Plugins/Experimental/UAF/UAF/Source/UAFEditor/Private/UAF/Debug/SVariableWatch.cpp
+++ b/Engine/Plugins/Experimental/UAF/UAF/Source/UAFEditor/Private/UAF/Debug/SVariableWatch.cpp
@@ -12,10 +12,10 @@
 #include "IWorkspaceEditorModule.h"
 #include "ObjectAsTraceIdProxyArchiveReader.h"
 #include "ObjectTrace.h"
-#include "RewindDebugger/AnimNextTrace.h"
-#if ANIMNEXT_TRACE_ENABLED
+#include "RewindDebugger/UAFTrace.h"
+#if UAF_TRACE_ENABLED
 #include "RewindDebugger/AnimNextProvider.h"
-#endif // ANIMNEXT_TRACE_ENABLED
+#endif // UAF_TRACE_ENABLED
 #include "SceneOutlinerPublicTypes.h"
 #include "SSceneOutliner.h"
 #include "TraceServices/Model/AnalysisSession.h"
@@ -181,7 +181,7 @@ namespace UE::UAF::Editor
 	
 	void SVariableWatch::RegenerateDebugInstances(const bool bForceRegenerate)
 	{
-#if ANIMNEXT_TRACE_ENABLED
+#if UAF_TRACE_ENABLED
 		if (!RewindDebugger)
 		{
 			return;
@@ -205,7 +205,7 @@ namespace UE::UAF::Editor
 		}
 
 		const IGameplayProvider* GameplayProvider = AnalysisSession->ReadProvider<IGameplayProvider>("GameplayProvider");
-		const FAnimNextProvider* UAFProvider = AnalysisSession->ReadProvider<FAnimNextProvider>("AnimNextProvider");
+		const FUAFProvider* UAFProvider = AnalysisSession->ReadProvider<FUAFProvider>("UAFProvider");
 		
 		if (!GameplayProvider || !UAFProvider)
 		{
@@ -298,7 +298,7 @@ namespace UE::UAF::Editor
 		}
 		
 		OnDebugInstanceChanged();
-#endif // ANIMNEXT_TRACE_ENABLED
+#endif // UAF_TRACE_ENABLED
 	}
 
 	void SVariableWatch::HandleFocusedDocumentChanged(const UE::Workspace::FWorkspaceDocument& InDocument)
@@ -322,7 +322,7 @@ namespace UE::UAF::Editor
 	
 	void SVariableWatch::OnDebugInstanceChanged()
 	{
-#if ANIMNEXT_TRACE_ENABLED
+#if UAF_TRACE_ENABLED
 		VariablesBoxSource.Empty();
 	
 		const TraceServices::IAnalysisSession* AnalysisSession = RewindDebugger->GetAnalysisSession();
@@ -332,7 +332,7 @@ namespace UE::UAF::Editor
 		const double CurrentTraceTime  = RewindDebugger->CurrentTraceTime();
 				
 		const IGameplayProvider* GameplayProvider = AnalysisSession->ReadProvider<IGameplayProvider>("GameplayProvider");
-		const FAnimNextProvider* UAFProvider = AnalysisSession->ReadProvider<FAnimNextProvider>("AnimNextProvider");
+		const FUAFProvider* UAFProvider = AnalysisSession->ReadProvider<FUAFProvider>("UAFProvider");
         if (!GameplayProvider || !UAFProvider)
         {
 	        return;
@@ -369,7 +369,7 @@ namespace UE::UAF::Editor
 				TArray<FPropertyBagPropertyDesc> PropertyDescriptions;
 				FMemoryReader DescriptionReader(DescriptionData->Data);
 				FObjectAsTraceIdProxyArchiveReader DescriptionReaderProxy(DescriptionReader, GameplayProvider);
-				DescriptionReaderProxy.UsingCustomVersion(UE::UAF::FAnimNextTrace::CustomVersionGUID);
+				DescriptionReaderProxy.UsingCustomVersion(UE::UAF::FUAFTrace::CustomVersionGUID);
 				DescriptionReaderProxy << PropertyDescriptions;
 				
 				//Properties->Properties.SetNum(Properties->Properties.Num() + 1);
@@ -411,7 +411,7 @@ namespace UE::UAF::Editor
 		}
 		
 		VariablesBox->RequestListRefresh();
-#endif // ANIMNEXT_TRACE_ENABLED
+#endif // UAF_TRACE_ENABLED
 	}
 
 	void SVariableWatch::HandleWorkspaceEditorModeChanged(const FEditorModeID& EditorModeId, bool bIsEnteringMode)
diff --git a/Engine/Plugins/Experimental/UAF/UAFAnimGraph/Source/UAFAnimGraph/Private/Graph/RigUnit_AnimNextRunAnimationGraph_v2.cpp b/Engine/Plugins/Experimental/UAF/UAFAnimGraph/Source/UAFAnimGraph/Private/Graph/RigUnit_AnimNextRunAnimationGraph_v2.cpp
--- a/Engine/Plugins/Experimental/UAF/UAFAnimGraph/Source/UAFAnimGraph/Private/Graph/RigUnit_AnimNextRunAnimationGraph_v2.cpp
+++ b/Engine/Plugins/Experimental/UAF/UAFAnimGraph/Source/UAFAnimGraph/Private/Graph/RigUnit_AnimNextRunAnimationGraph_v2.cpp
@@ -125,8 +125,8 @@ FRigUnit_AnimNextRunAnimationGraph_v2_Execute()
 	EvaluateGraphContext.SetBindingObject(RefPose.SkeletalMeshComponent);
 	EvaluateGraph(EvaluateGraphContext, Result);
 
-	TRACE_ANIMNEXT_MODULE(ModuleInstance);
-	TRACE_ANIMNEXT_GRAPHINSTANCES(HostInstance);
+	TRACE_UAF_MODULE(ModuleInstance);
+	TRACE_UAF_GRAPHINSTANCES(HostInstance);
 
 	// We might have appended new input/output events, append them
 	{
diff --git a/Engine/Plugins/Experimental/UAF/UAFAnimGraph/Source/UAFAnimGraph/Private/Graph/TraceAnimNextGraphInstances.cpp b/Engine/Plugins/Experimental/UAF/UAFAnimGraph/Source/UAFAnimGraph/Private/Graph/TraceAnimNextGraphInstances.cpp
--- a/Engine/Plugins/Experimental/UAF/UAFAnimGraph/Source/UAFAnimGraph/Private/Graph/TraceAnimNextGraphInstances.cpp
+++ b/Engine/Plugins/Experimental/UAF/UAFAnimGraph/Source/UAFAnimGraph/Private/Graph/TraceAnimNextGraphInstances.cpp
@@ -13,7 +13,7 @@
 #include "TraitCore/NodeInstance.h"
 #include "TraitInterfaces/IHierarchy.h"
 
-#if ANIMNEXT_TRACE_ENABLED
+#if UAF_TRACE_ENABLED
 
 UE_TRACE_EVENT_BEGIN(UAFAnimGraph, EvaluationProgram)
 	UE_TRACE_EVENT_FIELD(uint64, Cycle)
@@ -56,7 +56,7 @@ namespace UE::UAF
 				FAnimNextGraphInstance& Graph = Trait.GetNodeInstance()->GetOwner();
 				if (!TracedInstances.Contains(Graph.GetUniqueId()))
 				{
-					TRACE_ANIMNEXT_VARIABLES(&Graph, Graph.GetModuleInstance() ? Graph.GetModuleInstance()->GetObject() : nullptr);
+					TRACE_UAF_VARIABLES(&Graph, Graph.GetModuleInstance() ? Graph.GetModuleInstance()->GetObject() : nullptr);
 				}
 
 				FTraitStackBinding TraitStack;
diff --git a/Engine/Plugins/Experimental/UAF/UAFAnimGraph/Source/UAFAnimGraph/Private/Graph/TraceAnimNextGraphInstances.h b/Engine/Plugins/Experimental/UAF/UAFAnimGraph/Source/UAFAnimGraph/Private/Graph/TraceAnimNextGraphInstances.h
--- a/Engine/Plugins/Experimental/UAF/UAFAnimGraph/Source/UAFAnimGraph/Private/Graph/TraceAnimNextGraphInstances.h
+++ b/Engine/Plugins/Experimental/UAF/UAFAnimGraph/Source/UAFAnimGraph/Private/Graph/TraceAnimNextGraphInstances.h
@@ -2,9 +2,9 @@
 #pragma once
 
 #include "EvaluationVM/EvaluationProgram.h"
-#include "RewindDebugger/AnimNextTrace.h"
+#include "RewindDebugger/UAFTrace.h"
 
-#if ANIMNEXT_TRACE_ENABLED
+#if UAF_TRACE_ENABLED
 
 struct FAnimNextGraphInstance;
 
@@ -16,12 +16,12 @@ namespace UE::UAF
 	void TraceEvaluationProgram(const FEvaluationProgram& Program, const FAnimNextGraphInstance& RootGraph);
 }
 
-#define TRACE_ANIMNEXT_GRAPHINSTANCES(RootGraph) UE::UAF::TraceGraphInstances(RootGraph);
-#define TRACE_ANIMNEXT_EVALUATIONPROGRAM(Program, RootGraph) UE::UAF::TraceEvaluationProgram(Program, RootGraph);
+#define TRACE_UAF_GRAPHINSTANCES(RootGraph) UE::UAF::TraceGraphInstances(RootGraph);
+#define TRACE_UAF_EVALUATIONPROGRAM(Program, RootGraph) UE::UAF::TraceEvaluationProgram(Program, RootGraph);
 
 #else
 
-#define TRACE_ANIMNEXT_GRAPHINSTANCES(RootGraph)
-#define TRACE_ANIMNEXT_EVALUATIONPROGRAM(Program, RootGraph)
+#define TRACE_UAF_GRAPHINSTANCES(RootGraph)
+#define TRACE_UAF_EVALUATIONPROGRAM(Program, RootGraph)
 
 #endif
diff --git a/Engine/Plugins/Experimental/UAF/UAFAnimGraph/Source/UAFAnimGraph/Private/TraitInterfaces/IEvaluate.cpp b/Engine/Plugins/Experimental/UAF/UAFAnimGraph/Source/UAFAnimGraph/Private/TraitInterfaces/IEvaluate.cpp
--- a/Engine/Plugins/Experimental/UAF/UAFAnimGraph/Source/UAFAnimGraph/Private/TraitInterfaces/IEvaluate.cpp
+++ b/Engine/Plugins/Experimental/UAF/UAFAnimGraph/Source/UAFAnimGraph/Private/TraitInterfaces/IEvaluate.cpp
@@ -230,7 +230,7 @@ namespace UE::UAF
 
 		const FEvaluationProgram EvaluationProgram = EvaluateGraph(EvaluateGraphContext);
 
-		TRACE_ANIMNEXT_EVALUATIONPROGRAM(EvaluationProgram, GraphInstance);
+		TRACE_UAF_EVALUATIONPROGRAM(EvaluationProgram, GraphInstance);
 
 		{
 			SCOPED_NAMED_EVENT(UAF_EvaluatePose, FColor::Orange);
