Change workspace outliner icon for Collapsed graphs to match other tools

diff --git a/Engine/Plugins/Experimental/UAF/UAF/Source/UAFEditor/Private/Common/AnimNextCollapseNodeItemDetails.cpp b/Engine/Plugins/Experimental/UAF/UAF/Source/UAFEditor/Private/Common/AnimNextCollapseNodeItemDetails.cpp
--- a/Engine/Plugins/Experimental/UAF/UAF/Source/UAFEditor/Private/Common/AnimNextCollapseNodeItemDetails.cpp
+++ b/Engine/Plugins/Experimental/UAF/UAF/Source/UAFEditor/Private/Common/AnimNextCollapseNodeItemDetails.cpp
@@ -172,7 +172,7 @@ UPackage* FAnimNextCollapseNodeItemDetails::GetPackage(const FWorkspaceOutlinerI
 
 const FSlateBrush* FAnimNextCollapseNodeItemDetails::GetItemIcon(const FWorkspaceOutlinerItemExport& Export) const
 {
-	return FAppStyle::GetBrush(TEXT("GraphEditor.EventGraph_24x"));
+	return FAppStyle::GetBrush(TEXT("GraphEditor.SubGraph_24x"));
 }
 
 void FAnimNextCollapseNodeItemDetails::RegisterToolMenuExtensions()
