UAF Browser Layout changes
- UAF Browser disabled by default via CVar: a.UAF.EnableUAFBrowser
- Added UAF Browser to UAF Editor Mode
- Added UAF Project / User Settings.
- Initial UAF Tagged Asset Browser integration (SUAFBrowser still in Prototype Phase).

#tests editor

diff --git a/Engine/Plugins/Experimental/UAF/UAF/Source/UAFEditor/Private/AnimNextEditorModule.cpp b/Engine/Plugins/Experimental/UAF/UAF/Source/UAFEditor/Private/AnimNextEditorModule.cpp
--- a/Engine/Plugins/Experimental/UAF/UAF/Source/UAFEditor/Private/AnimNextEditorModule.cpp
+++ b/Engine/Plugins/Experimental/UAF/UAF/Source/UAFEditor/Private/AnimNextEditorModule.cpp
@@ -17,6 +17,7 @@
 #include "AnimNextEdGraphNode.h"
 #include "Common/AnimNextCompilerResultsTabSummoner.h"
 #include "Common/AnimNextFindTabSummoner.h"
+#include "Common/UAFBrowserTabSummoner.h"
 #include "HAL/PlatformApplicationMisc.h"
 #include "Variables/VariableCustomization.h"
 #include "Variables/VariableBindingPropertyCustomization.h"
@@ -69,6 +70,7 @@
 #include "RewindDebugger/RewindDebuggerAnimNext.h"
 #include "RewindDebugger/AnimNextModuleTrack.h"
 #include "RewindDebugger/AnimNextTraceModule.h"
+#include "Settings/UAFEditorUserSettings.h"
 #include "Variables/VariableReferencePropertyCustomization.h"
 #include "AnimNextSharedVariableNode.h"
 #include "UAFStyle.h"
@@ -78,6 +80,8 @@
 #include "Toolkits/AssetEditorToolkitMenuContext.h"
 #include "UAF/Viewport/SystemViewportController.h"
 #include "Variables/Outliner/VariablesOutlinerCommands.h"
+#include "WidgetDrawerConfig.h"
+#include "Widgets/Docking/SDockTab.h"
 #include "Wizard/AssetWizard.h"
 #include "Interfaces/IPluginManager.h"
 #include "AssetToolsModule.h"
@@ -123,7 +127,13 @@ static FAutoConsoleCommand ToggleHiddenUAFContent(
     TEXT("UAF.ToggleHiddenContent"),
     TEXT("Toggles whether internal UAF content is shown in the editor"),
     ToggleUAFContentHidden);
-	
+
+static bool bEnableUAFBrowser = false;
+static FAutoConsoleVariableRef CVarEnableUAFBrowser(
+	TEXT("a.UAF.EnableUAFBrowser"),
+	bEnableUAFBrowser,
+	TEXT("Experimental. If true enables the UAF Browser."));
+
 void FAnimNextEditorModule::StartupModule()
 {
 	FVariableOverrideCommands::Register();
@@ -197,31 +207,49 @@ void FAnimNextEditorModule::StartupModule()
 	Workspace::IWorkspaceEditorModule& WorkspaceEditorModule = FModuleManager::Get().LoadModuleChecked<Workspace::IWorkspaceEditorModule>("WorkspaceEditor");
 
 	WorkspaceEditorModule.OnRegisterTabsForEditor().AddLambda([](FWorkflowAllowedTabSet& TabFactories, const TSharedRef<FTabManager>& InTabManager, TSharedPtr<UE::Workspace::IWorkspaceEditor> InEditorPtr)
-		{
-			TSharedRef<FAnimNextCompilerResultsTabSummoner> CompilerResultsTabSummoner = MakeShared<FAnimNextCompilerResultsTabSummoner>(InEditorPtr);
-			TabFactories.RegisterFactory(CompilerResultsTabSummoner);
-			CompilerResultsTabSummoner->RegisterTabSpawner(InTabManager, nullptr);
+	{
+		TSharedRef<FAnimNextCompilerResultsTabSummoner> CompilerResultsTabSummoner = MakeShared<FAnimNextCompilerResultsTabSummoner>(InEditorPtr);
+		TabFactories.RegisterFactory(CompilerResultsTabSummoner);
+		CompilerResultsTabSummoner->RegisterTabSpawner(InTabManager, nullptr);
 
-			TSharedRef<FAnimNextFindTabSummoner> FindTabSummoner = MakeShared<FAnimNextFindTabSummoner>(InEditorPtr);
-			TabFactories.RegisterFactory(FindTabSummoner);
-			FindTabSummoner->RegisterTabSpawner(InTabManager, nullptr);
+		TSharedRef<FAnimNextFindTabSummoner> FindTabSummoner = MakeShared<FAnimNextFindTabSummoner>(InEditorPtr);
+		TabFactories.RegisterFactory(FindTabSummoner);
+		FindTabSummoner->RegisterTabSpawner(InTabManager, nullptr);
 
-			TSharedRef<FAnimNextVariablesTabSummoner> VariablesTabSummoner = MakeShared<FAnimNextVariablesTabSummoner>(InEditorPtr);
-			TabFactories.RegisterFactory(VariablesTabSummoner);
-			VariablesTabSummoner->RegisterTabSpawner(InTabManager, nullptr);
+		TSharedRef<FAnimNextVariablesTabSummoner> VariablesTabSummoner = MakeShared<FAnimNextVariablesTabSummoner>(InEditorPtr);
+		TabFactories.RegisterFactory(VariablesTabSummoner);
+		VariablesTabSummoner->RegisterTabSpawner(InTabManager, nullptr);
 
-			TSharedRef<FAnimNextFunctionsTabSummoner> FunctionsTabSummoner = MakeShared<FAnimNextFunctionsTabSummoner>(InEditorPtr);
-			TabFactories.RegisterFactory(FunctionsTabSummoner);
-			FunctionsTabSummoner->RegisterTabSpawner(InTabManager, nullptr);
+		TSharedRef<FAnimNextFunctionsTabSummoner> FunctionsTabSummoner = MakeShared<FAnimNextFunctionsTabSummoner>(InEditorPtr);
+		TabFactories.RegisterFactory(FunctionsTabSummoner);
+		FunctionsTabSummoner->RegisterTabSpawner(InTabManager, nullptr);
 
-			TSharedRef<FAnimNextFindAndReplaceTabSummoner> FindAndReplaceTabSummoner = MakeShared<FAnimNextFindAndReplaceTabSummoner>(InEditorPtr);
-			TabFactories.RegisterFactory(FindAndReplaceTabSummoner);
-			FindAndReplaceTabSummoner->RegisterTabSpawner(InTabManager, nullptr);
+		TSharedRef<FAnimNextFindAndReplaceTabSummoner> FindAndReplaceTabSummoner = MakeShared<FAnimNextFindAndReplaceTabSummoner>(InEditorPtr);
+		TabFactories.RegisterFactory(FindAndReplaceTabSummoner);
+		FindAndReplaceTabSummoner->RegisterTabSpawner(InTabManager, nullptr);
+
+		TSharedRef<FVariableWatchTabSummoner> VariableWatchTabSummoner = MakeShared<FVariableWatchTabSummoner>(InEditorPtr);
+		TabFactories.RegisterFactory(VariableWatchTabSummoner);
+		VariableWatchTabSummoner->RegisterTabSpawner(InTabManager, nullptr);
+
+		if (bEnableUAFBrowser)
+		{
+			TSharedRef<FUAFBrowserTabSummoner> DockedBrowserTabSummoner = MakeShared<FUAFBrowserTabSummoner>(InEditorPtr
+				, ETabState::ClosedTab
+				, FUAFBrowserTabSummoner::DockedTabID
+				, LOCTEXT("UAFBrowserDockedLabel", "UAF Browser (Docked)"));
+			TabFactories.RegisterFactory(DockedBrowserTabSummoner);
+			DockedBrowserTabSummoner->RegisterTabSpawner(InTabManager, nullptr);
+
+			TSharedRef<FUAFBrowserTabSummoner> SideBarBrowserTabSummoner = MakeShared<FUAFBrowserTabSummoner>(InEditorPtr
+				, ETabState::SidebarTab
+				, FUAFBrowserTabSummoner::SidebarTabID
+				, LOCTEXT("UAFBrowserSidebarLabel", "UAF Browser (Sidebar)"));
+			TabFactories.RegisterFactory(SideBarBrowserTabSummoner);
+			SideBarBrowserTabSummoner->RegisterTabSpawner(InTabManager, nullptr);
+		}
+	});
 
-			TSharedRef<FVariableWatchTabSummoner> VariableWatchTabSummoner = MakeShared<FVariableWatchTabSummoner>(InEditorPtr);
-			TabFactories.RegisterFactory(VariableWatchTabSummoner);
-			VariableWatchTabSummoner->RegisterTabSpawner(InTabManager, nullptr);
-		});
 
 	WorkspaceEditorModule.OnExtendTabs().AddLambda([](FLayoutExtender& InLayoutExtender, TSharedPtr<UE::Workspace::IWorkspaceEditor> InEditorPtr)
 	{
@@ -235,10 +263,25 @@ void FAnimNextEditorModule::StartupModule()
 		InLayoutExtender.ExtendLayout(FTabId(Workspace::WorkspaceTabs::BottomMiddleDocumentArea), ELayoutExtensionPosition::After, FindAndReplaceTab);
 
 		FTabManager::FTab VariablesTab(FTabId(VariablesTabName), ETabState::OpenedTab);
-		InLayoutExtender.ExtendStack(Workspace::WorkspaceTabs::BottomLeftDocumentArea, ELayoutExtensionPosition::After, VariablesTab, true);
+		InLayoutExtender.ExtendLayout(FTabId(Workspace::WorkspaceTabs::BottomLeftDocumentArea), ELayoutExtensionPosition::After, VariablesTab);
 
 		FTabManager::FTab FunctionsTab(FTabId(FunctionsTabName), ETabState::OpenedTab);
 		InLayoutExtender.ExtendStack(Workspace::WorkspaceTabs::BottomLeftDocumentArea, ELayoutExtensionPosition::After, FunctionsTab);
+
+		if (bEnableUAFBrowser)
+		{
+			FTabManager::FTab UAFBrowserDockedTab(FTabId(FUAFBrowserTabSummoner::DockedTabID), ETabState::ClosedTab);
+			InLayoutExtender.ExtendLayout(FTabId(Workspace::WorkspaceTabs::BottomMiddleDocumentArea), ELayoutExtensionPosition::After, UAFBrowserDockedTab);
+
+			const float SidebarSizeCoefficient = 0.2f;
+			const bool bPinnedInSidebar = false;
+			FTabManager::FTab UAFBrowserSidebarTab(FTabId(FUAFBrowserTabSummoner::SidebarTabID)
+				, ETabState::SidebarTab
+				, ESidebarLocation::Right
+				, SidebarSizeCoefficient
+				, bPinnedInSidebar);
+			InLayoutExtender.ExtendLayout(FTabId(Workspace::WorkspaceTabs::RightSidebarDocumentArea), ELayoutExtensionPosition::After, UAFBrowserSidebarTab);
+		}
 	});
 
 	UToolMenu* AnimNextEdGraphContextMenu = UToolMenus::Get()->ExtendMenu( "GraphEditor.GraphNodeContextMenu.AnimNextEdGraphNode");
@@ -1478,6 +1521,31 @@ void FAnimNextEditorModule::RegisterWorkspaceDocumentTypes(Workspace::IWorkspace
 
 						CommandList->MapAction(Commands.FindAndReplaceInAnimNextRigVMAsset,
 							FExecuteAction::CreateLambda(TryFindAndReplaceAnimNextRigVMAsset));
+
+						if (bEnableUAFBrowser)
+						{
+							auto OpenUAFBrowser = [WeakWorkspaceEditor]()
+							{
+								if (TSharedPtr<Workspace::IWorkspaceEditor> WorkspaceEditorPinned = WeakWorkspaceEditor.Pin())
+								{
+									if (const UUAFEditorUserSettings* UAFEditorUserSettings = GetDefault<UUAFEditorUserSettings>())
+									{
+										if (UAFEditorUserSettings->bSummonUAFBrowserHotkeyOpensSidebar)
+										{
+											WorkspaceEditorPinned->GetTabManager()->TryToggleSidebarTab(FUAFBrowserTabSummoner::SidebarTabID);
+										}
+										else
+										{
+											FUAFBrowserTabSummoner::ToggleDrawer();
+										}
+									}
+								}
+							};
+
+							CommandList->MapAction(Commands.OpenUAFBrowser,
+								FExecuteAction::CreateLambda(OpenUAFBrowser));
+						}
+
 					}
 				}
 			}
diff --git a/Engine/Plugins/Experimental/UAF/UAF/Source/UAFEditor/Private/Common/AnimNextRigVMAssetCommands.cpp b/Engine/Plugins/Experimental/UAF/UAF/Source/UAFEditor/Private/Common/AnimNextRigVMAssetCommands.cpp
--- a/Engine/Plugins/Experimental/UAF/UAF/Source/UAFEditor/Private/Common/AnimNextRigVMAssetCommands.cpp
+++ b/Engine/Plugins/Experimental/UAF/UAF/Source/UAFEditor/Private/Common/AnimNextRigVMAssetCommands.cpp
@@ -14,8 +14,15 @@ FAnimNextRigVMAssetCommands::FAnimNextRigVMAssetCommands()
 
 void FAnimNextRigVMAssetCommands::RegisterCommands()
 {
-	UI_COMMAND(FindInAnimNextRigVMAsset, "FindInUAFRigVMAsset", "Search the current UAF Asset.", EUserInterfaceActionType::Button, FInputChord(EModifierKey::Control, EKeys::F));
-	UI_COMMAND(FindAndReplaceInAnimNextRigVMAsset, "FindAndReplaceInUAFRigVMAsset", "Find and Replace across UAF Assets.", EUserInterfaceActionType::Button, FInputChord(EModifierKey::Control | EModifierKey::Shift, EKeys::F));
+	// On mac command and ctrl are automatically swapped. Command + Space is spotlight search so we use ctrl+space on mac to avoid the conflict
+	EModifierKey::Type PlatformSafeCtrl = EModifierKey::Control;
+#if PLATFORM_MAC
+	PlatformSafeCtrl = EModifierKey::Command
+#endif // PLATFORM_MAC
+
+	UI_COMMAND(FindInAnimNextRigVMAsset, "FindInUAFRigVMAsset", "Search the current UAF Asset.", EUserInterfaceActionType::Button, FInputChord(PlatformSafeCtrl, EKeys::F));
+	UI_COMMAND(FindAndReplaceInAnimNextRigVMAsset, "FindAndReplaceInUAFRigVMAsset", "Find and Replace across UAF Assets.", EUserInterfaceActionType::Button, FInputChord(PlatformSafeCtrl | EModifierKey::Shift, EKeys::F));
+	UI_COMMAND(OpenUAFBrowser, "OpenUAFBrowser", "Opens the UAF Browser.", EUserInterfaceActionType::Button, FInputChord(PlatformSafeCtrl | EModifierKey::Shift, EKeys::SpaceBar));
 }
 
 }
diff --git a/Engine/Plugins/Experimental/UAF/UAF/Source/UAFEditor/Private/Common/AnimNextRigVMAssetCommands.h b/Engine/Plugins/Experimental/UAF/UAF/Source/UAFEditor/Private/Common/AnimNextRigVMAssetCommands.h
--- a/Engine/Plugins/Experimental/UAF/UAF/Source/UAFEditor/Private/Common/AnimNextRigVMAssetCommands.h
+++ b/Engine/Plugins/Experimental/UAF/UAF/Source/UAFEditor/Private/Common/AnimNextRigVMAssetCommands.h
@@ -7,7 +7,7 @@
 namespace UE::UAF
 {
 
-// The Commands for the variable overrides menu
+// Shared commands valid for any UAF RigVM Asset.
 class FAnimNextRigVMAssetCommands : public TCommands<FAnimNextRigVMAssetCommands>
 {
 public:
@@ -18,6 +18,8 @@ public:
 	TSharedPtr<FUICommandInfo> FindInAnimNextRigVMAsset;
 
 	TSharedPtr<FUICommandInfo> FindAndReplaceInAnimNextRigVMAsset;
+
+	TSharedPtr<FUICommandInfo> OpenUAFBrowser;
 };
 
 }
\ No newline at end of file
diff --git a/Engine/Plugins/Experimental/UAF/UAF/Source/UAFEditor/Private/Common/SUAFBrowser.cpp b/Engine/Plugins/Experimental/UAF/UAF/Source/UAFEditor/Private/Common/SUAFBrowser.cpp
new file mode 100644
--- /dev/null
+++ b/Engine/Plugins/Experimental/UAF/UAF/Source/UAFEditor/Private/Common/SUAFBrowser.cpp
@@ -0,0 +1,254 @@
+// Copyright Epic Games, Inc. All Rights Reserved.
+
+#include "SUAFBrowser.h"
+
+#include "AnimNextRigVMAsset.h"
+#include "AnimNextRigVMAssetEditorData.h"
+#include "Assets/TaggedAssetBrowserConfiguration.h"
+#include "Common/UAFBrowserTabSummoner.h"
+#include "ContentBrowserModule.h"
+#include "IContentBrowserSingleton.h"
+#include "IWorkspaceEditor.h"
+#include "IWorkspaceEditorModule.h"
+#include "Sound/SoundWave.h"
+#include "Settings/UAFEditorProjectSettings.h"
+#include "Settings/UAFEditorUserSettings.h"
+#include "StatusBarSubsystem.h"
+#include "UncookedOnlyUtils.h"
+#include "Widgets/Layout/SBorder.h"
+#include "Widgets/Layout/SBox.h"
+#include "Widgets/STaggedAssetBrowser.h"
+
+#include "Animation/AnimSequence.h"
+#include "Animation/BlendSpace.h"
+
+#define LOCTEXT_NAMESPACE "SUAFBrowser"
+
+namespace UE::UAF::Editor
+{
+
+void SUAFBrowser::Construct(const FArguments& InArgs, const TSharedPtr<UE::Workspace::IWorkspaceEditor>& InHostingApp, ETabState::Type InTabState)
+{
+	WeakHostingApp = InHostingApp;
+	TabState = InTabState;
+
+	// Note: The layout & implementation of this widget may change drastically in future revisions.
+	// Please consider widget implementation details as placeholder for now.
+
+	// @TODO: Configure filter / content browser, will require updating content browser to add new extension points for UAF Browser view
+	FContentBrowserModule& ContentBrowserModule = FModuleManager::Get().LoadModuleChecked<FContentBrowserModule>(TEXT("ContentBrowser"));
+
+	bool bIsOrientedHorizontally = TabState != ETabState::SidebarTab;
+
+	EOrientation Orientation = bIsOrientedHorizontally ? EOrientation::Orient_Horizontal : EOrientation::Orient_Vertical;
+
+	TWeakPtr<SWidget> ThisWeak = AsShared();
+	auto CreateToggleSwapHotkeyButton = [ThisWeak]() -> TSharedRef<SWidget>
+	{
+		if (TSharedPtr<SUAFBrowser> ThisPinned = StaticCastSharedPtr<SUAFBrowser>(ThisWeak.Pin()))
+		{	
+			if (ThisPinned->TabState == ETabState::SidebarTab || ThisPinned->TabState == ETabState::StatusbarTab)
+			{
+				auto GetSwapBrowserIcon = []() -> const FSlateBrush*
+				{
+					if (UUAFEditorUserSettings* UAFEditorUserSettings = GetMutableDefault<UUAFEditorUserSettings>())
+					{
+						return FAppStyle::Get().GetBrush(UAFEditorUserSettings->bSummonUAFBrowserHotkeyOpensSidebar
+							? "EditorViewport.Right"
+							: "EditorViewport.Bottom");
+					}
+
+					return nullptr;
+				};
+
+				return
+					SNew(SButton)
+					.ButtonStyle(FAppStyle::Get(), "SimpleButton")
+					.ToolTipText(LOCTEXT("SwapHotkey_Tooltip", "Swaps UAF Browser Hotkey Between Sidebar & StatusBar."))
+					.ContentPadding(FMargin(2.f, 4.0f, 2.0f, 4.0f))
+					.OnClicked(ThisPinned.Get(), &SUAFBrowser::SwapHotkeyButtonClicked)
+					[
+						SNew(SImage)
+						.ColorAndOpacity(FSlateColor::UseForeground())
+						.Image(TAttribute<const FSlateBrush*>::CreateLambda(GetSwapBrowserIcon))
+					];
+			}
+		}
+		return SNullWidget::NullWidget;
+	};
+	auto CreateDrawerDockButton = [ThisWeak]() -> TSharedRef<SWidget>
+	{
+		if (TSharedPtr<SUAFBrowser> ThisPinned = StaticCastSharedPtr<SUAFBrowser>(ThisWeak.Pin()))
+		{
+			if (ThisPinned->TabState == ETabState::SidebarTab || ThisPinned->TabState == ETabState::StatusbarTab)
+			{
+				return
+					SNew(SButton)
+					.ButtonStyle(FAppStyle::Get(), "SimpleButton")
+					.ToolTipText(LOCTEXT("DockInLayout_Tooltip", "Docks UAF Browser in Tab."))
+					.ContentPadding(FMargin(2.f, 4.0f, 2.0f, 4.0f))
+					.OnClicked(ThisPinned.Get(), &SUAFBrowser::DockDrawerButtonClicked)
+					[
+						SNew(SImage)
+						.ColorAndOpacity(FSlateColor::UseForeground())
+						.Image(FAppStyle::Get().GetBrush("Icons.Layout"))
+					];
+			}
+		}
+		return SNullWidget::NullWidget;
+	};
+	auto ExtendAssetPickerTopBar = [ThisWeak, CreateToggleSwapHotkeyButton, CreateDrawerDockButton](TSharedRef<SHorizontalBox> InAssetPickerTopBar)
+	{
+		if (TSharedPtr<SUAFBrowser> ThisPinned = StaticCastSharedPtr<SUAFBrowser>(ThisWeak.Pin()))
+		{
+			if (ThisPinned->TabState == ETabState::SidebarTab || ThisPinned->TabState == ETabState::StatusbarTab)
+			{
+				// Note: Value from `FAssetPickerConfig::OnExtendAssetPickerTopBar`. 
+				// We could make it expose those indicies, but given expected usage that feels over-engineered for now.
+				const int ShowDeveloperContentIndex = 2;
+
+				InAssetPickerTopBar->InsertSlot(ShowDeveloperContentIndex + 1)
+				.AutoWidth()
+				[
+					CreateDrawerDockButton()
+				];
+
+				InAssetPickerTopBar->InsertSlot(ShowDeveloperContentIndex + 2)
+				.AutoWidth()
+				[
+					CreateToggleSwapHotkeyButton()
+				];
+			}
+		}
+	};
+
+	auto CustomizeAssetPicker = [&](FAssetPickerConfig& Config)
+	{
+		Config.OnExtendAssetPickerTopBar.BindLambda(ExtendAssetPickerTopBar);
+		Config.InitialAssetViewType = EAssetViewType::Column;
+		Config.bAddFilterUI = true;
+		//Config.bShowPathInColumnView = true;
+		//Config.bSortByPathInColumnView = true;
+		Config.bAllowDragging = true;
+
+		Config.OnGetAssetContextMenu.Unbind();
+		Config.OnAssetSelected.Unbind();
+		Config.OnAssetsActivated.Unbind();
+		Config.OnGetCustomAssetToolTip.Unbind();
+	};
+		
+	STaggedAssetBrowser::FInterfaceOverrideProfiles OverrideProfiles;
+	OverrideProfiles.AssetViewOptionsProfileName = FName("TaggedAssetBrowser");
+	OverrideProfiles.DefaultFilterMenuExpansion = EAssetTypeCategories::Animation;
+	OverrideProfiles.FilterBarSaveName = FName("TaggedAssetBrowser.UAFBrowser");
+
+	FDefaultDetailsTabConfiguration DefaultDetailsTabConfiguration;
+	DefaultDetailsTabConfiguration.bUseDefaultDetailsTab = false;
+
+	auto GetAssetPreviewWidget = [&]()
+	{
+		return SNew(SVerticalBox)
+			+SVerticalBox::Slot()
+			.HAlign(HAlign_Center)
+			.VAlign(VAlign_Center)
+			.FillHeight(1.0f)
+			[
+				SNew(STextBlock)
+				.Justification(ETextJustify::Center)
+				.AutoWrapText(true)
+				.Text(LOCTEXT("TempAssetPreviewText", "Asset Preview Here"))
+			];
+	};
+	auto GetAssetBrowser = [&]() -> TSharedRef<SWidget>
+	{
+		FSoftObjectPath CategoryConfigurationPath = GetDefault<UUAFEditorProjectSettings>()->UAFBrowserCategoryConfiguration;
+		UObject* CategoryConfigurationObject = CategoryConfigurationPath.TryLoad();
+		if (UTaggedAssetBrowserConfiguration* CategoryConfiguration = Cast<UTaggedAssetBrowserConfiguration>(CategoryConfigurationObject))
+		{
+			// We intentionally do not restrict AvailableClasses as this will filter types we may not know about from beind displayed
+			// Ex: Even in contexts such as filtering by tag
+			return SNew(STaggedAssetBrowser, *CategoryConfiguration)
+				.AvailableClasses({})
+				.DefaultDetailsTabConfiguration(DefaultDetailsTabConfiguration)
+				.InterfaceOverrideProfiles(OverrideProfiles)
+				.OnOverrideAssetPickerConfig(STaggedAssetBrowser::FOnOverrideAssetPickerConfig::CreateLambda(CustomizeAssetPicker));
+		}
+
+		return SNullWidget::NullWidget;
+	};
+	auto GetMainBrowserContents = [&]()
+	{
+		return SNew(SSplitter)
+			.Orientation(Orientation)
+			+SSplitter::Slot()
+			.Value(bIsOrientedHorizontally ? 0.85f : 0.15f)
+			[
+				bIsOrientedHorizontally
+					? GetAssetBrowser()
+					: GetAssetPreviewWidget()
+			]
+			+SSplitter::Slot()
+			.Value(bIsOrientedHorizontally ? 0.15f : 0.85f)
+			[
+				bIsOrientedHorizontally
+					? GetAssetPreviewWidget()
+					: GetAssetBrowser()
+			];
+	};
+	
+	FMargin DrawerPadding = TabState == ETabState::SidebarTab || TabState == ETabState::StatusbarTab
+		? 4.0f
+		: 0.0f;
+
+	ChildSlot
+	[
+		SNew(SBorder)
+		.BorderImage(FAppStyle::GetBrush("ToolPanel.GroupBorder"))
+		.Padding(DrawerPadding)
+		[
+			SNew(SOverlay)
+			+SOverlay::Slot()
+			.HAlign(HAlign_Fill)
+			.VAlign(VAlign_Fill)
+			.Padding(DrawerPadding)
+			[
+				GetMainBrowserContents()
+			]
+		]
+	];
+}
+
+FReply SUAFBrowser::SwapHotkeyButtonClicked()
+{
+	if (TSharedPtr<UE::Workspace::IWorkspaceEditor> HostingApp = WeakHostingApp.Pin())
+	{
+		GEditor->GetEditorSubsystem<UStatusBarSubsystem>()->ForceDismissDrawer();
+		HostingApp->GetTabManager()->TryDismissSidebarTab(FUAFBrowserTabSummoner::SidebarTabID);
+
+		if (UUAFEditorUserSettings* UAFEditorUserSettings = GetMutableDefault<UUAFEditorUserSettings>())
+		{
+			UAFEditorUserSettings->bSummonUAFBrowserHotkeyOpensSidebar = !UAFEditorUserSettings->bSummonUAFBrowserHotkeyOpensSidebar;
+		}
+	}
+	return FReply::Handled();
+}
+
+FReply SUAFBrowser::DockDrawerButtonClicked()
+{
+	if (TSharedPtr<UE::Workspace::IWorkspaceEditor> HostingApp = WeakHostingApp.Pin())
+	{
+		GEditor->GetEditorSubsystem<UStatusBarSubsystem>()->ForceDismissDrawer();
+		HostingApp->GetTabManager()->TryDismissSidebarTab(FUAFBrowserTabSummoner::SidebarTabID);
+
+		if (TSharedPtr<SDockTab> ExistingTab = HostingApp->GetTabManager()->TryInvokeTab(FUAFBrowserTabSummoner::DockedTabID))
+		{
+			ExistingTab->ActivateInParent(ETabActivationCause::SetDirectly);
+		}
+	}
+	return FReply::Handled();
+}
+
+} // namespace UE::UAF::Editor
+
+#undef LOCTEXT_NAMESPACE
+
diff --git a/Engine/Plugins/Experimental/UAF/UAF/Source/UAFEditor/Private/Common/SUAFBrowser.h b/Engine/Plugins/Experimental/UAF/UAF/Source/UAFEditor/Private/Common/SUAFBrowser.h
new file mode 100644
--- /dev/null
+++ b/Engine/Plugins/Experimental/UAF/UAF/Source/UAFEditor/Private/Common/SUAFBrowser.h
@@ -0,0 +1,55 @@
+// Copyright Epic Games, Inc. All Rights Reserved.
+
+#pragma once
+
+#include "CoreMinimal.h"
+#include "Framework/Docking/TabManager.h"
+#include "Widgets/Docking/SDockTab.h"
+#include "Widgets/SCompoundWidget.h"
+
+namespace UE::Workspace
+{
+	class IWorkspaceEditor;
+}	// end namespace UE::Workspace
+
+namespace UE::UAF::Editor
+{
+
+/** 
+ * Experimental. Widget for searching for UAF related assets
+ * 
+ * Note: The layout & implementation of this widget may change drastically in future revisions.
+ * Please consider widget implementation details as placeholder for now.
+ */
+class SUAFBrowser : public SCompoundWidget
+{
+
+public:
+
+	SLATE_BEGIN_ARGS(SUAFBrowser)
+	{}
+
+	SLATE_END_ARGS()
+
+public:
+
+	void Construct(const FArguments& InArgs, const TSharedPtr<UE::Workspace::IWorkspaceEditor>& InHostingApp, ETabState::Type InTabState);
+
+private:
+
+	/** Callback for when swap hotkey is pressed. Chnages hotkey behavior in project settings & hides all drawers. */
+	FReply SwapHotkeyButtonClicked();
+
+	/** Callback for when dock to drawer button is pressed. Hides all drawers and invokes docked browser tab. */
+	FReply DockDrawerButtonClicked();
+
+private:
+
+	/** App Hosting this Browser, to query settings against */
+	TWeakPtr<UE::Workspace::IWorkspaceEditor> WeakHostingApp = nullptr;
+
+	/** Current Tab host type, so we can resturucture layout based on how this tab is being hosted (Sidebar / statusbar / etc). */
+	ETabState::Type TabState = ETabState::ClosedTab;
+};
+
+} // namespace UE::UAF::Editor
\ No newline at end of file
diff --git a/Engine/Plugins/Experimental/UAF/UAF/Source/UAFEditor/Private/Common/UAFBrowserTabSummoner.cpp b/Engine/Plugins/Experimental/UAF/UAF/Source/UAFEditor/Private/Common/UAFBrowserTabSummoner.cpp
new file mode 100644
--- /dev/null
+++ b/Engine/Plugins/Experimental/UAF/UAF/Source/UAFEditor/Private/Common/UAFBrowserTabSummoner.cpp
@@ -0,0 +1,55 @@
+// Copyright Epic Games, Inc. All Rights Reserved.
+
+#include "Common/UAFBrowserTabSummoner.h"
+
+#include "AnimNextRigVMAsset.h"
+#include "Common/SUAFBrowser.h"
+#include "IAnimNextEditorModule.h"
+#include "IWorkspaceEditor.h"
+#include "Modules/ModuleManager.h"
+#include "PersonaModule.h"
+#include "StatusBarSubsystem.h"
+
+#define LOCTEXT_NAMESPACE "WorkspaceTabSummoner"
+
+namespace UE::UAF::Editor
+{
+
+const FName FUAFBrowserTabSummoner::DockedTabID(TEXT("UAFBrowserDockedTab"));
+const FName FUAFBrowserTabSummoner::SidebarTabID(TEXT("UAFBrowserSidebarTab"));
+const FName FUAFBrowserTabSummoner::DrawerTabID(TEXT("UAFBrowserDrawerTab"));
+const FName FUAFBrowserTabSummoner::StatusBarSystemDrawerID(TEXT("UAFBrowserDrawer"));
+
+void FUAFBrowserTabSummoner::ToggleDrawer()
+{
+	GEditor->GetEditorSubsystem<UStatusBarSubsystem>()->TryToggleDrawer(FUAFBrowserTabSummoner::StatusBarSystemDrawerID);
+}
+
+FUAFBrowserTabSummoner::FUAFBrowserTabSummoner(TSharedPtr<UE::Workspace::IWorkspaceEditor> InHostingApp, ETabState::Type InTabState, FName InTabName, FText InTabLabel)
+	: FWorkflowTabFactory(InTabName, StaticCastSharedPtr<FAssetEditorToolkit>(InHostingApp))
+	, TabState(InTabState)
+{
+	TabLabel = InTabLabel;
+	ViewMenuDescription = InTabLabel;
+
+	// @TODO: DarenC - Placeholder Icon
+	TabIcon = FSlateIcon(FAppStyle::GetAppStyleSetName(), "Kismet.Tabs.Palette");
+	ViewMenuTooltip = LOCTEXT("UAFBrowserTabToolTip", "Browse / Create Animation related assets.");
+	bIsSingleton = true;
+
+	UAFBrowserWidget = SNew(SUAFBrowser, InHostingApp, TabState);
+}
+
+TSharedRef<SWidget> FUAFBrowserTabSummoner::CreateTabBody(const FWorkflowTabSpawnInfo& Info) const
+{
+	return UAFBrowserWidget.ToSharedRef();
+}
+
+FText FUAFBrowserTabSummoner::GetTabToolTipText(const FWorkflowTabSpawnInfo& Info) const
+{
+	return ViewMenuTooltip;
+}
+
+} // end namespace UE::UAF::Editor
+
+#undef LOCTEXT_NAMESPACE
diff --git a/Engine/Plugins/Experimental/UAF/UAF/Source/UAFEditor/Private/Common/UAFBrowserTabSummoner.h b/Engine/Plugins/Experimental/UAF/UAF/Source/UAFEditor/Private/Common/UAFBrowserTabSummoner.h
new file mode 100644
--- /dev/null
+++ b/Engine/Plugins/Experimental/UAF/UAF/Source/UAFEditor/Private/Common/UAFBrowserTabSummoner.h
@@ -0,0 +1,58 @@
+// Copyright Epic Games, Inc. All Rights Reserved.
+
+#pragma once
+
+#include "Framework/Docking/TabManager.h"
+#include "Widgets/Docking/SDockTab.h"
+#include "WorkflowOrientedApp/WorkflowTabFactory.h"
+#include "IMessageLogListing.h"
+
+namespace UE::Workspace
+{
+	class IWorkspaceEditor;
+}	// end namespace UE::Workspace
+
+namespace UE::UAF::Editor
+{
+	class SUAFBrowser;
+}	// end namespace UE::UAF::Editor
+
+namespace UE::UAF::Editor
+{
+
+struct FUAFBrowserTabSummoner : public FWorkflowTabFactory
+{
+public:
+
+	/** Docked Tab ID. We have 3 distinct, non synchronized tabs. The UAF editor will manage hotkey / layouts for the user. */
+	static const FName DockedTabID;
+
+	/** Sidebar Tab ID. We have 3 distinct, non synchronized tabs. The UAF editor will manage hotkey / layouts for the user. */
+	static const FName SidebarTabID;
+
+	/** Drawer Tab ID. We have 3 distinct, non synchronized tabs. The UAF editor will manage hotkey / layouts for the user. */
+	static const FName DrawerTabID;
+
+	/** ID for this tab within the status bar system. Used to communicate with status bar system. */
+	static const FName StatusBarSystemDrawerID;
+
+	static void ToggleDrawer();
+
+public:
+
+	FUAFBrowserTabSummoner(TSharedPtr<UE::Workspace::IWorkspaceEditor> InHostingApp, ETabState::Type InTabState, FName InTabName, FText InTabLabel);
+
+	// FWorkflowTabFactory interface
+	virtual TSharedRef<SWidget> CreateTabBody(const FWorkflowTabSpawnInfo& Info) const override;
+	virtual FText GetTabToolTipText(const FWorkflowTabSpawnInfo& Info) const override;
+
+private:
+
+	/** The widget this tab spawner wraps */
+	TSharedPtr<SUAFBrowser> UAFBrowserWidget;
+
+	/** Current state of the tab, needed for the tab to change layout*/
+	ETabState::Type TabState = ETabState::ClosedTab;
+};
+
+} // end namespace UE::UAF::Editor
diff --git a/Engine/Plugins/Experimental/UAF/UAF/Source/UAFEditor/Private/Settings/UAFEditorProjectSettings.cpp b/Engine/Plugins/Experimental/UAF/UAF/Source/UAFEditor/Private/Settings/UAFEditorProjectSettings.cpp
new file mode 100644
--- /dev/null
+++ b/Engine/Plugins/Experimental/UAF/UAF/Source/UAFEditor/Private/Settings/UAFEditorProjectSettings.cpp
@@ -0,0 +1,9 @@
+// Copyright Epic Games, Inc. All Rights Reserved.
+
+#include "UAFEditorProjectSettings.h"
+
+#include UE_INLINE_GENERATED_CPP_BY_NAME(UAFEditorProjectSettings)
+
+UUAFEditorProjectSettings::UUAFEditorProjectSettings()
+{
+}
diff --git a/Engine/Plugins/Experimental/UAF/UAF/Source/UAFEditor/Private/Settings/UAFEditorProjectSettings.h b/Engine/Plugins/Experimental/UAF/UAF/Source/UAFEditor/Private/Settings/UAFEditorProjectSettings.h
new file mode 100644
--- /dev/null
+++ b/Engine/Plugins/Experimental/UAF/UAF/Source/UAFEditor/Private/Settings/UAFEditorProjectSettings.h
@@ -0,0 +1,23 @@
+﻿// Copyright Epic Games, Inc. All Rights Reserved.
+
+#pragma once
+
+#include "Engine/DeveloperSettings.h"
+#include "UAFEditorProjectSettings.generated.h"
+
+/**
+ * UAF Editor Settings that are shared across a project.
+ */
+UCLASS(MinimalAPI, Config = Editor, defaultconfig)
+class UUAFEditorProjectSettings : public UDeveloperSettings
+{
+	GENERATED_BODY()
+
+	UUAFEditorProjectSettings();
+
+public:
+
+	/** The configuration used the UAF browser for setting up categories */
+	UPROPERTY(config, EditAnywhere, DisplayName = "UAF Browser Configuration", Category = UAF, meta = (AllowedClasses = "/Script/UserAssetTagsEditor.TaggedAssetBrowserConfiguration"))
+	FSoftObjectPath UAFBrowserCategoryConfiguration;
+};
\ No newline at end of file
diff --git a/Engine/Plugins/Experimental/UAF/UAF/Source/UAFEditor/Private/Settings/UAFEditorUserSettings.cpp b/Engine/Plugins/Experimental/UAF/UAF/Source/UAFEditor/Private/Settings/UAFEditorUserSettings.cpp
new file mode 100644
--- /dev/null
+++ b/Engine/Plugins/Experimental/UAF/UAF/Source/UAFEditor/Private/Settings/UAFEditorUserSettings.cpp
@@ -0,0 +1,9 @@
+// Copyright Epic Games, Inc. All Rights Reserved.
+
+#include "UAFEditorUserSettings.h"
+
+#include UE_INLINE_GENERATED_CPP_BY_NAME(UAFEditorUserSettings)
+
+UUAFEditorUserSettings::UUAFEditorUserSettings()
+{
+}
diff --git a/Engine/Plugins/Experimental/UAF/UAF/Source/UAFEditor/Private/Settings/UAFEditorUserSettings.h b/Engine/Plugins/Experimental/UAF/UAF/Source/UAFEditor/Private/Settings/UAFEditorUserSettings.h
new file mode 100644
--- /dev/null
+++ b/Engine/Plugins/Experimental/UAF/UAF/Source/UAFEditor/Private/Settings/UAFEditorUserSettings.h
@@ -0,0 +1,23 @@
+﻿// Copyright Epic Games, Inc. All Rights Reserved.
+
+#pragma once
+
+#include "Engine/DeveloperSettings.h"
+#include "UAFEditorUserSettings.generated.h"
+
+/**
+ * UAF Editor Settings that are unique per user.
+ */
+UCLASS(MinimalAPI, Config = EditorPerProjectUserSettings)
+class UUAFEditorUserSettings : public UDeveloperSettings
+{
+	GENERATED_BODY()
+
+	UUAFEditorUserSettings();
+
+public:
+
+	/** True if the UAF Browser hotkey opens the sidebar (Right), false will open status bar (Bottom) */
+	UPROPERTY(EditAnywhere, Category="UAF", Config)
+	bool bSummonUAFBrowserHotkeyOpensSidebar = true;
+};
\ No newline at end of file
diff --git a/Engine/Plugins/Experimental/UAF/UAF/Source/UAFEditor/Private/Workspace/AnimNextWorkspaceEditorMode.cpp b/Engine/Plugins/Experimental/UAF/UAF/Source/UAFEditor/Private/Workspace/AnimNextWorkspaceEditorMode.cpp
--- a/Engine/Plugins/Experimental/UAF/UAF/Source/UAFEditor/Private/Workspace/AnimNextWorkspaceEditorMode.cpp
+++ b/Engine/Plugins/Experimental/UAF/UAF/Source/UAFEditor/Private/Workspace/AnimNextWorkspaceEditorMode.cpp
@@ -7,8 +7,10 @@
 #include "AnimNextRigVMAssetEditorData.h"
 #include "AnimNextScopedCompileJob.h"
 #include "Component/AnimNextComponent.h"
+#include "Common/UAFBrowserTabSummoner.h"
 #include "ContextObjectStore.h"
 #include "EditorModeManager.h"
+#include "HAL/IConsoleManager.h"
 #include "IAssetCompilationHandler.h"
 #include "InteractiveToolManager.h"
 #include "IWorkspaceEditor.h"
@@ -25,6 +27,8 @@
 #include "UAF/Workspace/SDebugObjectSelector.h"
 #include "UObject/UObjectIterator.h"
 #include "UncookedOnlyUtils.h"
+#include "WidgetDrawerConfig.h"
+#include "Widgets/Docking/SDockTab.h"
 
 #include UE_INLINE_GENERATED_CPP_BY_NAME(AnimNextWorkspaceEditorMode)
 
@@ -67,6 +71,59 @@ void UAnimNextWorkspaceEditorMode::Enter()
 		WeakAssets.Add(Asset);
 	}
 
+	// Add UAF Browser Drawer
+	IConsoleVariable* CVarEnableUAFBrowser = IConsoleManager::Get().FindConsoleVariable(TEXT("a.UAF.EnableUAFBrowser"));
+	if (CVarEnableUAFBrowser && CVarEnableUAFBrowser->GetBool())
+	{
+		using namespace UE::UAF::Editor;
+
+		FWidgetDrawerConfig UAFDrawer(FUAFBrowserTabSummoner::StatusBarSystemDrawerID);
+		TWeakPtr<UE::Workspace::IWorkspaceEditor> WeakEditor = WorkspaceEditor;
+
+		if (!UAFBrowserStatusbarWidget)
+		{
+			FUAFBrowserTabSummoner UAFBrowserTabSummoner(WorkspaceEditor, ETabState::StatusbarTab, FUAFBrowserTabSummoner::DrawerTabID, LOCTEXT("UAFBrowserDrawerLabel", "UAF Browser (Drawer)"));
+			FWorkflowTabSpawnInfo SpawnInfo;
+			UAFBrowserStatusbarWidget = UAFBrowserTabSummoner.CreateTabBody(SpawnInfo);
+		}
+
+		TWeakPtr<SWidget> WeakUAFBrowserStatusbarWidget = UAFBrowserStatusbarWidget;
+		UAFDrawer.GetDrawerContentDelegate.BindLambda([WeakEditor, WeakUAFBrowserStatusbarWidget]()
+		{
+			if (TSharedPtr<UE::Workspace::IWorkspaceEditor> EditorPtr = WeakEditor.Pin())
+			{
+				if (TSharedPtr<SWidget> UAFBrowserStatusbarWidgetPinned = WeakUAFBrowserStatusbarWidget.Pin())
+				{
+					return UAFBrowserStatusbarWidgetPinned.ToSharedRef();
+				}
+			}
+
+			return SNullWidget::NullWidget;
+		});
+		UAFDrawer.OnDrawerOpenedDelegate.BindLambda([WeakEditor](FName StatusBarWithDrawerName)
+		{
+			// @TODO: DarenC - Find best focus targets, these have issues dismissing via hotkey, but without these we have issues dismissing via click
+			//if (TSharedPtr<UE::Workspace::IWorkspaceEditor> EditorPtr = WeakEditor.Pin())
+			//{
+			//	FSlateApplication::Get().SetUserFocus(FSlateApplication::Get().GetUserIndexForKeyboard(), EditorPtr->GetExternalEditorWidget(FUAFBrowserTabSummoner::StatusBarSystemDrawerID));
+			//}
+		});
+		UAFDrawer.OnDrawerDismissedDelegate.BindLambda([WeakEditor](const TSharedPtr<SWidget>& NewlyFocusedWidget)
+		{
+			// @TODO: DarenC - Find best focus targets, these have issues dismissing via hotkey, but without these we have issues dismissing via click
+			//if (TSharedPtr<UE::Workspace::IWorkspaceEditor> EditorPtr = WeakEditor.Pin())
+			//{
+			//	FSlateApplication::Get().SetUserFocus(FSlateApplication::Get().GetUserIndexForKeyboard(), EditorPtr->GetDetailsView());
+			//}
+		});
+		UAFDrawer.ButtonText = LOCTEXT("StatusBar_UAFBrowser", "UAF Browser");
+		UAFDrawer.ToolTipText = LOCTEXT("StatusBar_UAFBrowserToolTip", "Opens UAF Browser (Ctrl+Shift+Space).");
+				
+		// @TODO: DarenC - temp, see if we have real icon
+		UAFDrawer.Icon = FAppStyle::Get().GetBrush("ContentBrowser.TabIcon");
+		WorkspaceEditor->RegisterDrawer(MoveTemp(UAFDrawer), 1);
+	}
+
 	UpdateCompileStatus();
 	
 	RewindDebuggerExtension = MakeShared<UE::UAF::Editor::FRewindDebuggerWorkspaceExtension>();
diff --git a/Engine/Plugins/Experimental/UAF/UAF/Source/UAFEditor/Private/Workspace/AnimNextWorkspaceEditorMode.h b/Engine/Plugins/Experimental/UAF/UAF/Source/UAFEditor/Private/Workspace/AnimNextWorkspaceEditorMode.h
--- a/Engine/Plugins/Experimental/UAF/UAF/Source/UAFEditor/Private/Workspace/AnimNextWorkspaceEditorMode.h
+++ b/Engine/Plugins/Experimental/UAF/UAF/Source/UAFEditor/Private/Workspace/AnimNextWorkspaceEditorMode.h
@@ -106,4 +106,7 @@ private:
 	TSharedPtr<UE::UAF::Editor::FRewindDebuggerWorkspaceExtension> RewindDebuggerExtension;
 
 	TSharedPtr<UE::UAF::Editor::IDebugObjectSelector> DebugObjectSelector;
+
+	// UAF browser status bar widget, used to store a custom drawer next to the content browser
+	TSharedPtr<SWidget> UAFBrowserStatusbarWidget = nullptr;
 };
\ No newline at end of file
diff --git a/Engine/Plugins/Experimental/UAF/UAF/Source/UAFEditor/Public/IAnimNextEditorModule.h b/Engine/Plugins/Experimental/UAF/UAF/Source/UAFEditor/Public/IAnimNextEditorModule.h
--- a/Engine/Plugins/Experimental/UAF/UAF/Source/UAFEditor/Public/IAnimNextEditorModule.h
+++ b/Engine/Plugins/Experimental/UAF/UAF/Source/UAFEditor/Public/IAnimNextEditorModule.h
@@ -26,6 +26,7 @@ const FLazyName CompilerResultsTabName("CompilerResultsTab");
 const FLazyName LogListingName("AnimNextCompilerResults");
 const FLazyName FindTabName("FindTab");
 const FLazyName FindAndReplaceTabName("FindAndReplaceTab");
+const FLazyName UAFBrowserTabName("UAFBrowserTab");
 
 struct FVariablePickerArgs;
 struct FActionMenuContextData;
diff --git a/Engine/Plugins/Experimental/UAF/UAF/Source/UAFEditor/UAFEditor.Build.cs b/Engine/Plugins/Experimental/UAF/UAF/Source/UAFEditor/UAFEditor.Build.cs
--- a/Engine/Plugins/Experimental/UAF/UAF/Source/UAFEditor/UAFEditor.Build.cs
+++ b/Engine/Plugins/Experimental/UAF/UAF/Source/UAFEditor/UAFEditor.Build.cs
@@ -38,11 +38,14 @@ namespace UnrealBuildTool.Rules
 					"ContentBrowserFileDataSource",
 					"SubobjectEditor",
 					"Settings",
+					"StatusBar",
 					"EditorWidgets",
 					"WorkspaceEditor",
 					"ContentBrowser",
+					"ContentBrowserData",
 					"UniversalObjectLocator",
 					"UniversalObjectLocatorEditor",
+					"UserAssetTagsEditor",
 					"Kismet",
 					"AdvancedWidgets", 
 					"SceneOutliner",
diff --git a/Engine/Plugins/Experimental/Workspace/Source/WorkspaceEditor/Private/WorkspaceEditor.cpp b/Engine/Plugins/Experimental/Workspace/Source/WorkspaceEditor/Private/WorkspaceEditor.cpp
--- a/Engine/Plugins/Experimental/Workspace/Source/WorkspaceEditor/Private/WorkspaceEditor.cpp
+++ b/Engine/Plugins/Experimental/Workspace/Source/WorkspaceEditor/Private/WorkspaceEditor.cpp
@@ -72,6 +72,7 @@ namespace WorkspaceTabs
 	const FName BottomMiddleDocumentArea("BottomMiddleDocumentArea");
 	const FName TopRightDocumentArea("TopRightDocumentArea");
 	const FName BottomRightDocumentArea("BottomRightDocumentArea");
+	const FName RightSidebarDocumentArea("RightSidebarDocumentArea");
 }
 
 const FName WorkspaceAppIdentifier("WorkspaceEditor");
@@ -117,10 +118,14 @@ void FWorkspaceEditor::CreateWidgets()
 	BottomRightAssetDocumentSummoner->SetAllowedClassPaths(WorkspaceEditorModule.GetAllowedObjectTypesForArea(WorkspaceTabs::BottomRightDocumentArea));
 	DocumentManager->RegisterDocumentFactory(BottomRightAssetDocumentSummoner);
 
+	const TSharedRef<FAssetDocumentSummoner> RightSidebarAssetDocumentSummoner = MakeShared<FAssetDocumentSummoner>(WorkspaceTabs::RightSidebarDocumentArea, SharedThis(this));
+	RightSidebarAssetDocumentSummoner->SetAllowedClassPaths(WorkspaceEditorModule.GetAllowedObjectTypesForArea(WorkspaceTabs::RightSidebarDocumentArea));
+	DocumentManager->RegisterDocumentFactory(RightSidebarAssetDocumentSummoner);
+
 	check(DetailsView.IsValid());
 	WorkspaceEditorModule.ApplyWorkspaceDetailsCustomization(StaticCastWeakPtr<IWorkspaceEditor>(this->AsWeak()), DetailsView);
 
-	StandaloneDefaultLayout = FTabManager::NewLayout("Standalone_WorkspaceEditor_Layout_v3")
+	StandaloneDefaultLayout = FTabManager::NewLayout("Standalone_WorkspaceEditor_Layout_v4")
 	->AddArea
 	(
 		FTabManager::NewPrimaryArea()
@@ -198,6 +203,11 @@ void FWorkspaceEditor::CreateWidgets()
 					->SetExtensionId(WorkspaceTabs::BottomRightDocumentArea)
 				)
 			)
+			->Split
+			(
+				FTabManager::NewStack()
+				->AddTab(WorkspaceTabs::RightSidebarDocumentArea, ETabState::SidebarTab, ESidebarLocation::Right, /*SidebarSizeCoefficient*/ 0.2f)
+			)
 		)
 	);
 
@@ -1086,6 +1096,6 @@ UWorkspaceViewportSceneDescription* FWorkspaceEditor::GetSceneDescription() cons
 	return SceneDescription;
 }
 	
-}
+} // namespace UE::Workspace
 
 #undef LOCTEXT_NAMESPACE
diff --git a/Engine/Plugins/Experimental/Workspace/Source/WorkspaceEditor/Public/IWorkspaceEditorModule.h b/Engine/Plugins/Experimental/Workspace/Source/WorkspaceEditor/Public/IWorkspaceEditorModule.h
--- a/Engine/Plugins/Experimental/Workspace/Source/WorkspaceEditor/Public/IWorkspaceEditorModule.h
+++ b/Engine/Plugins/Experimental/Workspace/Source/WorkspaceEditor/Public/IWorkspaceEditorModule.h
@@ -46,6 +46,7 @@ namespace WorkspaceTabs
 	WORKSPACEEDITOR_API extern const FName BottomMiddleDocumentArea;
 	WORKSPACEEDITOR_API extern const FName TopRightDocumentArea;
 	WORKSPACEEDITOR_API extern const FName BottomRightDocumentArea;
+	WORKSPACEEDITOR_API extern const FName RightSidebarDocumentArea;
 }
 
 struct FWorkspaceDocument
