Add missing semicolon Mac builds.

diff --git a/Engine/Plugins/Experimental/UAF/UAF/Source/UAFEditor/Private/Common/AnimNextRigVMAssetCommands.cpp b/Engine/Plugins/Experimental/UAF/UAF/Source/UAFEditor/Private/Common/AnimNextRigVMAssetCommands.cpp
--- a/Engine/Plugins/Experimental/UAF/UAF/Source/UAFEditor/Private/Common/AnimNextRigVMAssetCommands.cpp
+++ b/Engine/Plugins/Experimental/UAF/UAF/Source/UAFEditor/Private/Common/AnimNextRigVMAssetCommands.cpp
@@ -17,7 +17,7 @@ void FAnimNextRigVMAssetCommands::RegisterCommands()
 	// On mac command and ctrl are automatically swapped. Command + Space is spotlight search so we use ctrl+space on mac to avoid the conflict
 	EModifierKey::Type PlatformSafeCtrl = EModifierKey::Control;
 #if PLATFORM_MAC
-	PlatformSafeCtrl = EModifierKey::Command
+	PlatformSafeCtrl = EModifierKey::Command;
 #endif // PLATFORM_MAC
 
 	UI_COMMAND(FindInAnimNextRigVMAsset, "FindInUAFRigVMAsset", "Search the current UAF Asset.", EUserInterfaceActionType::Button, FInputChord(PlatformSafeCtrl, EKeys::F));
