Add ability to specify iconography and color for OutlinerEntryItem(s)

diff --git a/Engine/Plugins/Experimental/UAF/UAF/Source/UAFEditor/Private/Common/Outliner/OutlinerEntryItem.h b/Engine/Plugins/Experimental/UAF/UAF/Source/UAFEditor/Private/Common/Outliner/OutlinerEntryItem.h
--- a/Engine/Plugins/Experimental/UAF/UAF/Source/UAFEditor/Private/Common/Outliner/OutlinerEntryItem.h
+++ b/Engine/Plugins/Experimental/UAF/UAF/Source/UAFEditor/Private/Common/Outliner/OutlinerEntryItem.h
@@ -19,12 +19,11 @@ struct FOutlinerEntryItem : FOutlinerItem
 	virtual TSharedRef<SWidget> GenerateLabelWidget(ISceneOutliner& Outliner, const STableRow<FSceneOutlinerTreeItemPtr>& InRow) override;
 
 	virtual bool IsReadOnly() const = 0;
-
 	virtual void SetAccessSpecifier(const EAnimNextExportAccessSpecifier& InSpecifier) const = 0;
 	virtual EAnimNextExportAccessSpecifier GetAccessSpecifier() const = 0;
 	virtual bool CanSetAccessSpecifier() const = 0;
-
 	virtual FStringView GetCategoryPath() const = 0;
+	virtual void GetItemIconAndColor(const FSlateBrush*& OutIcon, FSlateColor& OutColor) const {}
 };
 
 }
diff --git a/Engine/Plugins/Experimental/UAF/UAF/Source/UAFEditor/Private/Common/Outliner/OutlinerEntryItemWidget.cpp b/Engine/Plugins/Experimental/UAF/UAF/Source/UAFEditor/Private/Common/Outliner/OutlinerEntryItemWidget.cpp
--- a/Engine/Plugins/Experimental/UAF/UAF/Source/UAFEditor/Private/Common/Outliner/OutlinerEntryItemWidget.cpp
+++ b/Engine/Plugins/Experimental/UAF/UAF/Source/UAFEditor/Private/Common/Outliner/OutlinerEntryItemWidget.cpp
@@ -2,6 +2,8 @@
 
 #include "OutlinerEntryItemWidget.h"
 
+#include "SPinTypeSelector.h"
+
 #define LOCTEXT_NAMESPACE "SOutlinerEntryLabel"
 
 namespace UE::UAF::Editor
@@ -12,7 +14,43 @@ void SOutlinerEntryLabel::Construct(const FArguments& InArgs, FOutlinerEntryItem
 {
 	WeakSceneOutliner = StaticCastSharedRef<ISceneOutliner>(SceneOutliner.AsShared());
 	TreeItem = StaticCastSharedRef<FOutlinerEntryItem>(InTreeItem.AsShared());
+	
+	TSharedRef<SWidget> IconWidget = SNew(SImage)
+	.Image_Lambda([this]() -> const FSlateBrush*
+	{
+		const FSlateBrush* Brush = nullptr;
+		if (TreeItem.IsValid())
+		{
+			FSlateColor Color;
+			TreeItem.Pin()->GetItemIconAndColor(Brush, Color);
+		}
+
+		return Brush;
+	})
+	.ColorAndOpacity_Lambda([this]()
+	{
+		FSlateColor Color = GetForegroundColor();
+		if (TreeItem.IsValid())
+		{
+			const FSlateBrush* Brush = nullptr;
+			TreeItem.Pin()->GetItemIconAndColor(Brush, Color);
+		}
+
+		return Color;
+	})
+	.Visibility_Lambda([this]()
+	{
+		
+		const FSlateBrush* Brush = nullptr;
+		if (TreeItem.IsValid())
+		{
+			FSlateColor Color;
+			TreeItem.Pin()->GetItemIconAndColor(Brush, Color);
+		}
 
+		return Brush != nullptr ? EVisibility::Visible : EVisibility::Collapsed;
+	});
+	
 	ChildSlot
 	[
 		SNew(SHorizontalBox)
@@ -20,6 +58,13 @@ void SOutlinerEntryLabel::Construct(const FArguments& InArgs, FOutlinerEntryItem
 		.AutoWidth()
 		.VAlign(VAlign_Center)
 		.Padding(2.0f, 2.0f)
+		[
+			IconWidget
+		]
+		+SHorizontalBox::Slot()
+		.AutoWidth()
+		.VAlign(VAlign_Center)
+		.Padding(2.0f, 2.0f)
 		[
 			SAssignNew(TextBlock, SInlineEditableTextBlock)
 			.IsReadOnly(this, &SOutlinerEntryLabel::IsReadOnly)
