Fix initial function rename giving error message

diff --git a/Engine/Plugins/Experimental/UAF/UAF/Source/UAFEditor/Private/Functions/Outliner/FunctionOutlinerEntry.cpp b/Engine/Plugins/Experimental/UAF/UAF/Source/UAFEditor/Private/Functions/Outliner/FunctionOutlinerEntry.cpp
--- a/Engine/Plugins/Experimental/UAF/UAF/Source/UAFEditor/Private/Functions/Outliner/FunctionOutlinerEntry.cpp
+++ b/Engine/Plugins/Experimental/UAF/UAF/Source/UAFEditor/Private/Functions/Outliner/FunctionOutlinerEntry.cpp
@@ -105,7 +105,8 @@ bool FFunctionsOutlinerEntryItem::ValidateName(const FText& InNewName, FText& Ou
 		{
 			if (IRigVMClientHost* RigVMHost = FunctionNode->GetImplementingOuter<IRigVMClientHost>())
 			{
-				if (RigVMHost->GetLocalFunctionLibrary()->FindFunction(*InNewName.ToString()))
+				const URigVMLibraryNode* ExistingFunction = RigVMHost->GetLocalFunctionLibrary()->FindFunction(*InNewName.ToString());
+				if (ExistingFunction && ExistingFunction != FunctionNode)
 				{
 					OutErrorMessage = LOCTEXT("NameExistsError", "Function name already exists in this asset");
 					return false;
