Add tooltip text to Function output column widget

diff --git a/Engine/Plugins/Experimental/UAF/UAF/Source/UAFEditor/Private/Functions/Outliner/FunctionsOutlinerColumns.cpp b/Engine/Plugins/Experimental/UAF/UAF/Source/UAFEditor/Private/Functions/Outliner/FunctionsOutlinerColumns.cpp
--- a/Engine/Plugins/Experimental/UAF/UAF/Source/UAFEditor/Private/Functions/Outliner/FunctionsOutlinerColumns.cpp
+++ b/Engine/Plugins/Experimental/UAF/UAF/Source/UAFEditor/Private/Functions/Outliner/FunctionsOutlinerColumns.cpp
@@ -10,6 +10,7 @@
 #include "UncookedOnlyUtils.h"
 #include "Common/Outliner/OutlinerAssetItem.h"
 #include "RigVMModel/Nodes/RigVMLibraryNode.h"
+#include "Textures/SlateIcon.h"
 #include "Widgets/Text/STextBlock.h"
 
 #define LOCTEXT_NAMESPACE "FunctionsOutlinerColumns"
@@ -72,6 +73,35 @@ class SFunctionsOutlinerFunctionOutput : public SCompoundWidget
 		ChildSlot
 		[
 			SAssignNew(HBox, SHorizontalBox)
+			.ToolTipText_Lambda([this]()-> FText
+			{
+				const FText TooltipFormat(LOCTEXT("SFunctionsOutlinerFunctionOutput_Format", "{0}\n\nOutputs:{1}"));
+				
+				const FString OutputString = [this]() -> FString
+                {
+				    if (FunctionOutputs.Num() == 0)
+				    {
+					    return TEXT("none");
+				    }
+				    else
+				    {
+						FString TempString;
+					    for (const FRigVMGraphFunctionArgument* Output : FunctionOutputs)
+					    {
+						    if (TempString.Len() > 0)
+						    {
+							    TempString.Append(TEXT(", "));
+						    }
+						    
+						    TempString.Appendf(TEXT("%s(%s)"), *Output->DisplayName.ToString(), *Output->CPPType.ToString());
+					    }
+					    
+					    return TempString;
+				    }
+                }();
+				
+				return FText::Format(TooltipFormat, FText::FromName(FunctionHeader.Name), FText::FromString(OutputString));
+			})
 		];
 
 		if (FunctionOutputs.Num() == 0)
@@ -80,6 +110,7 @@ class SFunctionsOutlinerFunctionOutput : public SCompoundWidget
 			.AutoWidth()
 			.HAlign(HAlign_Left)
 			.VAlign(VAlign_Center)
+			.Padding(FMargin(2.0f, 3.0f, 2.0f, 3.0f))
 			[
 				SNew(STextBlock)
 				.Text_Lambda([]() -> FText
@@ -122,6 +153,33 @@ class SFunctionsOutlinerFunctionOutput : public SCompoundWidget
 				})
 			];
 		}
+		else if (FunctionOutputs.Num() >= 2)
+		{
+			HBox->AddSlot()
+			.AutoWidth()
+			.HAlign(HAlign_Center)
+			.VAlign(VAlign_Center)
+			.Padding(FMargin(2.0f, 3.0f, 2.0f, 3.0f))
+			[
+				SNew(SImage)
+				.Image(FAppStyle::Get().GetBrush(TEXT("Kismet.VariableList.TypeIcon")))
+				.ColorAndOpacity(GetForegroundColor())
+			];
+
+			HBox->AddSlot()
+			.AutoWidth()
+			.HAlign(HAlign_Left)
+			.VAlign(VAlign_Center)
+			.Padding(FMargin(4.f, 0.f))
+			[
+				SNew(STextBlock)
+				.Text_Lambda([]() -> FText
+				{
+					return LOCTEXT("MultipleValuesLabel", "multiple");
+				})
+				.ColorAndOpacity(FSlateColor::UseForeground())
+			];
+		}
 	}
 
 	TWeakPtr<FFunctionsOutlinerEntryItem> WeakTreeItem;
