Removed function workspace item details implementation
Removed function asset registry workspace exports

diff --git a/Engine/Plugins/Experimental/UAF/UAF/Source/UAFEditor/Private/AnimNextEditorModule.cpp b/Engine/Plugins/Experimental/UAF/UAF/Source/UAFEditor/Private/AnimNextEditorModule.cpp
--- a/Engine/Plugins/Experimental/UAF/UAF/Source/UAFEditor/Private/AnimNextEditorModule.cpp
+++ b/Engine/Plugins/Experimental/UAF/UAF/Source/UAFEditor/Private/AnimNextEditorModule.cpp
@@ -40,7 +40,6 @@
 #include "EditorModeManager.h"
 #include "Common/AnimNextGraphItemDetails.h"
 #include "Common/AnimNextCollapseNodeItemDetails.h"
-#include "Common/AnimNextFunctionItemDetails.h"
 #include "Common/AnimNextRigVMAssetCommands.h"
 #include "Param/AnimNextActorLocatorEditor.h"
 #include "RigVMCommands.h"
@@ -556,11 +555,9 @@ void FAnimNextEditorModule::StartupModule()
 	Workspace::IWorkspaceEditorModule& WorkspaceModule = FModuleManager::Get().LoadModuleChecked<Workspace::IWorkspaceEditorModule>("WorkspaceEditor");
 	WorkspaceModule.RegisterWorkspaceItemDetails(Workspace::FOutlinerItemDetailsId(FAnimNextGraphOutlinerData::StaticStruct()->GetFName()), MakeShared<FAnimNextGraphItemDetails>());
 	WorkspaceModule.RegisterWorkspaceItemDetails(Workspace::FOutlinerItemDetailsId(FAnimNextCollapseGraphOutlinerData::StaticStruct()->GetFName()), MakeShared<FAnimNextCollapseNodeItemDetails>());
-	WorkspaceModule.RegisterWorkspaceItemDetails(Workspace::FOutlinerItemDetailsId(FAnimNextGraphFunctionOutlinerData::StaticStruct()->GetFName()), MakeShared<FAnimNextFunctionItemDetails>());
 
 	FAnimNextGraphItemDetails::RegisterToolMenuExtensions();
 	FAnimNextCollapseNodeItemDetails::RegisterToolMenuExtensions();
-	FAnimNextFunctionItemDetails::RegisterToolMenuExtensions();
 
 	const TSharedPtr<FAnimNextAssetItemDetails> AssetItemDetails = MakeShared<FAnimNextAssetItemDetails>();
 	WorkspaceModule.RegisterWorkspaceItemDetails(Workspace::FOutlinerItemDetailsId(FAnimNextModuleOutlinerData::StaticStruct()->GetFName()), AssetItemDetails);
@@ -628,7 +625,6 @@ void FAnimNextEditorModule::ShutdownModule()
 		Workspace::IWorkspaceEditorModule& WorkspaceModule = FModuleManager::Get().LoadModuleChecked<Workspace::IWorkspaceEditorModule>("WorkspaceEditor");
 		WorkspaceModule.UnregisterWorkspaceItemDetails(Workspace::FOutlinerItemDetailsId(FAnimNextGraphOutlinerData::StaticStruct()->GetFName()));
 		WorkspaceModule.UnregisterWorkspaceItemDetails(Workspace::FOutlinerItemDetailsId(FAnimNextCollapseGraphOutlinerData::StaticStruct()->GetFName()));
-		WorkspaceModule.UnregisterWorkspaceItemDetails(Workspace::FOutlinerItemDetailsId(FAnimNextGraphFunctionOutlinerData::StaticStruct()->GetFName()));
 		FAnimNextGraphItemDetails::UnregisterToolMenuExtensions();
 		WorkspaceModule.UnregisterWorkspaceItemDetails(Workspace::FOutlinerItemDetailsId(FAnimNextModuleOutlinerData::StaticStruct()->GetFName()));
 		FAnimNextAssetItemDetails::UnregisterToolMenuExtensions();
diff --git a/Engine/Plugins/Experimental/UAF/UAF/Source/UAFEditor/Private/Common/AnimNextFunctionItemDetails.cpp b/Engine/Plugins/Experimental/UAF/UAF/Source/UAFEditor/Private/Common/AnimNextFunctionItemDetails.cpp
deleted file mode 100644
diff --git a/Engine/Plugins/Experimental/UAF/UAF/Source/UAFEditor/Private/Common/AnimNextFunctionItemDetails.h b/Engine/Plugins/Experimental/UAF/UAF/Source/UAFEditor/Private/Common/AnimNextFunctionItemDetails.h
deleted file mode 100644
diff --git a/Engine/Plugins/Experimental/UAF/UAF/Source/UAFUncookedOnly/Internal/UncookedOnlyUtils.h b/Engine/Plugins/Experimental/UAF/UAF/Source/UAFUncookedOnly/Internal/UncookedOnlyUtils.h
--- a/Engine/Plugins/Experimental/UAF/UAF/Source/UAFUncookedOnly/Internal/UncookedOnlyUtils.h
+++ b/Engine/Plugins/Experimental/UAF/UAF/Source/UAFUncookedOnly/Internal/UncookedOnlyUtils.h
@@ -236,9 +236,7 @@ private:
 	static UE_API void CompileVariableBindingsInternal(const FRigVMCompileSettings& InSettings, UAnimNextRigVMAsset* InAsset, TArray<URigVMGraph*>& OutGraphs, bool bInThreadSafe);
 	
 	static void GetSubGraphWorkspaceExportsRecursive(const UAnimNextRigVMAssetEditorData* EditorData, FWorkspaceOutlinerItemExports& OutExports, const FWorkspaceOutlinerItemExport& RootExport, int32 ParentExportIndex, URigVMEdGraph* RigVMEdGraph, FAssetRegistryTagsContext Context);
-	static void GetFunctionLibraryWorkspaceExportsRecursive(const UAnimNextRigVMAssetEditorData* EditorData, FWorkspaceOutlinerItemExports& OutExports, const FWorkspaceOutlinerItemExport& RootExport, int32 ParentExportIndex, const TArray<FRigVMGraphFunctionData>& PublicFunctions, const TArray<FRigVMGraphFunctionData>& PrivateFunctions);
-	static void GetFunctionsWorkspaceExportsRecursive(const UAnimNextRigVMAssetEditorData* EditorData, FWorkspaceOutlinerItemExports& OutExports, const FWorkspaceOutlinerItemExport& RootExport, int32 ParentExportIndex, const TArray<FRigVMGraphFunctionData>& Functions, bool bPublicFunctions);
-	
+
 	// Gets the exported variables that are used by a RigVM asset	
 	static void GetAssetVariableExports(const UAnimNextRigVMAssetEditorData* EditorData, TSet<FAnimNextExport>& OutExports, FAssetRegistryTagsContext Context, EAnimNextExportedVariableFlags InFlags = EAnimNextExportedVariableFlags::NoFlags);
 
diff --git a/Engine/Plugins/Experimental/UAF/UAF/Source/UAFUncookedOnly/Private/UncookedOnlyUtils.cpp b/Engine/Plugins/Experimental/UAF/UAF/Source/UAFUncookedOnly/Private/UncookedOnlyUtils.cpp
--- a/Engine/Plugins/Experimental/UAF/UAF/Source/UAFUncookedOnly/Private/UncookedOnlyUtils.cpp
+++ b/Engine/Plugins/Experimental/UAF/UAF/Source/UAFUncookedOnly/Private/UncookedOnlyUtils.cpp
@@ -1229,8 +1229,6 @@ void FUtils::GetAssetWorkspaceExports(const UAnimNextRigVMAssetEditorData* Edito
 		}
 		return true;
 	});	
-	
-	GetFunctionLibraryWorkspaceExportsRecursive(EditorData, OutExports, AssetIdentifier, INDEX_NONE, EditorData->GetRigVMGraphFunctionStore()->PublicFunctions, EditorData->GetRigVMGraphFunctionStore()->PrivateFunctions);
 }
 
 void FUtils::ProcessPinAssetReferences(const URigVMPin* InPin, FWorkspaceOutlinerItemExports& OutExports, const FWorkspaceOutlinerItemExport& RootExport, int32 ParentExportIndex)
@@ -1436,38 +1434,6 @@ void FUtils::GetSubGraphWorkspaceExportsRecursive(const UAnimNextRigVMAssetEdito
 	}
 }
 
-void FUtils::GetFunctionLibraryWorkspaceExportsRecursive(const UAnimNextRigVMAssetEditorData* EditorData, FWorkspaceOutlinerItemExports& OutExports, const FWorkspaceOutlinerItemExport& RootExport, int32 ParentExportIndex, const TArray<FRigVMGraphFunctionData>& PublicFunctions, const TArray<FRigVMGraphFunctionData>& PrivateFunctions)
-{
-	if (PrivateFunctions.Num() > 0 || PublicFunctions.Num() > 0)
-	{
-		const FWorkspaceOutlinerItemExport& ParentExport = ParentExportIndex != INDEX_NONE ? OutExports.Exports[ParentExportIndex] : RootExport;
-		FWorkspaceOutlinerItemExport& GroupExport = OutExports.Exports.Add_GetRef(FWorkspaceOutlinerItemExport(FName("FunctionsGroup"), ParentExport));
-		GroupExport.GetData().InitializeAs<FWorkspaceOutlinerGroupItemData>();
-		GroupExport.GetData().GetMutable<FWorkspaceOutlinerGroupItemData>().GroupName = TEXT("Functions");
-		GroupExport.GetData().GetMutable<FWorkspaceOutlinerGroupItemData>().GroupIcon = *FAppStyle::GetBrush(TEXT("GraphEditor.Function_24x"));
-		const int32 GroupExportIndex = OutExports.Exports.Num() - 1;
-
-		GetFunctionsWorkspaceExportsRecursive(EditorData, OutExports, GroupExport, GroupExportIndex, PrivateFunctions, false);
-		GetFunctionsWorkspaceExportsRecursive(EditorData, OutExports, GroupExport, GroupExportIndex, PublicFunctions, true);
-	}
-}
-
-void FUtils::GetFunctionsWorkspaceExportsRecursive(const UAnimNextRigVMAssetEditorData* EditorData, FWorkspaceOutlinerItemExports& OutExports, const FWorkspaceOutlinerItemExport& RootExport, int32 ParentExportIndex, const TArray<FRigVMGraphFunctionData>& Functions, bool bPublicFunctions)
-{
-	for (const FRigVMGraphFunctionData& FunctionData : Functions)
-	{
-		const FWorkspaceOutlinerItemExport& ParentExport = ParentExportIndex == INDEX_NONE ? RootExport : OutExports.Exports[ParentExportIndex];
-		if (const URigVMLibraryNode* FunctionNode = EditorData->RigVMClient.GetFunctionLibrary()->FindFunction(FunctionData.Header.LibraryPointer.GetFunctionFName()))
-		{
-			const FWorkspaceOutlinerItemExport Export = MakeFunctionExport(FunctionNode, ParentExport);
-			if (Export.GetIdentifier() != NAME_None)
-			{
-				OutExports.Exports.Add(Export);
-			}
-		}
-	}
-}
-
 void FUtils::GetSubGraphVariableExportsRecursive(const URigVMEdGraph* RigVMEdGraph, TSet<FAnimNextExport>& OutExports)
 {
 	TArray<URigVMEdGraphNode*> EdNodes;
