Update functions view on function category change

diff --git a/Engine/Plugins/Experimental/UAF/UAF/Source/UAFEditor/Private/Functions/Outliner/FunctionsOutlinerView.cpp b/Engine/Plugins/Experimental/UAF/UAF/Source/UAFEditor/Private/Functions/Outliner/FunctionsOutlinerView.cpp
--- a/Engine/Plugins/Experimental/UAF/UAF/Source/UAFEditor/Private/Functions/Outliner/FunctionsOutlinerView.cpp
+++ b/Engine/Plugins/Experimental/UAF/UAF/Source/UAFEditor/Private/Functions/Outliner/FunctionsOutlinerView.cpp
@@ -97,6 +97,7 @@ void SFunctionsView::OnRigVMModified(ERigVMGraphNotifType NotifType, URigVMGraph
 		break;
 	case ERigVMGraphNotifType::NodeAdded:
 	case ERigVMGraphNotifType::NodeRemoved:
+	case ERigVMGraphNotifType::NodeCategoryChanged:
 		{
 			if (Cast<URigVMFunctionLibrary>(Graph))
 			{
