//UE5/Main - Build Engine Localization - Text conflict from LOCTEXT macro for namespace 'AnimNextUncookedOnlyUtils' and key 'InvalidSharedVariableAssetFound'
#fix rename LOCTEXT entry

diff --git a/Engine/Plugins/Experimental/UAF/UAF/Source/UAFUncookedOnly/Private/UncookedOnlyUtils.cpp b/Engine/Plugins/Experimental/UAF/UAF/Source/UAFUncookedOnly/Private/UncookedOnlyUtils.cpp
--- a/Engine/Plugins/Experimental/UAF/UAF/Source/UAFUncookedOnly/Private/UncookedOnlyUtils.cpp
+++ b/Engine/Plugins/Experimental/UAF/UAF/Source/UAFUncookedOnly/Private/UncookedOnlyUtils.cpp
@@ -214,7 +214,7 @@ void FUtils::CompileVariables(const FRigVMCompileSettings& InSettings, UAnimNext
 					
 				TSharedRef<FTokenizedMessage> Message = FTokenizedMessage::Create(EMessageSeverity::Error);
 				Message->AddToken(FUObjectToken::Create(InAsset));
-				Message->AddText(LOCTEXT("InvalidSharedVariableAssetFound", "Invalid shared variables entry found ({0})"), FText::FromString(SharedVariablesEntry->GetObjectPath().ToString()));
+				Message->AddText(LOCTEXT("InvalidSharedVariableEntryFound", "Invalid shared variables entry found ({0})"), FText::FromString(SharedVariablesEntry->GetObjectPath().ToString()));
 				Message->AddToken(FActionToken::Create(LOCTEXT("RemoveSharedVariablesEntry", "Remove Shared Variables entry"), LOCTEXT("RemoveSharedVariablesEntryDesc", "Removes the Shared Variables entry from the asset."), FOnActionTokenExecuted::CreateLambda([WeakEditorData, WeakEntry]()
 					{
 						UAnimNextRigVMAssetEditorData* EditorData = WeakEditorData.Get();
