Add iconography and colouration for function entries in function outliner

diff --git a/Engine/Plugins/Experimental/UAF/UAF/Source/UAFEditor/Private/Functions/Outliner/FunctionOutlinerEntry.cpp b/Engine/Plugins/Experimental/UAF/UAF/Source/UAFEditor/Private/Functions/Outliner/FunctionOutlinerEntry.cpp
--- a/Engine/Plugins/Experimental/UAF/UAF/Source/UAFEditor/Private/Functions/Outliner/FunctionOutlinerEntry.cpp
+++ b/Engine/Plugins/Experimental/UAF/UAF/Source/UAFEditor/Private/Functions/Outliner/FunctionOutlinerEntry.cpp
@@ -3,6 +3,7 @@
 #include "FunctionOutlinerEntry.h"
 
 #include "AnimNextRigVMAssetEditorData.h"
+#include "GraphEditorSettings.h"
 #include "IAnimNextRigVMExportInterface.h"
 #include "ISceneOutliner.h"
 #include "ScopedTransaction.h"
@@ -23,6 +24,7 @@ FFunctionsOutlinerEntryItem::FFunctionsOutlinerEntryItem(const FItemData& ItemDa
 	: FOutlinerEntryItem(FFunctionsOutlinerEntryItem::Type, {ItemData.Asset, ItemData.SortValue}),
 		WeakLibraryNode(ItemData.LibraryNode)
 {
+	Flags.bIsExpanded = false;
 }
 
 bool FFunctionsOutlinerEntryItem::IsValid() const
@@ -169,6 +171,25 @@ bool FFunctionsOutlinerEntryItem::IsReadOnly() const
 {
 	return !WeakLibraryNode.IsValid();
 }
-}
+
+void FFunctionsOutlinerEntryItem::GetItemIconAndColor(const FSlateBrush*& OutBrush, FSlateColor& OutColor) const
+{
+	if (const UGraphEditorSettings* Settings = GetDefault<UGraphEditorSettings>())
+    {    
+	    if (URigVMLibraryNode* Node = WeakLibraryNode.Get())
+	    {
+		    if (Node->IsPure())
+		    {
+			    OutColor = Settings->PureFunctionCallNodeTitleColor;
+			    OutBrush = FAppStyle::GetBrush(TEXT("GraphEditor.PureFunction_16x"));
+
+			    return;
+		    }
+	    }
+	    
+	    OutColor = Settings->FunctionCallNodeTitleColor;
+	    OutBrush = FAppStyle::GetBrush(TEXT("GraphEditor.Function_16x"));
+	}
+}}
 
 #undef LOCTEXT_NAMESPACE // "FFunctionsOutlinerEntryItem"
diff --git a/Engine/Plugins/Experimental/UAF/UAF/Source/UAFEditor/Private/Functions/Outliner/FunctionOutlinerEntry.h b/Engine/Plugins/Experimental/UAF/UAF/Source/UAFEditor/Private/Functions/Outliner/FunctionOutlinerEntry.h
--- a/Engine/Plugins/Experimental/UAF/UAF/Source/UAFEditor/Private/Functions/Outliner/FunctionOutlinerEntry.h
+++ b/Engine/Plugins/Experimental/UAF/UAF/Source/UAFEditor/Private/Functions/Outliner/FunctionOutlinerEntry.h
@@ -45,6 +45,7 @@ struct FFunctionsOutlinerEntryItem : FOutlinerEntryItem
 	virtual EAnimNextExportAccessSpecifier GetAccessSpecifier() const override;
 	virtual bool CanSetAccessSpecifier() const override;
 	virtual bool IsReadOnly() const override;
+	virtual void GetItemIconAndColor(const FSlateBrush*& OutBrush, FSlateColor& OutColor) const override;
 
 	// Ptr to the underlying RigVM library node representing this function within the RigVMAsset
 	TWeakObjectPtr<URigVMLibraryNode> WeakLibraryNode;
