Populate Collapse node (function outliner) treeview items for nodes within a function

diff --git a/Engine/Plugins/Experimental/UAF/UAF/Source/UAFEditor/Private/Functions/Outliner/FunctionsOutlinerCollapsedGraphItem.cpp b/Engine/Plugins/Experimental/UAF/UAF/Source/UAFEditor/Private/Functions/Outliner/FunctionsOutlinerCollapsedGraphItem.cpp
new file mode 100644
--- /dev/null
+++ b/Engine/Plugins/Experimental/UAF/UAF/Source/UAFEditor/Private/Functions/Outliner/FunctionsOutlinerCollapsedGraphItem.cpp
@@ -0,0 +1,138 @@
+// Copyright Epic Games, Inc. All Rights Reserved.
+
+#include "FunctionsOutlinerCollapsedGraphItem.h"
+
+#include "AnimNextRigVMAssetEditorData.h"
+#include "IAnimNextRigVMExportInterface.h"
+#include "ScopedTransaction.h"
+#include "RigVMModel/Nodes/RigVMCollapseNode.h"
+
+#define LOCTEXT_NAMESPACE "FFunctionsOutlinerCollapsedGraphItem"
+
+namespace UE::UAF::Editor
+{
+
+const FSceneOutlinerTreeItemType FFunctionsOutlinerCollapsedGraphItem::Type(&FOutlinerEntryItem::Type);
+
+FFunctionsOutlinerCollapsedGraphItem::FFunctionsOutlinerCollapsedGraphItem(const FItemData& ItemData)
+	: FOutlinerEntryItem(FFunctionsOutlinerCollapsedGraphItem::Type, {ItemData.Asset, ItemData.SortValue}),
+	WeakCollapseNode(ItemData.CollapseNode), WeakOwningNode(ItemData.OwningNode)
+{
+}
+
+bool FFunctionsOutlinerCollapsedGraphItem::IsValid() const
+{
+	return WeakCollapseNode.IsValid();
+}
+
+FSceneOutlinerTreeItemID FFunctionsOutlinerCollapsedGraphItem::GetID() const
+{
+	if (URigVMCollapseNode* Node = WeakCollapseNode.Get())
+	{
+		return GetTypeHash(Node);
+	}
+
+	return GetTypeHash(WeakCollapseNode);
+}
+
+FString FFunctionsOutlinerCollapsedGraphItem::GetDisplayString() const
+{
+	if (URigVMCollapseNode* Node = WeakCollapseNode.Get())
+	{
+		return Node->GetName();
+	}
+
+	return FString();
+}
+
+FString FFunctionsOutlinerCollapsedGraphItem::GetPackageName() const
+{
+	if (URigVMCollapseNode* Node = WeakCollapseNode.Get())
+	{
+		return Node->GetOutermost()->GetName();
+	}
+
+	return ISceneOutlinerTreeItem::GetPackageName();
+}
+
+bool FFunctionsOutlinerCollapsedGraphItem::IsReadOnly() const
+{
+	return false;
+}
+EAnimNextExportAccessSpecifier FFunctionsOutlinerCollapsedGraphItem::GetAccessSpecifier() const
+{
+	if (URigVMLibraryNode* OwningNode = WeakOwningNode.Get())
+	{
+		if (UAnimNextRigVMAssetEditorData* EditorData = OwningNode->GetTypedOuter<UAnimNextRigVMAssetEditorData>())
+		{
+			if(EditorData->GetLocalFunctionLibrary()->IsFunctionPublic(OwningNode->GetFName()))
+			{
+				return EAnimNextExportAccessSpecifier::Public;
+			}
+
+			return EAnimNextExportAccessSpecifier::Private;
+		}
+	}
+	
+	return EAnimNextExportAccessSpecifier::Private;
+}
+
+bool FFunctionsOutlinerCollapsedGraphItem::CanSetAccessSpecifier() const
+{
+	return false;
+}
+
+FStringView FFunctionsOutlinerCollapsedGraphItem::GetCategoryPath() const
+{
+	return FStringView(nullptr, 0);
+}
+
+void FFunctionsOutlinerCollapsedGraphItem::GetItemIconAndColor(const FSlateBrush*& OutBrush, FSlateColor& OutColor) const
+{
+	OutBrush = FAppStyle::GetBrush(TEXT("GraphEditor.SubGraph_16x"));
+}
+
+void FFunctionsOutlinerCollapsedGraphItem::Rename(const FText& InNewName) const
+{
+	if (URigVMCollapseNode* CollapseNode = WeakCollapseNode.Get())
+	{
+		if (const URigVMGraph* ContainedGraph = CollapseNode->GetContainedGraph())
+		{
+			if (IRigVMClientHost* RigVMHost = CollapseNode->GetImplementingOuter<IRigVMClientHost>())
+			{
+				if (URigVMEdGraph* ContainedEdGraph = Cast<URigVMEdGraph>(RigVMHost->GetEditorObjectForRigVMGraph(ContainedGraph)))
+				{
+					if (const UEdGraphSchema* GraphSchema = ContainedEdGraph->GetSchema())
+					{
+						FGraphDisplayInfo DisplayInfo;
+						GraphSchema->GetGraphDisplayInformation(*ContainedEdGraph, DisplayInfo);
+
+						// Check if the name is unchanged
+						if (InNewName.EqualTo(DisplayInfo.PlainName))
+						{
+							return;
+						}
+
+						FScopedTransaction Transaction(LOCTEXT("RenameCollapsedGraphInOutliner", "Rename Collapsed Graph"));
+						FRigVMControllerCompileBracketScope CompileScope(ContainedEdGraph->GetController());
+						GraphSchema->TryRenameGraph(ContainedEdGraph, *InNewName.ToString());
+					}
+				}
+			}
+		}
+	}
+}
+
+bool FFunctionsOutlinerCollapsedGraphItem::ValidateName(const FText& InNewName, FText& OutErrorMessage) const
+{
+	if (InNewName.IsEmpty())
+	{
+		OutErrorMessage = LOCTEXT("NameEmpty", "Collapsed Graph name cannot be empty");
+		return false;
+	}
+
+	return true;
+}
+}
+
+#undef LOCTEXT_NAMESPACE // "FFunctionsOutlinerCollapsedGraphItem"
diff --git a/Engine/Plugins/Experimental/UAF/UAF/Source/UAFEditor/Private/Functions/Outliner/FunctionsOutlinerCollapsedGraphItem.h b/Engine/Plugins/Experimental/UAF/UAF/Source/UAFEditor/Private/Functions/Outliner/FunctionsOutlinerCollapsedGraphItem.h
new file mode 100644
--- /dev/null
+++ b/Engine/Plugins/Experimental/UAF/UAF/Source/UAFEditor/Private/Functions/Outliner/FunctionsOutlinerCollapsedGraphItem.h
@@ -0,0 +1,47 @@
+// Copyright Epic Games, Inc. All Rights Reserved.
+
+#pragma once
+#include "Common/Outliner/OutlinerEntryItem.h"
+
+class UAnimNextRigVMAsset;
+class URigVMCollapseNode;
+class URigVMLibraryNode;
+
+namespace UE::UAF::Editor
+{
+struct FFunctionsOutlinerCollapsedGraphItem : FOutlinerEntryItem
+{
+	static const FSceneOutlinerTreeItemType Type;
+	
+	struct FItemData
+	{
+		const TWeakObjectPtr<UAnimNextRigVMAsset>& Asset;
+		const int32 SortValue;
+		URigVMCollapseNode* CollapseNode;
+		URigVMLibraryNode* OwningNode;
+	};
+	
+	FFunctionsOutlinerCollapsedGraphItem() = delete;
+	FFunctionsOutlinerCollapsedGraphItem(const FItemData& ItemData);
+	
+	// Begin ISceneOutlinerTreeItem overrides
+	virtual bool IsValid() const override;
+	virtual FSceneOutlinerTreeItemID GetID() const override;
+	virtual FString GetDisplayString() const override;
+	virtual bool CanInteract() const override { return true; }
+	virtual FString GetPackageName() const override;
+	// End ISceneOutlinerTreeItem overrides
+	
+	virtual bool IsReadOnly() const override;
+	virtual void SetAccessSpecifier(const EAnimNextExportAccessSpecifier& InSpecifier) const override {}
+	virtual EAnimNextExportAccessSpecifier GetAccessSpecifier() const override;
+	virtual bool CanSetAccessSpecifier() const override;
+	virtual FStringView GetCategoryPath() const override;
+	virtual void GetItemIconAndColor(const FSlateBrush*& OutBrush, FSlateColor& OutColor) const override;
+	virtual void Rename(const FText& InNewName) const override;
+	virtual bool ValidateName(const FText& InNewName, FText& OutErrorMessage) const override;
+
+	TWeakObjectPtr<URigVMCollapseNode> WeakCollapseNode;
+	TWeakObjectPtr<URigVMLibraryNode> WeakOwningNode;
+};
+}
diff --git a/Engine/Plugins/Experimental/UAF/UAF/Source/UAFEditor/Private/Functions/Outliner/FunctionsOutlinerHierarchy.cpp b/Engine/Plugins/Experimental/UAF/UAF/Source/UAFEditor/Private/Functions/Outliner/FunctionsOutlinerHierarchy.cpp
--- a/Engine/Plugins/Experimental/UAF/UAF/Source/UAFEditor/Private/Functions/Outliner/FunctionsOutlinerHierarchy.cpp
+++ b/Engine/Plugins/Experimental/UAF/UAF/Source/UAFEditor/Private/Functions/Outliner/FunctionsOutlinerHierarchy.cpp
@@ -5,6 +5,7 @@
 #include "AnimNextRigVMAssetEditorData.h"
 #include "FunctionOutlinerEntry.h"
 #include "FunctionsOutlinerCategoryItem.h"
+#include "FunctionsOutlinerCollapsedGraphItem.h"
 #include "ISceneOutlinerMode.h"
 #include "Common/Outliner/OutlinerAssetItem.h"
 
@@ -39,6 +40,14 @@ FSceneOutlinerTreeItemPtr FFunctionsOutlinerHierarchy::FindOrCreateParentItem(co
 			ParentHash = CategoryHash != INDEX_NONE ? HashCombine(GetTypeHash(OwnerPath), CategoryHash) : GetTypeHash(OwnerPath);				
 		}
 	}
+	
+	if (const FFunctionsOutlinerCollapsedGraphItem* CollapsedGraphItem = Item.CastTo<FFunctionsOutlinerCollapsedGraphItem>())
+	{
+		if (URigVMLibraryNode* LibraryNode = CollapsedGraphItem->WeakOwningNode.Get())
+		{
+			ParentHash = GetTypeHash(LibraryNode);
+		}
+	}
 
 	if (const FFunctionsOutlinerCategoryItem* CategoryItem = Item.CastTo<FFunctionsOutlinerCategoryItem>())
 	{
@@ -153,14 +162,28 @@ void FFunctionsOutlinerHierarchy::CreateItemsForAsset(const UAnimNextRigVMAssetE
 				if (FFunctionsOutlinerEntryItem* EntryItem = Item->CastTo<FFunctionsOutlinerEntryItem>())
 				{
 					EntryItem->SortValue = Index++;	
-					
+
 					UE_CLOG(GOutputHierarchyDebugInformationCvar.GetValueOnAnyThread(), LogAnimation, Warning, TEXT("\t- FFunctionsOutlinerEntryItem: %s %i %s"), *EntryItem->WeakLibraryNode.Get()->GetName(), Index, *LibraryNode->GetNodeCategory());
 				}
 
 				const FString CategoryPath = LibraryNode->GetNodeCategory();
 				RecursivelyAddCategories(CategoryPath);
-			
 				OutItems.Add(Item);
+
+				if (LibraryNode->GetContainedGraph())
+				{
+					int32 SubIndex = 0;
+					for (URigVMNode* Node : LibraryNode->GetContainedGraph()->GetNodes())
+					{
+						if (URigVMCollapseNode* CollapseNode = Cast<URigVMCollapseNode>(Node))
+						{
+							if (FSceneOutlinerTreeItemPtr CollapsedItem = Mode->CreateItemFor<FFunctionsOutlinerCollapsedGraphItem>(FFunctionsOutlinerCollapsedGraphItem::FItemData{UncookedOnly::FUtils::GetAsset<UAnimNextRigVMAsset>(EditorData), ++SubIndex, CollapseNode, LibraryNode}))
+							{
+								OutItems.Add(CollapsedItem);
+							}
+						}
+					}
+				}
 			}
 		}
 	}
diff --git a/Engine/Plugins/Experimental/UAF/UAF/Source/UAFEditor/Private/Functions/Outliner/FunctionsOutlinerMode.cpp b/Engine/Plugins/Experimental/UAF/UAF/Source/UAFEditor/Private/Functions/Outliner/FunctionsOutlinerMode.cpp
--- a/Engine/Plugins/Experimental/UAF/UAF/Source/UAFEditor/Private/Functions/Outliner/FunctionsOutlinerMode.cpp
+++ b/Engine/Plugins/Experimental/UAF/UAF/Source/UAFEditor/Private/Functions/Outliner/FunctionsOutlinerMode.cpp
@@ -5,6 +5,7 @@
 #include "AnimNextRigVMAsset.h"
 #include "AnimNextRigVMAssetEditorData.h"
 #include "FunctionOutlinerEntry.h"
+#include "FunctionsOutlinerCollapsedGraphItem.h"
 #include "FunctionsOutlinerView.h"
 #include "FunctionsOutlinerHierarchy.h"
 #include "IWorkspaceEditor.h"
@@ -493,10 +494,10 @@ void FFunctionsOutlinerMode::OnItemDoubleClick(FSceneOutlinerTreeItemPtr Item)
 {
 	if (TSharedPtr<Workspace::IWorkspaceEditor> WorkspaceEditor = WeakWorkspaceEditor.Pin())
 	{		
+		const Workspace::FWorkspaceDocument& Document = WorkspaceEditor->GetFocusedWorkspaceDocument();
+		FWorkspaceOutlinerItemExport DocumentExport = Document.Export;
 		if (FFunctionsOutlinerEntryItem* FunctionItem = Item->CastTo<FFunctionsOutlinerEntryItem>())
 		{
-			const Workspace::FWorkspaceDocument& Document = WorkspaceEditor->GetFocusedWorkspaceDocument();
-			FWorkspaceOutlinerItemExport DocumentExport = Document.Export;
 			if (FunctionItem->WeakLibraryNode.IsValid())
 			{			
 				const URigVMLibraryNode* FunctionNode = FunctionItem->WeakLibraryNode.Get();
@@ -518,10 +519,40 @@ void FFunctionsOutlinerMode::OnItemDoubleClick(FSceneOutlinerTreeItemPtr Item)
 					{
 						if (URigVMEdGraph* ContainedEdGraph = Cast<URigVMEdGraph>(Host->GetEditorObjectForRigVMGraph(ContainedGraph)))
 						{							
-							WorkspaceEditor->OpenObjects({ContainedEdGraph});							
+							WorkspaceEditor->OpenObjects({ContainedEdGraph});
 							return;
 						}
-					}	
+					}
+				}
+			}
+		}
+		else if (FFunctionsOutlinerCollapsedGraphItem* CollapsedGraphItem = Item->CastTo<FFunctionsOutlinerCollapsedGraphItem>())
+		{
+			if (CollapsedGraphItem->WeakOwningNode.IsValid() && CollapsedGraphItem->WeakCollapseNode.IsValid())
+			{
+				const URigVMLibraryNode* FunctionNode = CollapsedGraphItem->WeakOwningNode.Get();
+				UAnimNextRigVMAsset* Asset = FunctionNode->GetTypedOuter<UAnimNextRigVMAsset>();
+
+				const FWorkspaceOutlinerItemExport RelativeAssetExport = DocumentExport.GetRelativeAssetExport(Asset);
+				const FWorkspaceOutlinerItemExport AssetExport = UncookedOnly::FUtils::MakeAssetReferenceExport(Asset, RelativeAssetExport);
+				if (AssetExport.GetFirstAssetPath().IsValid() && RelativeAssetExport != AssetExport)
+				{
+					const FWorkspaceOutlinerItemExport FunctionsExport = UncookedOnly::FUtils::MakeFunctionExport(FunctionNode, AssetExport);
+					const FWorkspaceOutlinerItemExport CollapsedGraphExport = UncookedOnly::FUtils::MakeCollapsedGraphExport(CollapsedGraphItem->WeakCollapseNode.Get(), FunctionsExport);
+					WorkspaceEditor->OpenExports({CollapsedGraphExport});
+					return;
+				}
+
+				if (URigVMGraph* ContainedGraph = CollapsedGraphItem->WeakCollapseNode->GetContainedGraph())
+				{
+					if (IRigVMClientHost* Host = FunctionNode->GetImplementingOuter<IRigVMClientHost>())
+					{
+						if (URigVMEdGraph* ContainedEdGraph = Cast<URigVMEdGraph>(Host->GetEditorObjectForRigVMGraph(ContainedGraph)))
+						{
+							WorkspaceEditor->OpenObjects({ContainedEdGraph});
+							return;
+						}
+					}
 				}
 			}
 		}
@@ -532,16 +563,7 @@ void FFunctionsOutlinerMode::OnItemDoubleClick(FSceneOutlinerTreeItemPtr Item)
 
 bool FFunctionsOutlinerMode::CanRenameItem(const ISceneOutlinerTreeItem& Item) const
 {
-	if (const FFunctionsOutlinerEntryItem* EntryItem = Item.CastTo<FFunctionsOutlinerEntryItem>())
-	{
-		return EntryItem->IsValid();
-	}
-	else if (Item.IsA<FFunctionsOutlinerCategoryItem>())
-	{
-		return true;
-	}
-	
-	return false;
+	return Item.IsA<FFunctionsOutlinerEntryItem>() || Item.IsA<FFunctionsOutlinerCategoryItem>() || Item.IsA<FFunctionsOutlinerCollapsedGraphItem>();
 }
 
 bool FFunctionsOutlinerMode::UpdateOperationDecorator(const FDragDropEvent& Event, const FSceneOutlinerDragValidationInfo* ValidationInfo) const
@@ -552,8 +574,8 @@ bool FFunctionsOutlinerMode::UpdateOperationDecorator(const FDragDropEvent& Even
 		if (ValidationInfo)
 		{
 			const FSlateBrush* ValidHoverIcon = FAppStyle::GetBrush(TEXT("Graph.ConnectorFeedback.OK"));
-			const FSlateBrush* InvalidSymbol = FAppStyle::GetBrush(TEXT("Graph.ConnectorFeedback.Error"));			
-			GraphDropOp->SetSimpleFeedbackMessage(ValidationInfo->IsValid() ? ValidHoverIcon : InvalidSymbol, FLinearColor::White, ValidationInfo->ValidationText);		
+			const FSlateBrush* InvalidSymbol = FAppStyle::GetBrush(TEXT("Graph.ConnectorFeedback.Error"));	
+			GraphDropOp->SetSimpleFeedbackMessage(ValidationInfo->IsValid() ? ValidHoverIcon : InvalidSymbol, FLinearColor::White, ValidationInfo->ValidationText);
 		}
 		else
 		{
diff --git a/Engine/Plugins/Experimental/UAF/UAF/Source/UAFUncookedOnly/Internal/AnimNextRigVMAssetEditorData.h b/Engine/Plugins/Experimental/UAF/UAF/Source/UAFUncookedOnly/Internal/AnimNextRigVMAssetEditorData.h
--- a/Engine/Plugins/Experimental/UAF/UAF/Source/UAFUncookedOnly/Internal/AnimNextRigVMAssetEditorData.h
+++ b/Engine/Plugins/Experimental/UAF/UAF/Source/UAFUncookedOnly/Internal/AnimNextRigVMAssetEditorData.h
@@ -67,6 +67,7 @@ namespace UE::UAF::Editor
 	struct FFunctionsOutlinerEntryItem;
 	struct FVariablesOutlinerCategoryItem;
 	struct FFunctionsOutlinerCategoryItem;
+	struct FFunctionsOutlinerCollapsedGraphItem;
 	class FFunctionsOutlinerMode;
 	class SCommonOutliner;
 }
@@ -268,6 +269,7 @@ protected:
 	friend struct FAnimNextAnimationGraphTests;
 	friend FAnimNextRigVMAssetCompileContext;
 	friend UE::UAF::Editor::FFunctionsOutlinerEntryItem;
+	friend UE::UAF::Editor::FFunctionsOutlinerCollapsedGraphItem;
 	friend UE::UAF::Editor::SFunctionsView;
 	friend UE::UAF::Editor::SCommonOutliner;
 	friend UE::UAF::Editor::SVariablesView;
diff --git a/Engine/Plugins/Experimental/UAF/UAF/Source/UAFUncookedOnly/Internal/UncookedOnlyUtils.h b/Engine/Plugins/Experimental/UAF/UAF/Source/UAFUncookedOnly/Internal/UncookedOnlyUtils.h
--- a/Engine/Plugins/Experimental/UAF/UAF/Source/UAFUncookedOnly/Internal/UncookedOnlyUtils.h
+++ b/Engine/Plugins/Experimental/UAF/UAF/Source/UAFUncookedOnly/Internal/UncookedOnlyUtils.h
@@ -41,6 +41,7 @@ class UAnimNextRigVMAssetEditorData;
 class UAnimNextRigVMAssetEntry;
 class URigVMUnitNode;
 class UAnimNextControllerBase;
+class URigVMCollapseNode;
 
 namespace UE::UAF
 {
@@ -228,6 +229,9 @@ struct FUtils
 	
 	// Generate a (function) workspace export for the provided URigVMLibraryNode
 	static UE_API FWorkspaceOutlinerItemExport MakeFunctionExport(const URigVMLibraryNode* FunctionNode, const FWorkspaceOutlinerItemExport& ParentExport);
+	
+	// Generate a (collapse graph) workspace export for the provided URigVMCollapseNode
+	static UE_API FWorkspaceOutlinerItemExport MakeCollapsedGraphExport(const URigVMCollapseNode* CollapseNode, const FWorkspaceOutlinerItemExport& ParentExport);
 
 	// Generate a (asset reference) workspace export for the provided asset path
 	static UE_API FWorkspaceOutlinerItemExport MakeAssetReferenceExport(const FSoftObjectPath& ReferencedAssetPath, const FWorkspaceOutlinerItemExport& ParentExport, bool bShouldExpand = true);
diff --git a/Engine/Plugins/Experimental/UAF/UAF/Source/UAFUncookedOnly/Private/UncookedOnlyUtils.cpp b/Engine/Plugins/Experimental/UAF/UAF/Source/UAFUncookedOnly/Private/UncookedOnlyUtils.cpp
--- a/Engine/Plugins/Experimental/UAF/UAF/Source/UAFUncookedOnly/Private/UncookedOnlyUtils.cpp
+++ b/Engine/Plugins/Experimental/UAF/UAF/Source/UAFUncookedOnly/Private/UncookedOnlyUtils.cpp
@@ -571,6 +571,30 @@ FWorkspaceOutlinerItemExport FUtils::MakeFunctionExport(const URigVMLibraryNode*
 	return Export;
 }
 
+FWorkspaceOutlinerItemExport FUtils::MakeCollapsedGraphExport(const URigVMCollapseNode* CollapseNode, const FWorkspaceOutlinerItemExport& ParentExport)
+{
+	FWorkspaceOutlinerItemExport Export;
+	if (CollapseNode)
+	{		
+		if (IRigVMClientHost* OuterHost = CollapseNode->GetImplementingOuter<IRigVMClientHost>())
+		{
+			if (URigVMEdGraph* EditorObject = Cast<URigVMEdGraph>(OuterHost->GetEditorObjectForRigVMGraph(CollapseNode->GetContainedGraph())))
+			{
+				const FSoftObjectPath EditorObjectSoftObjPath = EditorObject;
+				ensureMsgf(EditorObjectSoftObjPath.IsSubobject(), TEXT("EditorObject for RigVMCollapseNode Graph was not a subobject as expected."));
+
+				const FName Identifier = EditorObjectSoftObjPath.IsSubobject() ? *EditorObjectSoftObjPath.GetSubPathUtf8String() : CollapseNode->GetFName();
+				Export = FWorkspaceOutlinerItemExport(Identifier, ParentExport);
+				Export.GetData().InitializeAsScriptStruct(FAnimNextCollapseGraphOutlinerData::StaticStruct());
+				FAnimNextCollapseGraphOutlinerData& FnGraphData = Export.GetData().GetMutable<FAnimNextCollapseGraphOutlinerData>();
+				FnGraphData.SoftEditorObject = EditorObject;
+			}
+		}
+	}
+	
+	return Export;
+}
+
 FWorkspaceOutlinerItemExport FUtils::MakeAssetReferenceExport(const FSoftObjectPath& ReferencedAssetPath, const FWorkspaceOutlinerItemExport& ParentExport, bool bShouldExpand)
 {
 	check(ReferencedAssetPath.IsValid());
