Static analysis fix, use default FStringView c-tor rather than explicit (nullptr, 0)

diff --git a/Engine/Plugins/Experimental/UAF/UAF/Source/UAFEditor/Private/Functions/Outliner/FunctionsOutlinerCollapsedGraphItem.cpp b/Engine/Plugins/Experimental/UAF/UAF/Source/UAFEditor/Private/Functions/Outliner/FunctionsOutlinerCollapsedGraphItem.cpp
--- a/Engine/Plugins/Experimental/UAF/UAF/Source/UAFEditor/Private/Functions/Outliner/FunctionsOutlinerCollapsedGraphItem.cpp
+++ b/Engine/Plugins/Experimental/UAF/UAF/Source/UAFEditor/Private/Functions/Outliner/FunctionsOutlinerCollapsedGraphItem.cpp
@@ -84,7 +84,7 @@ bool FFunctionsOutlinerCollapsedGraphItem::CanSetAccessSpecifier() const
 
 FStringView FFunctionsOutlinerCollapsedGraphItem::GetCategoryPath() const
 {
-	return FStringView(nullptr, 0);
+	return FStringView();
 }
 
 void FFunctionsOutlinerCollapsedGraphItem::GetItemIconAndColor(const FSlateBrush*& OutBrush, FSlateColor& OutColor) const
