UAF: Move ScriptEventPhase.h to public (it is included by public headers)

diff --git a/Engine/Plugins/Experimental/UAF/UAF/Source/UAF/Internal/Script/ScriptEventPhase.h b/Engine/Plugins/Experimental/UAF/UAF/Source/UAF/Internal/Script/ScriptEventPhase.h
deleted file mode 100644
diff --git a/Engine/Plugins/Experimental/UAF/UAF/Source/UAF/Public/Script/ScriptEventPhase.h b/Engine/Plugins/Experimental/UAF/UAF/Source/UAF/Public/Script/ScriptEventPhase.h
new file mode 100644
--- /dev/null
+++ b/Engine/Plugins/Experimental/UAF/UAF/Source/UAF/Public/Script/ScriptEventPhase.h
@@ -0,0 +1,20 @@
+﻿// Copyright Epic Games, Inc. All Rights Reserved.
+
+#pragma once
+
+#include "CoreTypes.h"
+
+namespace UE::UAF
+{
+
+// Phase is used as a general ordering constraint on event execution
+enum class EScriptEventPhase : uint8
+{
+	// Before any execution, e.g. for copying data from the game thread
+	PreExecute,
+
+	// General execution, e.g. a prephysics event
+	Execute,
+};
+
+}
