UAF Browser Impl
- Still disabled by default via `a.UAF.EnableUAFBrowser`. Will enable after submitting default layout asset.
- Cleanup SUAFBrowser lambdas into header methods
- Added support for asset previews on click
- Added IUAFAssetPreview widget which Async creates preview widgets for assets
- Added IUAFAssetViewportPreview widget to handle viewport boilerplate
- Added UAF Graph & Anim Sequence Asset Preview
- Added Factory system for asset preview widgets (To be replaced in future).
- Moved SWorkspaceViewport to public

#tests editor

diff --git a/Engine/Plugins/Experimental/UAF/UAF/Source/UAF/Public/Component/AnimNextComponent.h b/Engine/Plugins/Experimental/UAF/UAF/Source/UAF/Public/Component/AnimNextComponent.h
--- a/Engine/Plugins/Experimental/UAF/UAF/Source/UAF/Public/Component/AnimNextComponent.h
+++ b/Engine/Plugins/Experimental/UAF/UAF/Source/UAF/Public/Component/AnimNextComponent.h
@@ -34,6 +34,7 @@ namespace UE::UAF::Editor
 	class FAssetWizard;
 	class FSystemViewportController;
 	class FAnimGraphViewportController;
+	class SUAFGraphAssetPreview;
 }
 
 UCLASS(MinimalAPI, meta = (BlueprintSpawnableComponent))
@@ -224,6 +225,7 @@ private:
 	friend struct FAnimMixerUtils;
 	friend class UE::UAF::Editor::FSystemViewportController;
 	friend class UE::UAF::Editor::FAnimGraphViewportController;
+	friend class UE::UAF::Editor::SUAFGraphAssetPreview;
 
 	// The asset that this component will run
 	UPROPERTY(EditAnywhere, Category="Asset", DisplayName="Asset", meta=(GetAllowedClasses="/Script/UAF.AnimNextConfig:GetAllowedAssetClasses"))
diff --git a/Engine/Plugins/Experimental/UAF/UAF/Source/UAFEditor/Private/AnimNextEditorModule.cpp b/Engine/Plugins/Experimental/UAF/UAF/Source/UAFEditor/Private/AnimNextEditorModule.cpp
--- a/Engine/Plugins/Experimental/UAF/UAF/Source/UAFEditor/Private/AnimNextEditorModule.cpp
+++ b/Engine/Plugins/Experimental/UAF/UAF/Source/UAFEditor/Private/AnimNextEditorModule.cpp
@@ -47,6 +47,9 @@
 #include "Framework/Docking/LayoutExtender.h"
 #include "Common/AnimNextAssetItemDetails.h"
 #include "Common/AnimNextRigVMAssetEditorDataCustomization.h"
+#include "Common/AssetPreview/AnimSequenceAssetPreview.h"
+#include "Common/AssetPreview/IUAFAssetPreview.h"
+#include "Common/AssetPreview/UAFAssetPreviewFactorySubsystem.h"
 #include "Variables/AnimNextSharedVariables.h"
 #include "Tools/EdModeInteractiveToolsContext.h"
 #include "Entries/AnimNextEventGraphEntry.h"
@@ -203,6 +206,15 @@ void FAnimNextEditorModule::StartupModule()
 	GraphPanelPinFactory = MakeShared<FAnimNextGraphPanelPinFactoryEditor>();
 	FEdGraphUtilities::RegisterVisualPinFactory(GraphPanelPinFactory);
 
+	// Temp register asset previews while we still use the UUAFAssetPreviewFactorySubsystem
+	FCoreDelegates::OnPostEngineInit.AddLambda([]()
+	{
+		if (GEditor)
+		{
+			GEditor->GetEditorSubsystem<UUAFAssetPreviewFactorySubsystem>()->AddAssetPreviewFactory(MakeShared<FAnimSequenceAssetPreviewFactory>());
+		}
+	});
+
 	Workspace::IWorkspaceEditorModule& WorkspaceEditorModule = FModuleManager::Get().LoadModuleChecked<Workspace::IWorkspaceEditorModule>("WorkspaceEditor");
 
 	WorkspaceEditorModule.OnRegisterTabsForEditor().AddLambda([](FWorkflowAllowedTabSet& TabFactories, const TSharedRef<FTabManager>& InTabManager, TSharedPtr<UE::Workspace::IWorkspaceEditor> InEditorPtr)
diff --git a/Engine/Plugins/Experimental/UAF/UAF/Source/UAFEditor/Private/Common/AssetPreview/AnimSequenceAssetPreview.cpp b/Engine/Plugins/Experimental/UAF/UAF/Source/UAFEditor/Private/Common/AssetPreview/AnimSequenceAssetPreview.cpp
new file mode 100644
--- /dev/null
+++ b/Engine/Plugins/Experimental/UAF/UAF/Source/UAFEditor/Private/Common/AssetPreview/AnimSequenceAssetPreview.cpp
@@ -0,0 +1,95 @@
+// Copyright Epic Games, Inc. All Rights Reserved.
+
+#include "AnimSequenceAssetPreview.h"
+
+#include "Animation/AnimationAsset.h"
+#include "Animation/AnimSequenceBase.h"
+#include "Animation/AnimSequence.h"
+#include "Animation/DebugSkelMeshComponent.h"
+#include "AnimPreviewInstance.h"
+
+#define LOCTEXT_NAMESPACE "AnimSequenceAssetPreview"
+
+namespace UE::UAF::Editor
+{
+
+
+//////////////////////////////////////////////////////////////////////////
+// FAnimSequenceAssetPreviewFactory
+
+TSharedRef<IUAFAssetPreview> FAnimSequenceAssetPreviewFactory::CreateAssetPreviewWidget(TSharedPtr<FAssetEditorToolkit> InAssetEditorToolkit, const FAssetData& InAssetData) const
+{
+	return SNew(SAnimSequenceAssetPreview, InAssetEditorToolkit, InAssetData);
+}
+
+const UStruct* FAnimSequenceAssetPreviewFactory::GetPreviewType() const
+{
+	return UAnimSequence::StaticClass();
+}
+
+
+//////////////////////////////////////////////////////////////////////////
+// SAnimSequenceAssetPreview
+
+const UStruct* SAnimSequenceAssetPreview::GetAssetPreviewType() const
+{
+	return UAnimSequence::StaticClass();
+}
+
+bool SAnimSequenceAssetPreview::OnSameTypeAssetSelected(const FAssetData& InAssetData)
+{
+	bool bAssetValidForPreview = false;
+
+	USkeletalMesh* MeshToUse = nullptr;
+	UClass* AssetClass = FindObject<UClass>(CachedCurrentAssetData.AssetClassPath);
+	if(AssetClass->IsChildOf(UAnimationAsset::StaticClass()) && CachedCurrentAssetData.IsAssetLoaded() && CachedCurrentAssetData.GetAsset())
+	{
+		// Set up the viewport to show the asset.
+		UAnimationAsset* Asset = StaticCast<UAnimationAsset*>(CachedCurrentAssetData.GetAsset());
+		USkeleton* Skeleton = Asset->GetSkeleton();
+		if(Skeleton)
+		{
+			MeshToUse = Skeleton->GetAssetPreviewMesh(Asset);
+		}
+		
+		if(MeshToUse)
+		{
+			if(PreviewComponent->GetSkeletalMeshAsset() != MeshToUse)
+			{
+				PreviewComponent->SetSkeletalMesh(MeshToUse);
+			}
+
+			PreviewComponent->EnablePreview(true, Asset);
+			PreviewComponent->PreviewInstance->PlayAnim(true);
+
+			bAssetValidForPreview = true;
+		}
+	}
+
+	ViewportWidget->SetVisibility(bAssetValidForPreview ? EVisibility::Visible : EVisibility::Hidden);
+	return bAssetValidForPreview;
+}
+
+void SAnimSequenceAssetPreview::OnCustomizePreviewScene(FAdvancedPreviewScene& InPreviewScene, FEditorViewportClient& InEditorViewportClient)
+{
+	// @TODO: Note for review, most impl from SAnimationSequenceBrowser.h (Currently private, see IUAFAssetPreview.h for considerations).
+
+	// Setup the preview component to ensure an animation will update when requested.
+	// Note: Preview Component lifetime is managed by the preview scene
+	PreviewComponent = NewObject<UDebugSkelMeshComponent>();
+	PreviewComponent->VisibilityBasedAnimTickOption = EVisibilityBasedAnimTickOption::AlwaysTickPoseAndRefreshBones;
+	InPreviewScene.AddComponent(PreviewComponent, FTransform::Identity);
+
+	// @TODO: Use these settings?
+	//const USkeletalMeshEditorSettings* Options = GetDefault<USkeletalMeshEditorSettings>();
+	//PreviewScene.SetLightDirection(Options->AnimPreviewLightingDirection);
+	//PreviewScene.SetLightColor(Options->AnimPreviewDirectionalColor);
+	//PreviewScene.SetLightBrightness(Options->AnimPreviewLightBrightness);
+}
+
+//////////////////////////////////////////////////////////////////////////
+
+} // namespace UE::UAF::Editor
+
+#undef LOCTEXT_NAMESPACE
+
diff --git a/Engine/Plugins/Experimental/UAF/UAF/Source/UAFEditor/Private/Common/AssetPreview/AnimSequenceAssetPreview.h b/Engine/Plugins/Experimental/UAF/UAF/Source/UAFEditor/Private/Common/AssetPreview/AnimSequenceAssetPreview.h
new file mode 100644
--- /dev/null
+++ b/Engine/Plugins/Experimental/UAF/UAF/Source/UAFEditor/Private/Common/AssetPreview/AnimSequenceAssetPreview.h
@@ -0,0 +1,61 @@
+// Copyright Epic Games, Inc. All Rights Reserved.
+
+#pragma once
+
+#include "Common/AssetPreview/IUAFAssetViewportPreview.h"
+
+class FSceneViewport;
+class FUAFAssetPreviewViewportClient;
+class SAssetEditorViewport;
+class UDebugSkelMeshComponent;
+
+namespace UE::UAF::Editor
+{
+
+
+/**
+ * Experimental. May be removed / relocated whenever we switch to the future generic content browser IAssetPreviewCustomization API.
+ *
+ * Factory for Anim Sequence Asset Preview widgets
+ */
+struct FAnimSequenceAssetPreviewFactory : public FUAFAssetPreviewFactory
+{
+public:
+
+	//~ Begin FUAFAssetPreviewFactory Interface
+	virtual TSharedRef<IUAFAssetPreview> CreateAssetPreviewWidget(TSharedPtr<FAssetEditorToolkit> InAssetEditorToolkit, const FAssetData& InAssetData) const override;
+	virtual const UStruct* GetPreviewType() const override;
+	//~ End FUAFAssetPreviewFactory Interface
+};
+
+/** 
+ * Experimental. May be removed / relocated whenever we switch to the future generic content browser IAssetPreviewCustomization API. 
+ * 
+ * Preview widget for anim sequences. Includes an interactable viewport.
+ */
+class SAnimSequenceAssetPreview : public IUAFAssetViewportPreview
+{
+public:
+
+	//~ Begin IUAFAssetPreview Interface
+	virtual const UStruct* GetAssetPreviewType() const override;
+	//~ End IUAFAssetPreview Interface
+
+protected:
+
+	//~ Begin IUAFAssetPreview Interface
+	virtual bool OnSameTypeAssetSelected(const FAssetData& InAssetData) override;
+	//~ End IUAFAssetPreview Interface
+
+	//~ Begin IUAFAssetViewportPreview Interface
+	virtual void OnCustomizePreviewScene(FAdvancedPreviewScene& InPreviewScene, FEditorViewportClient& InEditorViewportClient) override;
+	//~ End IUAFAssetViewportPreview Interface
+
+protected:
+
+	/** Skeletal component to preview the animation asset on. Lifetime maintained by PreviewScene in IUAFAssetViewportPreview */
+	UDebugSkelMeshComponent* PreviewComponent = nullptr;
+};
+
+
+} // namespace UE::UAF::Editor
\ No newline at end of file
diff --git a/Engine/Plugins/Experimental/UAF/UAF/Source/UAFEditor/Private/Common/AssetPreview/IUAFAssetPreview.cpp b/Engine/Plugins/Experimental/UAF/UAF/Source/UAFEditor/Private/Common/AssetPreview/IUAFAssetPreview.cpp
new file mode 100644
--- /dev/null
+++ b/Engine/Plugins/Experimental/UAF/UAF/Source/UAFEditor/Private/Common/AssetPreview/IUAFAssetPreview.cpp
@@ -0,0 +1,171 @@
+// Copyright Epic Games, Inc. All Rights Reserved.
+
+#include "Common/AssetPreview/IUAFAssetPreview.h"
+
+#include "Engine/AssetManager.h"
+
+#define LOCTEXT_NAMESPACE "IUAFAssetPreview"
+
+namespace UE::UAF::Editor
+{
+
+
+//////////////////////////////////////////////////////////////////////////
+// IUAFAssetPreview
+
+void IUAFAssetPreview::Construct(const FArguments& InArgs, TSharedPtr<FAssetEditorToolkit> InAssetEditorToolkit, const FAssetData& InAssetData)
+{
+	AssetEditorToolkitPtr = InAssetEditorToolkit;
+	CachedCurrentAssetData = InAssetData;
+	OnConstructAsyncPreviewWidget(CachedCurrentAssetData);
+}
+
+void IUAFAssetPreview::OnSameTypeAssetSelectedAsync(const FAssetData& InAssetData)
+{
+	CachedCurrentAssetData = InAssetData;
+
+	FAssetData CachedCurrentAssetDataCopy = CachedCurrentAssetData;
+	TWeakPtr<SWidget> ThisWeak = AsShared();
+	auto SameAssetTypeSelectedCallback = [ThisWeak, CachedCurrentAssetDataCopy]()
+	{
+		if (TSharedPtr<IUAFAssetPreview> ThisPinned = StaticCastSharedPtr<IUAFAssetPreview>(ThisWeak.Pin()))
+		{
+			ThisPinned->OnSameTypeAssetSelected(CachedCurrentAssetDataCopy);
+
+			// Also update the widget from loading indicator
+			ThisPinned->ChildSlot
+			[
+				ThisPinned->OnConstructAssetPreviewWidget()
+			];
+		}
+	};
+
+	if (InAssetData.IsAssetLoaded())
+	{
+		SameAssetTypeSelectedCallback();
+	}
+	else
+	{
+		ChildSlot
+		[
+			OnConstructLoadingAssetPreviewWidget()
+		];
+
+		// Not done loading, kick off an async load
+		RequestAsyncLoad(InAssetData.GetSoftObjectPath(), FStreamableDelegate::CreateLambda(MoveTemp(SameAssetTypeSelectedCallback)));
+	}
+}
+
+void IUAFAssetPreview::RequestAsyncLoad(const FSoftObjectPath& SoftObjectPath, TFunction<void()>&& Callback)
+{
+	RequestAsyncLoad(SoftObjectPath, FStreamableDelegate::CreateLambda(MoveTemp(Callback)));
+}
+
+void IUAFAssetPreview::RequestAsyncLoad(const FSoftObjectPath& SoftObjectPath, FStreamableDelegate DelegateToCall)
+{
+	CancelAsyncLoad();
+	
+	if (UObject* LoadedObject = SoftObjectPath.ResolveObject())
+	{
+		// Already loaded, nothing to do
+		DelegateToCall.ExecuteIfBound();
+		return; 
+	}
+
+	TWeakPtr<SWidget> ThisWeak = AsShared();
+	StreamingObjectPath = SoftObjectPath;
+
+	auto AsyncLoadCompleteCallback = [ThisWeak, DelegateToCall, SoftObjectPath]()
+	{
+		if (TSharedPtr<IUAFAssetPreview> ThisPinned = StaticCastSharedPtr<IUAFAssetPreview>(ThisWeak.Pin()))
+		{
+			// If the object paths don't match, then this delegate was interrupted, so ignore everything and abort.
+			if (ThisPinned->StreamingObjectPath != SoftObjectPath)
+			{
+				return;
+			}
+
+			DelegateToCall.ExecuteIfBound();
+		}
+	};
+
+	StreamingHandle = UAssetManager::GetStreamableManager().RequestAsyncLoad(StreamingObjectPath
+		, AsyncLoadCompleteCallback
+		, FStreamableManager::AsyncLoadHighPriority);
+}
+
+void IUAFAssetPreview::CancelAsyncLoad()
+{
+	if (StreamingHandle.IsValid())
+	{
+		StreamingHandle->CancelHandle();
+		StreamingHandle.Reset();
+	}
+
+	StreamingObjectPath.Reset();
+}
+
+void IUAFAssetPreview::OnConstructAsyncPreviewWidget(const FAssetData& InAssetData)
+{
+	TWeakPtr<SWidget> ThisWeak = AsShared();
+	auto UpdateChildSlotCallback = [ThisWeak]()
+	{
+		if (TSharedPtr<IUAFAssetPreview> ThisPinned = StaticCastSharedPtr<IUAFAssetPreview>(ThisWeak.Pin()))
+		{
+			ThisPinned->ChildSlot
+			[
+				ThisPinned->OnConstructAssetPreviewWidget()
+			];
+		}
+	};
+
+	if (InAssetData.IsAssetLoaded())
+	{
+		UpdateChildSlotCallback();
+	}
+	else
+	{
+		ChildSlot
+		[
+			OnConstructLoadingAssetPreviewWidget()
+		];
+
+		// Not done loading, kick off an async load
+		RequestAsyncLoad(InAssetData.GetSoftObjectPath(), FStreamableDelegate::CreateLambda(MoveTemp(UpdateChildSlotCallback)));
+	}
+}
+
+TSharedRef<SWidget> IUAFAssetPreview::OnConstructLoadingAssetPreviewWidget()
+{
+	if (CachedLoadingAssetPreviewWidget)
+	{
+		return CachedLoadingAssetPreviewWidget.ToSharedRef();
+	}
+
+	// @TODO Fancy spinning circle? Hourglass Icon?
+	return SAssignNew(CachedLoadingAssetPreviewWidget, SVerticalBox)
+		+ SVerticalBox::Slot()
+		.HAlign(HAlign_Center)
+		.VAlign(VAlign_Center)
+		.FillHeight(1.0f)
+		[
+			SNew(STextBlock)
+				.Justification(ETextJustify::Center)
+				.AutoWrapText(true)
+				.Text(LOCTEXT("AssetPreviewLoading", "Loading Asset Preview..."))
+		];
+}
+
+TSharedRef<SWidget> IUAFAssetPreview::OnConstructAssetPreviewWidget()
+{
+	ensure(CachedCurrentAssetData.IsAssetLoaded());
+	return SNullWidget::NullWidget;
+}
+
+
+//////////////////////////////////////////////////////////////////////////
+
+} // namespace UE::UAF::Editor
+
+#undef LOCTEXT_NAMESPACE
+
diff --git a/Engine/Plugins/Experimental/UAF/UAF/Source/UAFEditor/Private/Common/AssetPreview/IUAFAssetViewportPreview.cpp b/Engine/Plugins/Experimental/UAF/UAF/Source/UAFEditor/Private/Common/AssetPreview/IUAFAssetViewportPreview.cpp
new file mode 100644
--- /dev/null
+++ b/Engine/Plugins/Experimental/UAF/UAF/Source/UAFEditor/Private/Common/AssetPreview/IUAFAssetViewportPreview.cpp
@@ -0,0 +1,80 @@
+// Copyright Epic Games, Inc. All Rights Reserved.
+
+#include "Common/AssetPreview/IUAFAssetViewportPreview.h"
+
+#include "AssetEditorViewportLayout.h"
+#include "Engine/AssetManager.h"
+#include "SceneInterface.h"
+#include "Slate/SceneViewport.h"
+#include "Viewports.h"
+
+#define LOCTEXT_NAMESPACE "IUAFAssetViewportPreview"
+
+namespace UE::UAF::Editor
+{
+
+
+//////////////////////////////////////////////////////////////////////////
+// IUAFAssetViewportPreview
+
+TSharedRef<SWidget> IUAFAssetViewportPreview::OnConstructAssetPreviewWidget()
+{
+	using namespace UE::Workspace;
+
+	ensure(CachedCurrentAssetData.IsAssetLoaded());
+	if (ViewportWidget)
+	{
+		return ViewportWidget.ToSharedRef();
+	}
+
+	PreviewScene = MakeUnique<FAdvancedPreviewScene>(FPreviewScene::ConstructionValues());
+
+	TSharedPtr<FEditorViewportClient> EditorViewportClient = MakeShared<FEditorViewportClient>(nullptr, PreviewScene.Get());
+	EditorViewportClient->ViewportType = LVT_Perspective;
+	EditorViewportClient->bSetListenerPosition = false;
+	EditorViewportClient->SetViewLocation(FVector(0.0f, 400.0f, 200.0f));
+	EditorViewportClient->SetViewRotation(EditorViewportDefs::DefaultPerspectiveViewRotation);
+	EditorViewportClient->SetRealtime(true);
+	EditorViewportClient->SetViewMode(VMI_Lit);
+	EditorViewportClient->ToggleOrbitCamera(true);
+
+	SAssignNew(ViewportWidget, SWorkspaceViewport)
+		.AssetEditorToolkit(AssetEditorToolkitPtr.Pin())
+		.ViewportClient(MoveTemp(EditorViewportClient))
+		.SceneDescription(nullptr);
+
+	OnCustomizePreviewScene(*PreviewScene.Get(), *ViewportWidget->GetViewportClient().Get());
+	
+	// Resolve the asset
+	bool bAssetValidForPreview = OnSameTypeAssetSelected(CachedCurrentAssetData);
+	if (bAssetValidForPreview)
+	{
+		RegisterActiveTimer(0.f, FWidgetActiveTimerDelegate::CreateSP(this, &IUAFAssetViewportPreview::UpdatePreviewWorld));
+	}
+
+	return ViewportWidget.ToSharedRef();
+}
+
+EActiveTimerReturnType IUAFAssetViewportPreview::UpdatePreviewWorld(double InCurrentTime, float InDeltaTime)
+{
+	if (PreviewScene.IsValid()
+		&& ViewportWidget->IsParentValid()
+		&& ViewportWidget->IsVisible()
+		&& ViewportWidget->GetVisibility() == EVisibility::Visible)
+	{
+		// Tick the world to update preview viewport for tooltips
+		PreviewScene->GetWorld()->Tick(LEVELTICK_All, InDeltaTime);
+	}
+
+	// We never stop this timer as long as the widget exists. It will no-op if we are not visible
+	// @TODO: Possibly add / remove timer as visibility changes? Consider work after a profile phase.
+	return EActiveTimerReturnType::Continue;
+}
+
+
+//////////////////////////////////////////////////////////////////////////
+
+} // namespace UE::UAF::Editor
+
+#undef LOCTEXT_NAMESPACE
+
diff --git a/Engine/Plugins/Experimental/UAF/UAF/Source/UAFEditor/Private/Common/AssetPreview/UAFAssetPreviewFactorySubsystem.cpp b/Engine/Plugins/Experimental/UAF/UAF/Source/UAFEditor/Private/Common/AssetPreview/UAFAssetPreviewFactorySubsystem.cpp
new file mode 100644
--- /dev/null
+++ b/Engine/Plugins/Experimental/UAF/UAF/Source/UAFEditor/Private/Common/AssetPreview/UAFAssetPreviewFactorySubsystem.cpp
@@ -0,0 +1,40 @@
+// Copyright Epic Games, Inc. All Rights Reserved.
+
+#include "Common/AssetPreview/UAFAssetPreviewFactorySubsystem.h"
+
+#include "Common/AssetPreview/IUAFAssetPreview.h"
+
+#include UE_INLINE_GENERATED_CPP_BY_NAME(UAFAssetPreviewFactorySubsystem)
+
+#define LOCTEXT_NAMESPACE "UAFAssetPreviewFactorySubsystem"
+
+bool UUAFAssetPreviewFactorySubsystem::AddAssetPreviewFactory(TSharedPtr<UE::UAF::Editor::FUAFAssetPreviewFactory> InFactory)
+{
+	if (InFactory)
+	{
+		if (!AssetTypeFactoryMap.Contains(InFactory->GetPreviewType()))
+		{
+			AssetTypeFactoryMap.Add(InFactory->GetPreviewType(), InFactory);
+
+			return true;
+		}
+	}
+
+	return false;
+}
+
+const TSharedPtr<UE::UAF::Editor::FUAFAssetPreviewFactory> UUAFAssetPreviewFactorySubsystem::GetAssetPreviewFactory(const UStruct* InPreviewType) const
+{
+	if (InPreviewType)
+	{
+		if (const TSharedPtr<UE::UAF::Editor::FUAFAssetPreviewFactory>* FactoryPtr = AssetTypeFactoryMap.Find(InPreviewType))
+		{
+			return *FactoryPtr;
+		}
+	}
+
+	return nullptr;
+}
+
+#undef LOCTEXT_NAMESPACE
+
diff --git a/Engine/Plugins/Experimental/UAF/UAF/Source/UAFEditor/Private/Common/SUAFBrowser.cpp b/Engine/Plugins/Experimental/UAF/UAF/Source/UAFEditor/Private/Common/SUAFBrowser.cpp
--- a/Engine/Plugins/Experimental/UAF/UAF/Source/UAFEditor/Private/Common/SUAFBrowser.cpp
+++ b/Engine/Plugins/Experimental/UAF/UAF/Source/UAFEditor/Private/Common/SUAFBrowser.cpp
@@ -5,6 +5,9 @@
 #include "AnimNextRigVMAsset.h"
 #include "AnimNextRigVMAssetEditorData.h"
 #include "Assets/TaggedAssetBrowserConfiguration.h"
+#include "Common/AssetPreview/AnimSequenceAssetPreview.h"
+#include "Common/AssetPreview/IUAFAssetPreview.h"
+#include "Common/AssetPreview/UAFAssetPreviewFactorySubsystem.h"
 #include "Common/UAFBrowserTabSummoner.h"
 #include "ContentBrowserModule.h"
 #include "IContentBrowserSingleton.h"
@@ -32,111 +35,54 @@ void SUAFBrowser::Construct(const FArguments& InArgs, const TSharedPtr<UE::Works
 	WeakHostingApp = InHostingApp;
 	TabState = InTabState;
 
-	// Note: The layout & implementation of this widget may change drastically in future revisions.
-	// Please consider widget implementation details as placeholder for now.
+	FMargin DrawerPadding = TabState == ETabState::SidebarTab || TabState == ETabState::StatusbarTab
+		? 4.0f
+		: 0.0f;
 
-	// @TODO: Configure filter / content browser, will require updating content browser to add new extension points for UAF Browser view
-	FContentBrowserModule& ContentBrowserModule = FModuleManager::Get().LoadModuleChecked<FContentBrowserModule>(TEXT("ContentBrowser"));
+	ChildSlot
+	[
+		SNew(SBorder)
+		.BorderImage(FAppStyle::GetBrush("ToolPanel.GroupBorder"))
+		.Padding(DrawerPadding)
+		[
+			SNew(SOverlay)
+			+SOverlay::Slot()
+			.HAlign(HAlign_Fill)
+			.VAlign(VAlign_Fill)
+			.Padding(DrawerPadding)
+			[
+				Construct_GetMainBrowserContents()
+			]
+		]
+	];
+}
 
+TSharedRef<SWidget> SUAFBrowser::Construct_GetMainBrowserContents()
+{
 	bool bIsOrientedHorizontally = TabState != ETabState::SidebarTab;
-
 	EOrientation Orientation = bIsOrientedHorizontally ? EOrientation::Orient_Horizontal : EOrientation::Orient_Vertical;
 
-	TWeakPtr<SWidget> ThisWeak = AsShared();
-	auto CreateToggleSwapHotkeyButton = [ThisWeak]() -> TSharedRef<SWidget>
-	{
-		if (TSharedPtr<SUAFBrowser> ThisPinned = StaticCastSharedPtr<SUAFBrowser>(ThisWeak.Pin()))
-		{	
-			if (ThisPinned->TabState == ETabState::SidebarTab || ThisPinned->TabState == ETabState::StatusbarTab)
-			{
-				auto GetSwapBrowserIcon = []() -> const FSlateBrush*
-				{
-					if (UUAFEditorUserSettings* UAFEditorUserSettings = GetMutableDefault<UUAFEditorUserSettings>())
-					{
-						return FAppStyle::Get().GetBrush(UAFEditorUserSettings->bSummonUAFBrowserHotkeyOpensSidebar
-							? "EditorViewport.Right"
-							: "EditorViewport.Bottom");
-					}
-
-					return nullptr;
-				};
-
-				return
-					SNew(SButton)
-					.ButtonStyle(FAppStyle::Get(), "SimpleButton")
-					.ToolTipText(LOCTEXT("SwapHotkey_Tooltip", "Swaps UAF Browser Hotkey Between Sidebar & StatusBar."))
-					.ContentPadding(FMargin(2.f, 4.0f, 2.0f, 4.0f))
-					.OnClicked(ThisPinned.Get(), &SUAFBrowser::SwapHotkeyButtonClicked)
-					[
-						SNew(SImage)
-						.ColorAndOpacity(FSlateColor::UseForeground())
-						.Image(TAttribute<const FSlateBrush*>::CreateLambda(GetSwapBrowserIcon))
-					];
-			}
-		}
-		return SNullWidget::NullWidget;
-	};
-	auto CreateDrawerDockButton = [ThisWeak]() -> TSharedRef<SWidget>
-	{
-		if (TSharedPtr<SUAFBrowser> ThisPinned = StaticCastSharedPtr<SUAFBrowser>(ThisWeak.Pin()))
-		{
-			if (ThisPinned->TabState == ETabState::SidebarTab || ThisPinned->TabState == ETabState::StatusbarTab)
-			{
-				return
-					SNew(SButton)
-					.ButtonStyle(FAppStyle::Get(), "SimpleButton")
-					.ToolTipText(LOCTEXT("DockInLayout_Tooltip", "Docks UAF Browser in Tab."))
-					.ContentPadding(FMargin(2.f, 4.0f, 2.0f, 4.0f))
-					.OnClicked(ThisPinned.Get(), &SUAFBrowser::DockDrawerButtonClicked)
-					[
-						SNew(SImage)
-						.ColorAndOpacity(FSlateColor::UseForeground())
-						.Image(FAppStyle::Get().GetBrush("Icons.Layout"))
-					];
-			}
-		}
-		return SNullWidget::NullWidget;
-	};
-	auto ExtendAssetPickerTopBar = [ThisWeak, CreateToggleSwapHotkeyButton, CreateDrawerDockButton](TSharedRef<SHorizontalBox> InAssetPickerTopBar)
-	{
-		if (TSharedPtr<SUAFBrowser> ThisPinned = StaticCastSharedPtr<SUAFBrowser>(ThisWeak.Pin()))
-		{
-			if (ThisPinned->TabState == ETabState::SidebarTab || ThisPinned->TabState == ETabState::StatusbarTab)
-			{
-				// Note: Value from `FAssetPickerConfig::OnExtendAssetPickerTopBar`. 
-				// We could make it expose those indicies, but given expected usage that feels over-engineered for now.
-				const int ShowDeveloperContentIndex = 2;
-
-				InAssetPickerTopBar->InsertSlot(ShowDeveloperContentIndex + 1)
-				.AutoWidth()
-				[
-					CreateDrawerDockButton()
-				];
-
-				InAssetPickerTopBar->InsertSlot(ShowDeveloperContentIndex + 2)
-				.AutoWidth()
-				[
-					CreateToggleSwapHotkeyButton()
-				];
-			}
-		}
-	};
+	return SNew(SSplitter)
+		.Orientation(Orientation)
+		+SSplitter::Slot()
+		.Value(bIsOrientedHorizontally ? 0.85f : 0.15f)
+		[
+			bIsOrientedHorizontally
+				? Construct_GetTaggedAssetBrowser()
+				: Construct_GetAssetPreviewWidget()
+		]
+		+SSplitter::Slot()
+		.Value(bIsOrientedHorizontally ? 0.15f : 0.85f)
+		[
+			bIsOrientedHorizontally
+				? Construct_GetAssetPreviewWidget()
+				: Construct_GetTaggedAssetBrowser()
+		];
+}
 
-	auto CustomizeAssetPicker = [&](FAssetPickerConfig& Config)
-	{
-		Config.OnExtendAssetPickerTopBar.BindLambda(ExtendAssetPickerTopBar);
-		Config.InitialAssetViewType = EAssetViewType::Column;
-		Config.bAddFilterUI = true;
-		//Config.bShowPathInColumnView = true;
-		//Config.bSortByPathInColumnView = true;
-		Config.bAllowDragging = true;
-
-		Config.OnGetAssetContextMenu.Unbind();
-		Config.OnAssetSelected.Unbind();
-		Config.OnAssetsActivated.Unbind();
-		Config.OnGetCustomAssetToolTip.Unbind();
-	};
-		
+TSharedRef<SWidget> SUAFBrowser::Construct_GetTaggedAssetBrowser()
+{
+	// Setup tagged asset browser config
 	STaggedAssetBrowser::FInterfaceOverrideProfiles OverrideProfiles;
 	OverrideProfiles.AssetViewOptionsProfileName = FName("TaggedAssetBrowser");
 	OverrideProfiles.DefaultFilterMenuExpansion = EAssetTypeCategories::Animation;
@@ -145,80 +91,136 @@ void SUAFBrowser::Construct(const FArguments& InArgs, const TSharedPtr<UE::Works
 	FDefaultDetailsTabConfiguration DefaultDetailsTabConfiguration;
 	DefaultDetailsTabConfiguration.bUseDefaultDetailsTab = false;
 
-	auto GetAssetPreviewWidget = [&]()
+	FSoftObjectPath CategoryConfigurationPath = GetDefault<UUAFEditorProjectSettings>()->UAFBrowserCategoryConfiguration;
+	UObject* CategoryConfigurationObject = CategoryConfigurationPath.TryLoad();
+
+	if (UTaggedAssetBrowserConfiguration* CategoryConfiguration = Cast<UTaggedAssetBrowserConfiguration>(CategoryConfigurationObject))
 	{
-		return SNew(SVerticalBox)
-			+SVerticalBox::Slot()
-			.HAlign(HAlign_Center)
-			.VAlign(VAlign_Center)
-			.FillHeight(1.0f)
-			[
-				SNew(STextBlock)
-				.Justification(ETextJustify::Center)
-				.AutoWrapText(true)
-				.Text(LOCTEXT("TempAssetPreviewText", "Asset Preview Here"))
-			];
-	};
-	auto GetAssetBrowser = [&]() -> TSharedRef<SWidget>
+		// We intentionally do not restrict AvailableClasses as this will filter types we may not know about from beind displayed
+		// Ex: Even in contexts such as filtering by tag
+		return SNew(STaggedAssetBrowser, *CategoryConfiguration)
+			.AvailableClasses({})
+			.DefaultDetailsTabConfiguration(DefaultDetailsTabConfiguration)
+			.InterfaceOverrideProfiles(OverrideProfiles)
+			.OnOverrideAssetPickerConfig(STaggedAssetBrowser::FOnOverrideAssetPickerConfig::CreateSP(this, &SUAFBrowser::OnOverrideAssetPickerConfig));
+	}
+
+	return SNullWidget::NullWidget;
+}
+
+TSharedRef<SWidget> SUAFBrowser::Construct_GetAssetPreviewWidget()
+{
+	// This is typically customized via the exposed AssetPreviewSlot. See: OnAssetSelected.
+	return SNew(SVerticalBox)
+		+SVerticalBox::Slot()
+		.HAlign(HAlign_Center)
+		.VAlign(VAlign_Center)
+		.Expose(AssetPreviewSlot)
+		[
+			Construct_GetAssetPreviewNotSupportedWidget()
+		];
+}
+
+TSharedRef<SWidget> SUAFBrowser::Construct_GetAssetPreviewNotSupportedWidget()
+{
+	return SNew(STextBlock)
+		.Justification(ETextJustify::Center)
+		.AutoWrapText(true)
+		.Text(LOCTEXT("AssetPreviewNotSetup", "Asset Preview Not Setup for Type"));
+}
+
+void SUAFBrowser::OnOverrideAssetPickerConfig(FAssetPickerConfig& Config)
+{
+	// @TODO: Meet with product design & choose defaults
+	Config.OnExtendAssetPickerTopBar.BindSP(this, &SUAFBrowser::OnExtendAssetPickerTopBar);
+	Config.InitialAssetViewType = EAssetViewType::Column;
+	Config.bAddFilterUI = true;
+	Config.bAllowDragging = true;
+	//Config.bShowPathInColumnView = true;
+	//Config.bSortByPathInColumnView = true;
+
+	Config.OnGetAssetContextMenu.Unbind();
+	Config.OnAssetSelected.BindSP(this, &SUAFBrowser::OnAssetSelected);
+	Config.OnAssetsActivated.Unbind();
+	Config.OnGetCustomAssetToolTip.Unbind();
+}
+
+void SUAFBrowser::OnExtendAssetPickerTopBar(TSharedRef<SHorizontalBox> InAssetPickerTopBar)
+{
+	if (TabState == ETabState::SidebarTab || TabState == ETabState::StatusbarTab)
 	{
-		FSoftObjectPath CategoryConfigurationPath = GetDefault<UUAFEditorProjectSettings>()->UAFBrowserCategoryConfiguration;
-		UObject* CategoryConfigurationObject = CategoryConfigurationPath.TryLoad();
-		if (UTaggedAssetBrowserConfiguration* CategoryConfiguration = Cast<UTaggedAssetBrowserConfiguration>(CategoryConfigurationObject))
-		{
-			// We intentionally do not restrict AvailableClasses as this will filter types we may not know about from beind displayed
-			// Ex: Even in contexts such as filtering by tag
-			return SNew(STaggedAssetBrowser, *CategoryConfiguration)
-				.AvailableClasses({})
-				.DefaultDetailsTabConfiguration(DefaultDetailsTabConfiguration)
-				.InterfaceOverrideProfiles(OverrideProfiles)
-				.OnOverrideAssetPickerConfig(STaggedAssetBrowser::FOnOverrideAssetPickerConfig::CreateLambda(CustomizeAssetPicker));
-		}
+		// Note: Value from `FAssetPickerConfig::OnExtendAssetPickerTopBar`. 
+		// We could make it expose those indicies, but given expected usage that feels over-engineered for now.
+
+		const int ShowDeveloperContentIndex = 2;
 
-		return SNullWidget::NullWidget;
-	};
-	auto GetMainBrowserContents = [&]()
+		InAssetPickerTopBar->InsertSlot(ShowDeveloperContentIndex + 1)
+		.AutoWidth()
+		[
+			OnExtendAssetPickerTopBar_GetDrawerDockButton()
+		];
+
+		InAssetPickerTopBar->InsertSlot(ShowDeveloperContentIndex + 2)
+		.AutoWidth()
+		[
+			OnExtendAssetPickerTopBar_GetToggleSwapHotkeyButton()
+		];
+	}
+}
+
+TSharedRef<SWidget> SUAFBrowser::OnExtendAssetPickerTopBar_GetDrawerDockButton()
+{
+	if (TabState == ETabState::SidebarTab || TabState == ETabState::StatusbarTab)
 	{
-		return SNew(SSplitter)
-			.Orientation(Orientation)
-			+SSplitter::Slot()
-			.Value(bIsOrientedHorizontally ? 0.85f : 0.15f)
+		return
+			SNew(SButton)
+			.ButtonStyle(FAppStyle::Get(), "SimpleButton")
+			.ToolTipText(LOCTEXT("DockInLayout_Tooltip", "Docks UAF Browser in Tab."))
+			.ContentPadding(FMargin(2.f, 4.0f, 2.0f, 4.0f))
+			.OnClicked(this, &SUAFBrowser::OnDockDrawerButtonClicked)
 			[
-				bIsOrientedHorizontally
-					? GetAssetBrowser()
-					: GetAssetPreviewWidget()
-			]
-			+SSplitter::Slot()
-			.Value(bIsOrientedHorizontally ? 0.15f : 0.85f)
-			[
-				bIsOrientedHorizontally
-					? GetAssetPreviewWidget()
-					: GetAssetBrowser()
+				SNew(SImage)
+				.ColorAndOpacity(FSlateColor::UseForeground())
+				.Image(FAppStyle::Get().GetBrush("Icons.Layout"))
 			];
-	};
+	}
 	
-	FMargin DrawerPadding = TabState == ETabState::SidebarTab || TabState == ETabState::StatusbarTab
-		? 4.0f
-		: 0.0f;
+	return SNullWidget::NullWidget; 
+}
 
-	ChildSlot
-	[
-		SNew(SBorder)
-		.BorderImage(FAppStyle::GetBrush("ToolPanel.GroupBorder"))
-		.Padding(DrawerPadding)
-		[
-			SNew(SOverlay)
-			+SOverlay::Slot()
-			.HAlign(HAlign_Fill)
-			.VAlign(VAlign_Fill)
-			.Padding(DrawerPadding)
+TSharedRef<SWidget> SUAFBrowser::OnExtendAssetPickerTopBar_GetToggleSwapHotkeyButton()
+{
+	if (TabState == ETabState::SidebarTab || TabState == ETabState::StatusbarTab)
+	{
+		auto GetSwapBrowserIcon = []() -> const FSlateBrush*
+		{
+			if (UUAFEditorUserSettings* UAFEditorUserSettings = GetMutableDefault<UUAFEditorUserSettings>())
+			{
+				return FAppStyle::Get().GetBrush(UAFEditorUserSettings->bSummonUAFBrowserHotkeyOpensSidebar
+					? "EditorViewport.Right"
+					: "EditorViewport.Bottom");
+			}
+
+			return nullptr;
+		};
+
+		return
+			SNew(SButton)
+			.ButtonStyle(FAppStyle::Get(), "SimpleButton")
+			.ToolTipText(LOCTEXT("SwapHotkey_Tooltip", "Swaps UAF Browser Hotkey Between Sidebar & StatusBar."))
+			.ContentPadding(FMargin(2.f, 4.0f, 2.0f, 4.0f))
+			.OnClicked(this, &SUAFBrowser::OnSwapHotkeyButtonClicked)
 			[
-				GetMainBrowserContents()
-			]
-		]
-	];
+				SNew(SImage)
+				.ColorAndOpacity(FSlateColor::UseForeground())
+				.Image(TAttribute<const FSlateBrush*>::CreateLambda(GetSwapBrowserIcon))
+			];
+	}
+
+	return SNullWidget::NullWidget;
 }
 
-FReply SUAFBrowser::SwapHotkeyButtonClicked()
+FReply SUAFBrowser::OnSwapHotkeyButtonClicked()
 {
 	if (TSharedPtr<UE::Workspace::IWorkspaceEditor> HostingApp = WeakHostingApp.Pin())
 	{
@@ -233,7 +235,7 @@ FReply SUAFBrowser::SwapHotkeyButtonClicked()
 	return FReply::Handled();
 }
 
-FReply SUAFBrowser::DockDrawerButtonClicked()
+FReply SUAFBrowser::OnDockDrawerButtonClicked()
 {
 	if (TSharedPtr<UE::Workspace::IWorkspaceEditor> HostingApp = WeakHostingApp.Pin())
 	{
@@ -248,6 +250,47 @@ FReply SUAFBrowser::DockDrawerButtonClicked()
 	return FReply::Handled();
 }
 
+void SUAFBrowser::OnAssetSelected(const FAssetData& InAssetData)
+{
+	if (UClass* AssetClass = InAssetData.GetClass())
+	{
+		// Handle selection of same types
+		if (TSharedPtr<IUAFAssetPreview> AssetPreviewWidget = AssetPreviewWidgetWeak.Pin())
+		{
+			if (AssetClass == AssetPreviewWidget->GetAssetPreviewType())
+			{
+				AssetPreviewWidget->OnSameTypeAssetSelectedAsync(InAssetData);
+				return;
+			}
+		}
+
+		// Handle selection of different or unsupported types
+		if (TSharedPtr<FUAFAssetPreviewFactory> AssetPreviewFactory = GEditor->GetEditorSubsystem<UUAFAssetPreviewFactorySubsystem>()->GetAssetPreviewFactory(AssetClass))
+		{
+			TSharedPtr<IUAFAssetPreview> AssetPreview = AssetPreviewFactory->CreateAssetPreviewWidget(WeakHostingApp.Pin(), InAssetData);
+			AssetPreviewSlot->AttachWidget(
+				AssetPreview.ToSharedRef()
+			);
+			AssetPreviewSlot->SetHorizontalAlignment(HAlign_Fill);
+			AssetPreviewSlot->SetVerticalAlignment(VAlign_Fill);
+
+			// We don't attempt to persist the preview between different asset type clicks. Doing so would save perf if clicking:
+			// type -> unsupported preview type -> previous type. But that is an edge case w/ overall low cost.
+			AssetPreviewWidgetWeak = AssetPreview;
+		}
+		else
+		{
+			AssetPreviewSlot->AttachWidget(
+				Construct_GetAssetPreviewNotSupportedWidget()
+			);
+			AssetPreviewSlot->SetHorizontalAlignment(HAlign_Center);
+			AssetPreviewSlot->SetVerticalAlignment(VAlign_Center);
+
+			AssetPreviewWidgetWeak = nullptr;
+		}
+	}
+}
+
 } // namespace UE::UAF::Editor
 
 #undef LOCTEXT_NAMESPACE
diff --git a/Engine/Plugins/Experimental/UAF/UAF/Source/UAFEditor/Private/Common/SUAFBrowser.h b/Engine/Plugins/Experimental/UAF/UAF/Source/UAFEditor/Private/Common/SUAFBrowser.h
--- a/Engine/Plugins/Experimental/UAF/UAF/Source/UAFEditor/Private/Common/SUAFBrowser.h
+++ b/Engine/Plugins/Experimental/UAF/UAF/Source/UAFEditor/Private/Common/SUAFBrowser.h
@@ -5,6 +5,7 @@
 #include "CoreMinimal.h"
 #include "Framework/Docking/TabManager.h"
 #include "Widgets/Docking/SDockTab.h"
+#include "Widgets/Layout/SSplitter.h"
 #include "Widgets/SCompoundWidget.h"
 
 namespace UE::Workspace
@@ -12,14 +13,22 @@ namespace UE::Workspace
 	class IWorkspaceEditor;
 }	// end namespace UE::Workspace
 
+namespace UE::UAF::Editor
+{
+	class IUAFAssetPreview;
+};
+
+struct FAssetPickerConfig;
+class SHorizontalBox;
+
 namespace UE::UAF::Editor
 {
 
 /** 
  * Experimental. Widget for searching for UAF related assets
  * 
- * Note: The layout & implementation of this widget may change drastically in future revisions.
- * Please consider widget implementation details as placeholder for now.
+ * Note: The implementation of this widget may change in the future as we move to use an in-progress capability of the content
+ * browser to preview assets.
  */
 class SUAFBrowser : public SCompoundWidget
 {
@@ -37,11 +46,42 @@ public:
 
 private:
 
-	/** Callback for when swap hotkey is pressed. Chnages hotkey behavior in project settings & hides all drawers. */
-	FReply SwapHotkeyButtonClicked();
+	/** Builds the top level of the UAF Browser, including preview splitter. Contents need to be ordered dynamically based on orientation. */
+	TSharedRef<SWidget> Construct_GetMainBrowserContents();
+
+	/** Builds the tagged asset browser part of the UAF browser. IE: Asset picker */
+	TSharedRef<SWidget> Construct_GetTaggedAssetBrowser();
+
+	/** Builds the tagged asset preview part of the UAF browser. Ex: Viewport. In the future this may get replaced by a native content browser API. */
+	TSharedRef<SWidget> Construct_GetAssetPreviewWidget();
+
+	/** Builds fallack when preview not supported by type in UAF browser Preview. */
+	TSharedRef<SWidget> Construct_GetAssetPreviewNotSupportedWidget();
+
+private:
+
+	/** Callback to override asset picker configs for our tagged asset browser */
+	void OnOverrideAssetPickerConfig(FAssetPickerConfig& Config);
+
+	/** Callback to customize top bar of our asset browser. To add layout control buttons. */
+	void OnExtendAssetPickerTopBar(TSharedRef<SHorizontalBox> InAssetPickerTopBar);
+
+	/** Builds button for docking the UAF Browser. This just unhides the in-layout browser tab. */
+	TSharedRef<SWidget> OnExtendAssetPickerTopBar_GetDrawerDockButton();
+
+	/** Builds button for swapping hotkey behavior (sidebar to statusbar). */
+	TSharedRef<SWidget> OnExtendAssetPickerTopBar_GetToggleSwapHotkeyButton();
+
+	/** Callback for when swap hotkey is pressed. Changes hotkey behavior in project settings & hides all drawers. */
+	FReply OnSwapHotkeyButtonClicked();
 
 	/** Callback for when dock to drawer button is pressed. Hides all drawers and invokes docked browser tab. */
-	FReply DockDrawerButtonClicked();
+	FReply OnDockDrawerButtonClicked();
+
+private:
+
+	/** Callback for when an asset is selected. */
+	void OnAssetSelected(const FAssetData& InAssetData);
 
 private:
 
@@ -50,6 +90,18 @@ private:
 
 	/** Current Tab host type, so we can resturucture layout based on how this tab is being hosted (Sidebar / statusbar / etc). */
 	ETabState::Type TabState = ETabState::ClosedTab;
+
+	/** 
+	 * Experimental. Slot used to swap between different asset preview widgets. 
+	 * Note: Long term this functionality will get replaced by the generic IAssetPreviewCustomization API.
+	 */
+	SVerticalBox::FSlot* AssetPreviewSlot = nullptr;
+
+	/**
+	 * Experimental. Widget used to preview an asset. Persisted weakly here so that we can tell it to just
+	 * update if swapping betweeen assets of the same type
+	 */
+	TWeakPtr<IUAFAssetPreview> AssetPreviewWidgetWeak = nullptr;
 };
 
 } // namespace UE::UAF::Editor
\ No newline at end of file
diff --git a/Engine/Plugins/Experimental/UAF/UAF/Source/UAFEditor/Public/Common/AssetPreview/IUAFAssetPreview.h b/Engine/Plugins/Experimental/UAF/UAF/Source/UAFEditor/Public/Common/AssetPreview/IUAFAssetPreview.h
new file mode 100644
--- /dev/null
+++ b/Engine/Plugins/Experimental/UAF/UAF/Source/UAFEditor/Public/Common/AssetPreview/IUAFAssetPreview.h
@@ -0,0 +1,131 @@
+// Copyright Epic Games, Inc. All Rights Reserved.
+
+#pragma once
+
+#include "CoreMinimal.h"
+#include "EditorViewportClient.h"
+#include "Engine/StreamableManager.h"
+#include "Widgets/SWidget.h"
+#include "Widgets/SCompoundWidget.h"
+#include "Toolkits/AssetEditorToolkit.h"
+
+#define UE_API UAFEDITOR_API
+
+class FAssetEditorToolkit;
+class FSceneViewport;
+struct FStreamableHandle;
+
+namespace UE::UAF::Editor
+{
+	class IUAFAssetPreview;
+};
+
+/**
+ * Note: In this experimental state we frequently mention an 'IAssetPreviewCustomization' API. At a high level, imagine this as an in-progress future API.
+ * Similar to detail customizations where you can register a widget for some particular type & it automatically gets used (Ex: In content browser previews).
+ * Our goal is to make a MVP that fully explores the space of what we want before we integrate with this API.
+ */
+namespace UE::UAF::Editor
+{
+
+
+/**
+ * Experimental. May be removed whenever we switch to the future generic content browser IAssetPreviewCustomization API.
+ *
+ * Base struct for factories that can create IUAFAssetPreview widgets.
+ * 
+ * @TODO: Reconsider the type constraint. It is okay but some factories may support multiple types. 
+ * Should we allow for multiple factories for the same type based on a priority system?
+ */
+struct FUAFAssetPreviewFactory
+{
+public:
+
+	/** Callback used to create the preview for the given asset data. */
+	UE_API virtual TSharedRef<IUAFAssetPreview> CreateAssetPreviewWidget(TSharedPtr<FAssetEditorToolkit> InAssetEditorToolkit, const FAssetData& InAssetData) const = 0;
+
+	/** Get the preview type this factory is used for. */
+	UE_API virtual const UStruct* GetPreviewType() const = 0;
+};
+
+
+/** 
+ * Experimental. May be removed whenever we switch to the future generic content browser IAssetPreviewCustomization API. 
+ * 
+ * Base class for asset preview widgets. Currently previews are restricted to UAF, but may not be in the future.
+ * 
+ * @TODO: This class currently supports individual async loading, but users may want to async preload a set of assets for responsiveness.
+ * @TODO: Once we use the future content browser API, consider removing the "UAF" aspect of this class. Maybe move to engine as IAsyncAssetPreview? Need editor buy in first.
+ */
+class IUAFAssetPreview : public SCompoundWidget
+{
+public:
+
+	SLATE_BEGIN_ARGS(IUAFAssetPreview) {}
+	SLATE_END_ARGS()
+
+	UE_API virtual void Construct(const FArguments& InArgs, TSharedPtr<FAssetEditorToolkit> InAssetEditorToolkit, const FAssetData& InAssetData);
+
+public:
+
+	/** 
+	 * Callback for when another asset of the same type is selected. Will load the asset async before updating preview.
+	 * When different assets are selected this preview should be replaced. 
+	 */
+	UE_API virtual void OnSameTypeAssetSelectedAsync(const FAssetData& InAssetData);
+
+	/** Utility method to get the current preview type this widget is associated with. */
+	UE_API virtual const UStruct* GetAssetPreviewType() const = 0;
+
+	/** Utility method to load an asset async with function callback */
+	UE_API virtual void RequestAsyncLoad(const FSoftObjectPath& SoftObject, TFunction<void()>&& Callback);
+
+	/** Utility method to load an asset async with a delegate callback */
+	UE_API virtual void RequestAsyncLoad(const FSoftObjectPath& SoftObject, FStreamableDelegate DelegateToCall);
+
+	/** Utility method to cancel async load of asset to preview */
+	UE_API virtual void CancelAsyncLoad();
+
+protected:
+
+	/** 
+	 * Callback for when another asset of the same type is selected & loaded. 
+	 * 
+	 * @return True if we were able to set up a valid preview for the given asset
+	 */
+	UE_API virtual bool OnSameTypeAssetSelected(const FAssetData& InAssetData) = 0;
+
+	/** Called to show a loading widget when the asset is loading, and a preview widget when the asset is done loading. */
+	UE_API virtual void OnConstructAsyncPreviewWidget(const FAssetData& InAssetData);
+
+	/** Callback used to construct preview widget when async loading is occuring. Results should be stored in CachedLoadingAssetPreviewWidget. */
+	UE_API virtual TSharedRef<SWidget> OnConstructLoadingAssetPreviewWidget();
+
+	/** 
+	 * Callback used to construct preview widget whenever async loading is done or an asset is already loaded. 
+	 * Implementers are encouraged to cache their results on first create to avoid widget recreation overhead. 
+	 */
+	UE_API virtual TSharedRef<SWidget> OnConstructAssetPreviewWidget();
+
+protected:
+
+	/** Pointer back to the editor that owns us */
+	TWeakPtr<FAssetEditorToolkit> AssetEditorToolkitPtr;
+
+	/** Last asset data previewed by this widget */
+	FAssetData CachedCurrentAssetData;
+
+	/** Handle for currently streaming asset, we load assets async for previews by default. */
+	TSharedPtr<FStreamableHandle> StreamingHandle;
+
+	/** Path for currently streaming asset */
+	FSoftObjectPath StreamingObjectPath;
+
+	/** Cached loading widget. Saved so we can communicate with it and don't need to re-create it uncessarily. */
+	TSharedPtr<SWidget> CachedLoadingAssetPreviewWidget;
+};
+
+
+} // namespace UE::UAF::Editor
+
+#undef UE_API
diff --git a/Engine/Plugins/Experimental/UAF/UAF/Source/UAFEditor/Public/Common/AssetPreview/IUAFAssetViewportPreview.h b/Engine/Plugins/Experimental/UAF/UAF/Source/UAFEditor/Public/Common/AssetPreview/IUAFAssetViewportPreview.h
new file mode 100644
--- /dev/null
+++ b/Engine/Plugins/Experimental/UAF/UAF/Source/UAFEditor/Public/Common/AssetPreview/IUAFAssetViewportPreview.h
@@ -0,0 +1,70 @@
+// Copyright Epic Games, Inc. All Rights Reserved.
+
+#pragma once
+
+#include "Common/AssetPreview/IUAFAssetPreview.h"
+#include "AdvancedPreviewScene.h"
+#include "SWorkspaceViewport.h"
+
+#define UE_API UAFEDITOR_API
+
+class FEditorViewportClient;
+class FSceneViewport;
+
+/**
+ * Note: In this experimental state we frequently mention an 'IAssetPreviewCustomization' API. At a high level, imagine this as an in-progress future API.
+ * Similar to detail customizations where you can register a widget for some particular type & it automatically gets used (Ex: In content browser previews).
+ * Our goal is to make a MVP that fully explores the space of what we want before we integrate with this API.
+ */
+namespace UE::UAF::Editor
+{
+
+
+/** 
+ * Experimental. May be removed / relocated whenever we switch to the future generic content browser IAssetPreviewCustomization API. 
+ * 
+ * Base class for asset previews that want to use an viewport and don't want to roll their own viewport implementation
+ * @TODO: Same consideration as above. Later on consider removing "UAF" aspect from this.
+ */
+class IUAFAssetViewportPreview : public IUAFAssetPreview
+{
+
+protected:
+
+	//~ Begin IUAFAssetPreview Interface
+	UE_API virtual TSharedRef<SWidget> OnConstructAssetPreviewWidget() override;
+	//~ End IUAFAssetPreview Interface
+
+protected:
+
+	/** 
+	 * Require child classes to customize preview scene / viewport. These are passed together as viewport customizations
+	 * may depend on the preview scene.
+	 * 
+	 * Note: Params are passed for convience. They are technically not needed since child classes have access.
+	 * 
+	 * @param InPreviewScene - Same as our local preview scene. Implementation should set this up to some meaningful asset preview.
+	 * @param InEditorViewportClient - Viewport client for camera setup and other controls. 
+	 */
+	UE_API virtual void OnCustomizePreviewScene(FAdvancedPreviewScene& InPreviewScene, FEditorViewportClient& InEditorViewportClient) = 0;
+
+protected:
+
+	/** The actual viewport widget. Used as cached result for OnConstructAssetPreviewWidget. */
+	TSharedPtr<UE::Workspace::SWorkspaceViewport> ViewportWidget;
+
+	/** The scene viewport data */
+	TSharedPtr<FSceneViewport> SceneViewport;
+
+	/** The scene to show in the asset previews */
+	TUniquePtr<FAdvancedPreviewScene> PreviewScene;
+
+protected:
+
+	/** Callback to tick the preview world */
+	UE_API virtual EActiveTimerReturnType UpdatePreviewWorld(double InCurrentTime, float InDeltaTime);
+};
+
+} // namespace UE::UAF::Editor
+
+#undef UE_API
diff --git a/Engine/Plugins/Experimental/UAF/UAF/Source/UAFEditor/Public/Common/AssetPreview/UAFAssetPreviewFactorySubsystem.h b/Engine/Plugins/Experimental/UAF/UAF/Source/UAFEditor/Public/Common/AssetPreview/UAFAssetPreviewFactorySubsystem.h
new file mode 100644
--- /dev/null
+++ b/Engine/Plugins/Experimental/UAF/UAF/Source/UAFEditor/Public/Common/AssetPreview/UAFAssetPreviewFactorySubsystem.h
@@ -0,0 +1,48 @@
+// Copyright Epic Games, Inc. All Rights Reserved.
+
+#pragma once
+
+#include "CoreMinimal.h"
+#include "Common/AssetPreview/IUAFAssetPreview.h"
+
+#include "UAFAssetPreviewFactorySubsystem.generated.h"
+
+#define UE_API UAFEDITOR_API
+
+namespace UE::UAF::Editor
+{
+	struct FUAFAssetPreviewFactory;
+};
+
+/**
+ * Experimental & Intended for removal. 
+ * 
+ * This class associates asset types to preview widgets for the UAF browser.
+ * However it is intended to be replaced with a future API within the editor / content browser that does the same.
+ * 
+ * See: IUAFAssetPreview.h for more future design considerations.
+ */
+UCLASS(MinimalAPI, Experimental)
+class UUAFAssetPreviewFactorySubsystem : public UEditorSubsystem
+{
+	GENERATED_BODY()
+
+public:
+
+	/** 
+	 * Method to associate a preview factory with some type. For now only supports one factory per type.
+	 * 
+	 * @return True if the factory was added for the type, false if we already have a type for it
+	 */
+	UE_API bool AddAssetPreviewFactory(TSharedPtr<UE::UAF::Editor::FUAFAssetPreviewFactory> InFactory);
+
+	/** Get a Factory for the given type, or nullptr if no factory exists. */
+	UE_API const TSharedPtr<UE::UAF::Editor::FUAFAssetPreviewFactory> GetAssetPreviewFactory(const UStruct* InPreviewType) const;
+
+private:
+
+	/** Map of classes to preview factories */
+	TMap<const UStruct*, TSharedPtr<UE::UAF::Editor::FUAFAssetPreviewFactory>> AssetTypeFactoryMap;
+};
+
+#undef UE_API
diff --git a/Engine/Plugins/Experimental/UAF/UAF/Source/UAFEditor/Public/UAF/Viewport/ViewportSceneDescription.h b/Engine/Plugins/Experimental/UAF/UAF/Source/UAFEditor/Public/UAF/Viewport/ViewportSceneDescription.h
--- a/Engine/Plugins/Experimental/UAF/UAF/Source/UAFEditor/Public/UAF/Viewport/ViewportSceneDescription.h
+++ b/Engine/Plugins/Experimental/UAF/UAF/Source/UAFEditor/Public/UAF/Viewport/ViewportSceneDescription.h
@@ -16,7 +16,7 @@ namespace UE::Workspace
 	class IWorkspaceEditor;
 }
 
-UCLASS()
+UCLASS(MinimalAPI)
 class UUAFViewportSceneDescription : public UWorkspaceViewportSceneDescription
 {
 	GENERATED_BODY()
diff --git a/Engine/Plugins/Experimental/UAF/UAF/Source/UAFEditor/UAFEditor.Build.cs b/Engine/Plugins/Experimental/UAF/UAF/Source/UAFEditor/UAFEditor.Build.cs
--- a/Engine/Plugins/Experimental/UAF/UAF/Source/UAFEditor/UAFEditor.Build.cs
+++ b/Engine/Plugins/Experimental/UAF/UAF/Source/UAFEditor/UAFEditor.Build.cs
@@ -20,6 +20,7 @@ namespace UnrealBuildTool.Rules
 					"UnrealEd",
 					"SlateCore",
 					"AnimationCore",
+					"AnimGraph",
 					"Slate",
 					"InputCore",
 					"PropertyEditor",
@@ -50,6 +51,7 @@ namespace UnrealBuildTool.Rules
 					"AdvancedWidgets", 
 					"SceneOutliner",
 					"EditorFramework",
+					"EditorSubsystem",
 					"InteractiveToolsFramework",
 					"EditorInteractiveToolsFramework",
 					"TraceAnalysis",
diff --git a/Engine/Plugins/Experimental/UAF/UAFAnimGraph/Source/UAFAnimGraphEditor/Private/AnimNextAnimGraphEditorModule.cpp b/Engine/Plugins/Experimental/UAF/UAFAnimGraph/Source/UAFAnimGraphEditor/Private/AnimNextAnimGraphEditorModule.cpp
--- a/Engine/Plugins/Experimental/UAF/UAFAnimGraph/Source/UAFAnimGraphEditor/Private/AnimNextAnimGraphEditorModule.cpp
+++ b/Engine/Plugins/Experimental/UAF/UAFAnimGraph/Source/UAFAnimGraphEditor/Private/AnimNextAnimGraphEditorModule.cpp
@@ -2,6 +2,7 @@
 
 #include "AdvancedPreviewScene.h"
 #include "AnimNextAnimGraphWorkspaceAssetUserData.h"
+#include "AssetPreview/UAFGraphAssetPreview.h"
 #include "Modules/ModuleInterface.h"
 #include "Modules/ModuleManager.h"
 #include "ISettingsModule.h"
@@ -30,6 +31,8 @@
 #include "Graph/AnimGraphEditorSchemaActions.h"
 #include "Graph/AnimNextAnimationGraphSchema.h"
 #include "Common/SActionMenu.h"
+#include "Common/AssetPreview/IUAFAssetPreview.h"
+#include "Common/AssetPreview/UAFAssetPreviewFactorySubsystem.h"
 #include "Graph/RigUnit_AnimNextTraitStack.h"
 #include "AnimNextTraitStackUnitNode.h"
 #include "RigVMCore/RigVMFunction.h"
@@ -100,6 +103,15 @@ class FAnimNextAnimGraphEditorModule : public IModuleInterface
 		IModularFeatures::Get().RegisterModularFeature(RewindDebugger::IRewindDebuggerTrackCreator::ModularFeatureName, &GAnimNextModulesTrackCreator);
 		IModularFeatures::Get().RegisterModularFeature(RewindDebugger::IRewindDebuggerTrackCreator::ModularFeatureName, &GSequenceInfoTrackCreator);
 
+		// Temp register asset previews while we still use the UUAFAssetPreviewFactorySubsystem
+		FCoreDelegates::OnPostEngineInit.AddLambda([]()
+		{
+			if (GEditor)
+			{
+				GEditor->GetEditorSubsystem<UUAFAssetPreviewFactorySubsystem>()->AddAssetPreviewFactory(MakeShared<FUAFGraphAssetPreviewFactory>());
+			}
+		});
+
 		Workspace::IWorkspaceEditorModule& WorkspaceEditorModule = FModuleManager::Get().LoadModuleChecked<Workspace::IWorkspaceEditorModule>("WorkspaceEditor");
 
 		// --- AnimNextAnimationGraph ---
diff --git a/Engine/Plugins/Experimental/UAF/UAFAnimGraph/Source/UAFAnimGraphEditor/Private/AssetPreview/UAFGraphAssetPreview.cpp b/Engine/Plugins/Experimental/UAF/UAFAnimGraph/Source/UAFAnimGraphEditor/Private/AssetPreview/UAFGraphAssetPreview.cpp
new file mode 100644
--- /dev/null
+++ b/Engine/Plugins/Experimental/UAF/UAFAnimGraph/Source/UAFAnimGraphEditor/Private/AssetPreview/UAFGraphAssetPreview.cpp
@@ -0,0 +1,112 @@
+// Copyright Epic Games, Inc. All Rights Reserved.
+
+#include "UAFGraphAssetPreview.h"
+
+#include "Animation/AnimationAsset.h"
+#include "Component/AnimNextComponent.h"
+#include "Components/SkeletalMeshComponent.h"
+#include "GameFramework/Actor.h"
+#include "Graph/AnimNextAnimationGraph.h"
+#include "IWorkspaceEditor.h"
+#include "Module/AnimNextModule.h"
+#include "UAF/Viewport/ViewportSceneDescription.h"
+
+#define LOCTEXT_NAMESPACE "UAFGraphAssetPreview"
+
+namespace UE::UAF::Editor
+{
+
+
+//////////////////////////////////////////////////////////////////////////
+// FUAFGraphAssetPreviewFactory
+
+TSharedRef<IUAFAssetPreview> FUAFGraphAssetPreviewFactory::CreateAssetPreviewWidget(TSharedPtr<FAssetEditorToolkit> InAssetEditorToolkit, const FAssetData& InAssetData) const
+{
+	return SNew(SUAFGraphAssetPreview, InAssetEditorToolkit, InAssetData);
+}
+
+const UStruct* FUAFGraphAssetPreviewFactory::GetPreviewType() const
+{
+	return UAnimNextAnimationGraph::StaticClass();
+}
+
+
+//////////////////////////////////////////////////////////////////////////
+// SUAFGraphAssetPreview
+
+const UStruct* SUAFGraphAssetPreview::GetAssetPreviewType() const
+{
+	return UAnimNextAnimationGraph::StaticClass();
+}
+
+bool SUAFGraphAssetPreview::OnSameTypeAssetSelected(const FAssetData& InAssetData)
+{
+	bool bAssetValidForPreview = false;
+
+	UAnimNextAnimationGraph* Graph = Cast<UAnimNextAnimationGraph>(CachedCurrentAssetData.GetAsset());
+	if (Graph
+		&& UAFComponent
+		&& UAFSceneDescription 
+		&& UAFSceneDescription->SkeletalMesh )
+	{
+		FAnimNextVariableReference VariableReference("Graph", Cast<UAnimNextModule>(UAFComponent->GetAsset()));
+		check(UAFComponent->SetVariable<TObjectPtr<UObject>>(VariableReference, Graph));
+		bAssetValidForPreview = true;
+	}
+
+	ViewportWidget->SetVisibility(bAssetValidForPreview ? EVisibility::Visible : EVisibility::Hidden);
+	return bAssetValidForPreview;
+}
+
+void SUAFGraphAssetPreview::OnCustomizePreviewScene(FAdvancedPreviewScene& InPreviewScene, FEditorViewportClient& InEditorViewportClient)
+{
+	// @TODO: Note for review. Most impl is from AnimGraphViewportController.cpp
+
+	// TODO: Consider building this module programatically
+	TObjectPtr<UAnimNextModule> Module = LoadObject<UAnimNextModule>(GetTransientPackage(), TEXT("/UAFAnimGraph/Internal/S_SingleGraph.S_SingleGraph"));
+
+	auto AddActor = [&](USkeletalMesh* InSkeletalMesh)
+	{
+		AActor* PreviewActor = PreviewScene->GetWorld()->SpawnActor<AActor>(AActor::StaticClass());
+		check(PreviewActor);
+
+		PreviewComponent = NewObject<USkeletalMeshComponent>(PreviewActor);
+		// PreviewComponent->VisibilityBasedAnimTickOption = EVisibilityBasedAnimTickOption::AlwaysTickPoseAndRefreshBones; 
+		PreviewScene->AddComponent(PreviewComponent, FTransform::Identity);
+		PreviewActor->SetRootComponent(PreviewComponent);
+		PreviewComponent->SetEnableAnimation(false);
+		PreviewComponent->SetSkeletalMesh(InSkeletalMesh);
+
+		UAFComponent = NewObject<UAnimNextComponent>(PreviewActor);
+		PreviewActor->AddInstanceComponent(UAFComponent);
+		UAFComponent->SetAsset(Module);
+		UAFComponent->RegisterComponent();
+		UAFComponent->InitializeComponent();
+	};
+
+	if (TSharedPtr<FAssetEditorToolkit> OwningEditor = AssetEditorToolkitPtr.Pin())
+	{
+		using namespace UE::Workspace;
+
+		if (OwningEditor->GetToolkitFName() == IWorkspaceEditor::GetWorkspaceEditorToolkitName())
+		{
+			TSharedPtr<IWorkspaceEditor> OwningWorkspaceEditor = StaticCastSharedPtr<IWorkspaceEditor>(OwningEditor);
+			if (UUAFViewportSceneDescription* WorkspaceSceneDescription = Cast<UUAFViewportSceneDescription>(OwningWorkspaceEditor->GetSceneDescription()))
+			{
+				UAFSceneDescription = WorkspaceSceneDescription;
+			}
+		}
+	}
+
+	if (UAFSceneDescription && UAFSceneDescription->SkeletalMesh)
+	{
+		AddActor(UAFSceneDescription->SkeletalMesh);
+	}
+}
+
+//////////////////////////////////////////////////////////////////////////
+
+} // namespace UE::UAF::Editor
+
+#undef LOCTEXT_NAMESPACE
+
diff --git a/Engine/Plugins/Experimental/UAF/UAFAnimGraph/Source/UAFAnimGraphEditor/Private/AssetPreview/UAFGraphAssetPreview.h b/Engine/Plugins/Experimental/UAF/UAFAnimGraph/Source/UAFAnimGraphEditor/Private/AssetPreview/UAFGraphAssetPreview.h
new file mode 100644
--- /dev/null
+++ b/Engine/Plugins/Experimental/UAF/UAFAnimGraph/Source/UAFAnimGraphEditor/Private/AssetPreview/UAFGraphAssetPreview.h
@@ -0,0 +1,74 @@
+// Copyright Epic Games, Inc. All Rights Reserved.
+
+#pragma once
+
+#include "Common/AssetPreview/IUAFAssetViewportPreview.h"
+
+class FSceneViewport;
+class FUAFAssetPreviewViewportClient;
+class SAssetEditorViewport;
+class UAnimNextComponent;
+class USkeletalMeshComponent;
+class UUAFViewportSceneDescription;
+
+namespace UE::UAF::Editor
+{
+
+
+/**
+ * Experimental. May be removed / relocated whenever we switch to the future generic content browser IAssetPreviewCustomization API.
+ *
+ * Factory for UAF Graph Asset Preview widgets.
+ */
+struct FUAFGraphAssetPreviewFactory : public FUAFAssetPreviewFactory
+{
+public:
+
+	//~ Begin FUAFAssetPreviewFactory Interface
+	virtual TSharedRef<IUAFAssetPreview> CreateAssetPreviewWidget(TSharedPtr<FAssetEditorToolkit> InAssetEditorToolkit, const FAssetData& InAssetData) const override;
+	virtual const UStruct* GetPreviewType() const override;
+	//~ End FUAFAssetPreviewFactory Interface
+};
+
+/** 
+ * Experimental. May be removed / relocated whenever we switch to the future generic content browser IAssetPreviewCustomization API. 
+ * 
+ * Preview widget for UAF Graph Assets. Includes an interactable viewport.
+ * Currently these previews only work in a UAF workspace that has set it's preview skeletal mesh.
+ */
+class SUAFGraphAssetPreview : public IUAFAssetViewportPreview
+{
+public:
+
+	//~ Begin IUAFAssetPreview Interface
+	virtual const UStruct* GetAssetPreviewType() const override;
+	//~ End IUAFAssetPreview Interface
+
+protected:
+
+	//~ Begin IUAFAssetPreview Interface
+	virtual bool OnSameTypeAssetSelected(const FAssetData& InAssetData) override;
+	//~ End IUAFAssetPreview Interface
+
+	//~ Begin IUAFAssetViewportPreview Interface
+	virtual void OnCustomizePreviewScene(FAdvancedPreviewScene& InPreviewScene, FEditorViewportClient& InEditorViewportClient) override;
+	//~ End IUAFAssetViewportPreview Interface
+
+protected:
+
+	/** Skeletal component to preview the animation asset on. Lifetime maintained by PreviewScene in IUAFAssetViewportPreview */
+	USkeletalMeshComponent* PreviewComponent = nullptr;
+
+	/** UAF Graph Preview Driver. Lifetime maintained by PreviewScene in IUAFAssetViewportPreview */
+	UAnimNextComponent* UAFComponent = nullptr;
+
+	/** 
+	 * Object used to describe UAF previews. Lifetime maintained by workspace editor.
+	 * 
+	 * Currently we only acquire this by checking if the current editor is a workspace editor & acquiring it's viewport scene description.
+	 */
+	UUAFViewportSceneDescription* UAFSceneDescription = nullptr;
+};
+
+
+} // namespace UE::UAF::Editor
\ No newline at end of file
diff --git a/Engine/Plugins/Experimental/Workspace/Source/WorkspaceEditor/Private/SWorkspaceViewport.cpp b/Engine/Plugins/Experimental/Workspace/Source/WorkspaceEditor/Private/SWorkspaceViewport.cpp
--- a/Engine/Plugins/Experimental/Workspace/Source/WorkspaceEditor/Private/SWorkspaceViewport.cpp
+++ b/Engine/Plugins/Experimental/Workspace/Source/WorkspaceEditor/Private/SWorkspaceViewport.cpp
@@ -106,18 +106,21 @@ namespace UE::Workspace
 									DetailsViewArgs.bShowOptions = false;
 									DetailsViewArgs.bShowObjectLabel = false;
 									DetailsViewArgs.bHideSelectionTip = true;
-
-									const auto DetailsView = PropertyEditorModule.CreateDetailView(DetailsViewArgs);
-									DetailsView->SetObject(Context->SceneDescription);
-							
-									Submenu->AddMenuEntry(NAME_None, FToolMenuEntry::InitWidget(
-										"Details",
-										DetailsView,
-										FText::GetEmpty(),
-										true,
-										false,
-										true
-									));
+									
+									if (Context->SceneDescription)
+									{
+										const auto DetailsView = PropertyEditorModule.CreateDetailView(DetailsViewArgs);
+										DetailsView->SetObject(Context->SceneDescription);
+
+										Submenu->AddMenuEntry(NAME_None, FToolMenuEntry::InitWidget(
+											"Details",
+											DetailsView,
+											FText::GetEmpty(),
+											true,
+											false,
+											true
+										));
+									}
 								}),
 							false,
 							FSlateIcon(FAppStyle::GetAppStyleSetName(), "AssetEditor.PreviewSceneSettings")));
diff --git a/Engine/Plugins/Experimental/Workspace/Source/WorkspaceEditor/Private/SWorkspaceViewport.h b/Engine/Plugins/Experimental/Workspace/Source/WorkspaceEditor/Private/SWorkspaceViewport.h
deleted file mode 100644
diff --git a/Engine/Plugins/Experimental/Workspace/Source/WorkspaceEditor/Public/SWorkspaceViewport.h b/Engine/Plugins/Experimental/Workspace/Source/WorkspaceEditor/Public/SWorkspaceViewport.h
new file mode 100644
--- /dev/null
+++ b/Engine/Plugins/Experimental/Workspace/Source/WorkspaceEditor/Public/SWorkspaceViewport.h
@@ -0,0 +1,54 @@
+// Copyright Epic Games, Inc. All Rights Reserved.
+
+#pragma once
+
+#include "SAssetEditorViewport.h"
+#include "SCommonEditorViewportToolbarBase.h"
+
+#define UE_API WORKSPACEEDITOR_API
+
+class UWorkspaceViewportSceneDescription;
+
+namespace UE::Workspace
+{
+	class SWorkspaceViewport : public SAssetEditorViewport, public ICommonEditorViewportToolbarInfoProvider, public FGCObject
+	{
+	public:
+		SLATE_BEGIN_ARGS(SWorkspaceViewport){}
+			SLATE_ARGUMENT(TWeakPtr<FAssetEditorToolkit>, AssetEditorToolkit)
+			SLATE_ARGUMENT(TSharedPtr<FEditorViewportClient>, ViewportClient)
+			SLATE_ARGUMENT(UWorkspaceViewportSceneDescription*, SceneDescription)
+			SLATE_ATTRIBUTE(FSoftObjectPath, PreviewAssetPath)
+			SLATE_ATTRIBUTE(bool, bIsPinned)
+			SLATE_EVENT(FSimpleDelegate, OnPinnedClicked)
+		SLATE_END_ARGS()
+		
+		TWeakPtr<FAssetEditorToolkit> AssetEditorToolkit = nullptr;
+
+		UE_API void Construct(const FArguments& InArgs);
+		
+		// SEditorViewport interface
+		UE_API virtual TSharedPtr<SWidget> BuildViewportToolbar() override;
+		UE_API virtual TSharedRef<FEditorViewportClient> MakeEditorViewportClient() override;
+		
+		// ICommonEditorViewportToolbarInfoProvider interface
+		UE_API virtual TSharedRef<SEditorViewport> GetViewportWidget() override;
+		UE_API virtual TSharedPtr<FExtender> GetExtenders() const override;
+		UE_API virtual void OnFloatingButtonClicked() override;
+
+		UE_API virtual void AddReferencedObjects( FReferenceCollector& Collector ) override;
+		UE_API virtual FString GetReferencerName() const override
+		{
+			return TEXT("UE::Workspace::SWorkspaceViewport");
+		}
+
+	private:
+		TAttribute<FSoftObjectPath> PreviewAssetPath;
+		TAttribute<bool> bIsPinned;
+		FSimpleDelegate OnPinnedClicked;
+
+		TObjectPtr<UWorkspaceViewportSceneDescription> SceneDescription;
+	};
+}
+
+#undef UE_API
\ No newline at end of file
diff --git a/Engine/Source/Editor/UnrealEd/Private/SAssetEditorViewport.cpp b/Engine/Source/Editor/UnrealEd/Private/SAssetEditorViewport.cpp
--- a/Engine/Source/Editor/UnrealEd/Private/SAssetEditorViewport.cpp
+++ b/Engine/Source/Editor/UnrealEd/Private/SAssetEditorViewport.cpp
@@ -99,7 +99,11 @@ void SAssetEditorViewport::Construct(const FArguments& InArgs, const FAssetEdito
 {
 	ConfigKey = InViewportConstructionArgs.ConfigKey;
 	Client = InArgs._EditorViewportClient;
-	ParentTabContent = InViewportConstructionArgs.ParentLayout->GetParentTabContent();
+	
+	if (InViewportConstructionArgs.ParentLayout)
+	{
+		ParentTabContent = InViewportConstructionArgs.ParentLayout->GetParentTabContent();
+	}
 
 	SEditorViewport::FArguments ViewportArgs;
 	if (InArgs._ViewportSize.IsSet())
