UAF Browser BugHawk feedback
- Add missing virtual dtor.

#tests editor

diff --git a/Engine/Plugins/Experimental/UAF/UAF/Source/UAFEditor/Public/Common/AssetPreview/IUAFAssetPreview.h b/Engine/Plugins/Experimental/UAF/UAF/Source/UAFEditor/Public/Common/AssetPreview/IUAFAssetPreview.h
--- a/Engine/Plugins/Experimental/UAF/UAF/Source/UAFEditor/Public/Common/AssetPreview/IUAFAssetPreview.h
+++ b/Engine/Plugins/Experimental/UAF/UAF/Source/UAFEditor/Public/Common/AssetPreview/IUAFAssetPreview.h
@@ -46,6 +46,8 @@ public:
 
 	/** Get the preview type this factory is used for. */
 	UE_API virtual const UStruct* GetPreviewType() const = 0;
+
+	UE_API virtual ~FUAFAssetPreviewFactory() = default;
 };
 
 
