Add callback for when a drag starts from the content browser
Used to dismiss on drag the UAF Drawer Sidebar, which covers the details panel.

#tests editor

diff --git a/Engine/Plugins/Experimental/UAF/UAF/Source/UAFEditor/Private/Common/SUAFBrowser.cpp b/Engine/Plugins/Experimental/UAF/UAF/Source/UAFEditor/Private/Common/SUAFBrowser.cpp
--- a/Engine/Plugins/Experimental/UAF/UAF/Source/UAFEditor/Private/Common/SUAFBrowser.cpp
+++ b/Engine/Plugins/Experimental/UAF/UAF/Source/UAFEditor/Private/Common/SUAFBrowser.cpp
@@ -17,6 +17,7 @@
 #include "Settings/UAFEditorProjectSettings.h"
 #include "Settings/UAFEditorUserSettings.h"
 #include "StatusBarSubsystem.h"
+#include "Subsystems/AssetEditorSubsystem.h"
 #include "UncookedOnlyUtils.h"
 #include "Widgets/Layout/SBorder.h"
 #include "Widgets/Layout/SBox.h"
@@ -140,9 +141,20 @@ void SUAFBrowser::OnOverrideAssetPickerConfig(FAssetPickerConfig& Config)
 	//Config.bSortByPathInColumnView = true;
 
 	Config.OnGetAssetContextMenu.Unbind();
-	Config.OnAssetSelected.BindSP(this, &SUAFBrowser::OnAssetSelected);
 	Config.OnAssetsActivated.Unbind();
 	Config.OnGetCustomAssetToolTip.Unbind();
+
+	auto OpenAssetCallback = [](const FAssetData& AssetData)
+	{
+		if (UAssetEditorSubsystem* EditorSubsystem = GEditor->GetEditorSubsystem<UAssetEditorSubsystem>())
+		{
+			EditorSubsystem->OpenEditorForAsset(AssetData.ToSoftObjectPath());
+		}
+	};
+	Config.OnAssetDoubleClicked.BindLambda(OpenAssetCallback);
+
+	Config.OnAssetSelected.BindSP(this, &SUAFBrowser::OnAssetSelected);
+	Config.OnAssetsDragged.BindSP(this, &SUAFBrowser::OnAssetsDragged);
 }
 
 void SUAFBrowser::OnExtendAssetPickerTopBar(TSharedRef<SHorizontalBox> InAssetPickerTopBar)
@@ -291,6 +303,16 @@ void SUAFBrowser::OnAssetSelected(const FAssetData& InAssetData)
 	}
 }
 
+void SUAFBrowser::OnAssetsDragged(const TArray<FAssetData>& InAssetData)
+{
+	// Dismiss all drawers immediately on drag, our workflow involves sidebar drawers that will cover the details panel
+	if (TSharedPtr<UE::Workspace::IWorkspaceEditor> HostingApp = WeakHostingApp.Pin())
+	{
+		GEditor->GetEditorSubsystem<UStatusBarSubsystem>()->ForceDismissDrawer();
+		HostingApp->GetTabManager()->TryDismissSidebarTab(FUAFBrowserTabSummoner::SidebarTabID);
+	}
+}
+
 } // namespace UE::UAF::Editor
 
 #undef LOCTEXT_NAMESPACE
diff --git a/Engine/Plugins/Experimental/UAF/UAF/Source/UAFEditor/Private/Common/SUAFBrowser.h b/Engine/Plugins/Experimental/UAF/UAF/Source/UAFEditor/Private/Common/SUAFBrowser.h
--- a/Engine/Plugins/Experimental/UAF/UAF/Source/UAFEditor/Private/Common/SUAFBrowser.h
+++ b/Engine/Plugins/Experimental/UAF/UAF/Source/UAFEditor/Private/Common/SUAFBrowser.h
@@ -83,6 +83,9 @@ private:
 	/** Callback for when an asset is selected. */
 	void OnAssetSelected(const FAssetData& InAssetData);
 
+	/** Callback for when assets are dragged */
+	void OnAssetsDragged(const TArray<FAssetData>& InAssetData);
+
 private:
 
 	/** App Hosting this Browser, to query settings against */
diff --git a/Engine/Source/Editor/ContentBrowser/Private/SAssetPicker.cpp b/Engine/Source/Editor/ContentBrowser/Private/SAssetPicker.cpp
--- a/Engine/Source/Editor/ContentBrowser/Private/SAssetPicker.cpp
+++ b/Engine/Source/Editor/ContentBrowser/Private/SAssetPicker.cpp
@@ -76,6 +76,7 @@ void SAssetPicker::Construct( const FArguments& InArgs )
 	BindCommands();
 
 	OnAssetsActivated = InArgs._AssetPickerConfig.OnAssetsActivated;
+	OnAssetsDragged = InArgs._AssetPickerConfig.OnAssetsDragged;
 	OnAssetSelected = InArgs._AssetPickerConfig.OnAssetSelected;
 	OnAssetDoubleClicked = InArgs._AssetPickerConfig.OnAssetDoubleClicked;
 	OnAssetEnterPressed = InArgs._AssetPickerConfig.OnAssetEnterPressed;
@@ -344,6 +345,7 @@ void SAssetPicker::Construct( const FArguments& InArgs )
 		.OnNewItemRequested(this, &SAssetPicker::HandleNewItemRequested)
 		.OnItemSelectionChanged(this, &SAssetPicker::HandleItemSelectionChanged)
 		.OnItemsActivated(this, &SAssetPicker::HandleItemsActivated)
+		.OnItemsDragged(this, &SAssetPicker::HandleItemsDragged)
 		.OnGetItemContextMenu(OnGetItemContextMenu)
 		.OnIsAssetValidForCustomToolTip(InArgs._AssetPickerConfig.OnIsAssetValidForCustomToolTip)
 		.OnGetCustomAssetToolTip(InArgs._AssetPickerConfig.OnGetCustomAssetToolTip)
@@ -655,6 +657,29 @@ void SAssetPicker::HandleItemsActivated(TArrayView<const FContentBrowserItem> Ac
 	OnAssetsActivated.ExecuteIfBound( ActivatedAssets, ActivationMethod );
 }
 
+void SAssetPicker::HandleItemsDragged(TArrayView<const FContentBrowserItem> DraggedItems)
+{
+	TArray<FAssetData> DraggedAssets;
+	for (const FContentBrowserItem& DraggedItem : DraggedItems)
+	{
+		if (DraggedItem.IsFile())
+		{
+			FAssetData ItemAssetData;
+			if (DraggedItem.Legacy_TryGetAssetData(ItemAssetData))
+			{
+				DraggedAssets.Add(MoveTemp(ItemAssetData));
+			}
+		}
+	}
+
+	if (DraggedAssets.Num() == 0)
+	{
+		return;
+	}
+
+	OnAssetsDragged.ExecuteIfBound(DraggedAssets);
+}
+
 void SAssetPicker::SyncToAssets(const TArray<FAssetData>& AssetDataList)
 {
 	AssetViewPtr->SyncToLegacy(AssetDataList, TArray<FString>());
diff --git a/Engine/Source/Editor/ContentBrowser/Private/SAssetPicker.h b/Engine/Source/Editor/ContentBrowser/Private/SAssetPicker.h
--- a/Engine/Source/Editor/ContentBrowser/Private/SAssetPicker.h
+++ b/Engine/Source/Editor/ContentBrowser/Private/SAssetPicker.h
@@ -102,6 +102,9 @@ private:
 	/** Handler for when the user double clicks, presses enter, or presses space on an asset */
 	void HandleItemsActivated(TArrayView<const FContentBrowserItem> ActivatedItems, EAssetTypeActivationMethod::Type ActivationMethod);
 
+	/** Handler for when the user drags an asset */
+	void HandleItemsDragged(TArrayView<const FContentBrowserItem> DraggedItems);
+
 	/** Selects the paths containing the specified assets. */
 	void SyncToAssets(const TArray<FAssetData>& AssetDataList);
 
@@ -178,6 +181,9 @@ private:
 	/** Called when any number of assets are activated */
 	FOnAssetsActivated OnAssetsActivated;
 
+	/** Called any number of assets are dragged */
+	FOnAssetsDragged OnAssetsDragged;
+
 	/** Called when a folder is entered in the asset view */
 	FOnPathSelected OnFolderEnteredDelegate;
 
diff --git a/Engine/Source/Editor/ContentBrowser/Private/SAssetView.cpp b/Engine/Source/Editor/ContentBrowser/Private/SAssetView.cpp
--- a/Engine/Source/Editor/ContentBrowser/Private/SAssetView.cpp
+++ b/Engine/Source/Editor/ContentBrowser/Private/SAssetView.cpp
@@ -1387,6 +1387,7 @@ void SAssetView::Construct( const FArguments& InArgs )
 	OnNewItemRequested = InArgs._OnNewItemRequested;
 	OnItemSelectionChanged = InArgs._OnItemSelectionChanged;
 	OnItemsActivated = InArgs._OnItemsActivated;
+	OnItemsDragged = InArgs._OnItemsDragged;
 	OnGetItemContextMenu = InArgs._OnGetItemContextMenu;
 	OnItemRenameCommitted = InArgs._OnItemRenameCommitted;
 	OnAssetTagWantsToBeDisplayed = InArgs._OnAssetTagWantsToBeDisplayed;
@@ -5763,6 +5764,7 @@ FReply SAssetView::OnDraggingAssetItem( const FGeometry& MyGeometry, const FPoin
 
 			if (TSharedPtr<FDragDropOperation> DragDropOp = DragDropHandler::CreateDragOperation(SelectedItems))
 			{
+				OnItemsDragged.ExecuteIfBound(SelectedItems);
 				return FReply::Handled().BeginDragDrop(DragDropOp.ToSharedRef());
 			}
 		}
diff --git a/Engine/Source/Editor/ContentBrowser/Public/ContentBrowserDelegates.h b/Engine/Source/Editor/ContentBrowser/Public/ContentBrowserDelegates.h
--- a/Engine/Source/Editor/ContentBrowser/Public/ContentBrowserDelegates.h
+++ b/Engine/Source/Editor/ContentBrowser/Public/ContentBrowserDelegates.h
@@ -27,6 +27,9 @@ DECLARE_DELEGATE_TwoParams(FOnContentBrowserItemSelectionChanged, const FContent
 /** Called when the user double clicks, presses enter, or presses space on a Content Browser item */
 DECLARE_DELEGATE_TwoParams(FOnContentBrowserItemsActivated, TArrayView<const FContentBrowserItem> /*ActivatedItems*/, EAssetTypeActivationMethod::Type /*ActivationMethod*/);
 
+/** Called when the user drags one or more Content Browser items */
+DECLARE_DELEGATE_OneParam(FOnContentBrowserItemsDragged, TArrayView<const FContentBrowserItem> /*Items*/);
+
 /** Called to request the menu when right clicking on an selection of items (maybe be a mix of folders and files) */
 DECLARE_DELEGATE_RetVal_OneParam(TSharedPtr<SWidget>, FOnGetContentBrowserItemContextMenu, TArrayView<const FContentBrowserItem> /*SelectedItems*/);
 
@@ -72,6 +75,9 @@ DECLARE_DELEGATE_OneParam(FOnAssetSelected, const FAssetData& /*AssetData*/);
 /** Called when the user double clicks, presses enter, or presses space on an asset */
 DECLARE_DELEGATE_TwoParams(FOnAssetsActivated, const TArray<FAssetData>& /*ActivatedAssets*/, EAssetTypeActivationMethod::Type /*ActivationMethod*/);
 
+/** Called when asset(s) are dragged in the asset view */
+DECLARE_DELEGATE_OneParam(FOnAssetsDragged, const TArray<FAssetData>& /*DraggedAssets*/);
+
 /** Called when an asset is double clicked in the asset view */
 DECLARE_DELEGATE_OneParam(FOnAssetDoubleClicked, const FAssetData& /*AssetData*/);
 
diff --git a/Engine/Source/Editor/ContentBrowser/Public/IContentBrowserSingleton.h b/Engine/Source/Editor/ContentBrowser/Public/IContentBrowserSingleton.h
--- a/Engine/Source/Editor/ContentBrowser/Public/IContentBrowserSingleton.h
+++ b/Engine/Source/Editor/ContentBrowser/Public/IContentBrowserSingleton.h
@@ -243,6 +243,9 @@ struct FAssetPickerConfig
 	/** The delegate that fires when any number of assets are activated */
 	FOnAssetsActivated OnAssetsActivated;
 
+	/** The delegate that fires when any number of assets are dragged */
+	FOnAssetsDragged OnAssetsDragged;
+
 	/** The delegate that fires when an asset is right clicked and a context menu is requested */
 	FOnGetAssetContextMenu OnGetAssetContextMenu;
 
diff --git a/Engine/Source/Editor/ContentBrowser/Public/SAssetView.h b/Engine/Source/Editor/ContentBrowser/Public/SAssetView.h
--- a/Engine/Source/Editor/ContentBrowser/Public/SAssetView.h
+++ b/Engine/Source/Editor/ContentBrowser/Public/SAssetView.h
@@ -169,6 +169,9 @@ public:
 		/** Called when the user double clicks, presses enter, or presses space on an Content Browser item */
 		SLATE_EVENT( FOnContentBrowserItemsActivated, OnItemsActivated )
 
+		/** Called when the user double clicks, presses enter, or presses space on an Content Browser item */
+		SLATE_EVENT( FOnContentBrowserItemsDragged, OnItemsDragged )
+
 		/** Delegate to invoke when a context menu for an item is opening. */
 		SLATE_EVENT( FOnGetContentBrowserItemContextMenu, OnGetItemContextMenu )
 
@@ -1112,6 +1115,9 @@ private:
 	/** Called when the user double clicks, presses enter, or presses space on a Content Browser item */
 	FOnContentBrowserItemsActivated OnItemsActivated;
 
+	/** Called when the user drags a Content Browser item */
+	FOnContentBrowserItemsDragged OnItemsDragged;
+
 	/** Delegate to invoke when generating the context menu for an item */
 	FOnGetContentBrowserItemContextMenu OnGetItemContextMenu;
 
