Allow Abstract classes Tagged Asset Browser Filter
Allow Tagged Asset Browser searching for NativeParentClass
Allow Tagged Asset Browser searching for Skeleton

#tests editor

diff --git a/Engine/Plugins/Experimental/UAF/UAF/Source/UAFEditor/Private/Common/TaggedAssetBrowser_UAFBrowserFilters.cpp b/Engine/Plugins/Experimental/UAF/UAF/Source/UAFEditor/Private/Common/TaggedAssetBrowser_UAFBrowserFilters.cpp
new file mode 100644
--- /dev/null
+++ b/Engine/Plugins/Experimental/UAF/UAF/Source/UAFEditor/Private/Common/TaggedAssetBrowser_UAFBrowserFilters.cpp
@@ -0,0 +1,87 @@
+// Copyright Epic Games, Inc. All Rights Reserved.
+
+#include "Common/TaggedAssetBrowser_UAFBrowserFilters.h"
+
+#include "Animation/Skeleton.h"
+#include "IContentBrowserDataModule.h"
+#include "UserAssetTagEditorUtilities.h"
+#include "UserAssetTagsEditorModule.h"
+#include "Logging/StructuredLog.h"
+#include "Styling/SlateIconFinder.h"
+#include "Widgets/STaggedAssetBrowser.h"
+
+#include UE_INLINE_GENERATED_CPP_BY_NAME(TaggedAssetBrowser_UAFBrowserFilters)
+
+#define LOCTEXT_NAMESPACE "UserAssetTags"
+
+
+//////////////////////////////////////////////////////////////////////////
+// UTaggedAssetBrowserFilter_NativeClass
+
+void UTaggedAssetBrowserFilter_NativeClass::ModifyARFilterInternal(FARFilter& Filter) const
+{
+	for(UClass* Class : Classes)
+	{
+		if(Class)
+		{
+			if (bSearchNativeParentClasses)
+			{
+				Filter.TagsAndValues.Add(FBlueprintTags::NativeParentClassPath, FObjectPropertyBase::GetExportPath(Class));
+			}
+			else
+			{
+				Filter.ClassPaths.Add(Class->GetClassPathName());
+			}
+		}
+	}
+	Filter.bRecursiveClasses = bAllowChildClasses;
+}
+
+
+//////////////////////////////////////////////////////////////////////////
+// UTaggedAssetBrowserFilter_Skeleton
+
+FString UTaggedAssetBrowserFilter_Skeleton::ToString() const
+{
+	if (!FilterName.IsNone())
+	{
+		return FilterName.ToString();
+	}
+
+	if (Skeleton)
+	{
+		return Skeleton->GetFName().ToString();
+	}
+
+	return UTaggedAssetBrowserFilterBase::ToString();
+}
+
+FText UTaggedAssetBrowserFilter_Skeleton::GetTooltip() const
+{
+	if (Skeleton)
+	{
+		FText BaseTooltip = LOCTEXT("UTaggedAssetBrowserFilter_ClassTooltip", "Skeleton Filter: {0}.");
+		return FText::FormatOrdered(BaseTooltip, FText::FromString(Skeleton->GetFName().ToString()));
+	}
+
+	return UTaggedAssetBrowserFilterBase::GetTooltip();
+}
+
+FSlateIcon UTaggedAssetBrowserFilter_Skeleton::GetIcon() const
+{
+	return FSlateIconFinder::FindIconForClass(USkeleton::StaticClass());
+}
+
+void UTaggedAssetBrowserFilter_Skeleton::ModifyARFilterInternal(FARFilter& Filter) const
+{
+	static const FName SkeletonTagName("Skeleton");
+	if(Skeleton)
+	{
+		Filter.TagsAndValues.Add(SkeletonTagName, FAssetData(Skeleton).GetExportTextName());
+	}
+}
+
+
+//////////////////////////////////////////////////////////////////////////
+
+#undef LOCTEXT_NAMESPACE
diff --git a/Engine/Plugins/Experimental/UAF/UAF/Source/UAFEditor/Private/Common/TaggedAssetBrowser_UAFBrowserFilters.h b/Engine/Plugins/Experimental/UAF/UAF/Source/UAFEditor/Private/Common/TaggedAssetBrowser_UAFBrowserFilters.h
new file mode 100644
--- /dev/null
+++ b/Engine/Plugins/Experimental/UAF/UAF/Source/UAFEditor/Private/Common/TaggedAssetBrowser_UAFBrowserFilters.h
@@ -0,0 +1,55 @@
+// Copyright Epic Games, Inc. All Rights Reserved.
+
+#pragma once
+
+#include "TaggedAssetBrowserFilters/TaggedAssetBrowser_CommonFilters.h"
+#include "TaggedAssetBrowser_UAFBrowserFilters.generated.h"
+
+#define UE_API UAFEDITOR_API
+
+class USkeleton;
+
+/** 
+ * Experimental. Temp filter to allow us to find BP whose BPGC would match a class filter.
+ * 
+ * Our solution here isn't as robust as we would like. But it does work for UAF Browser MVP needs.
+ */
+UCLASS(MinimalAPI, DisplayName="Native Classes")
+class UTaggedAssetBrowserFilter_NativeClass : public UTaggedAssetBrowserFilter_Class
+{
+	GENERATED_BODY()
+public:
+
+	UE_API virtual void ModifyARFilterInternal(FARFilter& Filter) const override;
+
+	/** True to search for a native parent class, useful for finding BPs whose BPGCs would match but BP itself does not match. */
+	UPROPERTY(EditAnywhere, Category = "Filter")
+	bool bSearchNativeParentClasses = false;
+};
+
+
+/** Filter for a specifc skeleton */
+UCLASS(MinimalAPI, DisplayName="Skeleton")
+class UTaggedAssetBrowserFilter_Skeleton : public UTaggedAssetBrowserFilterBase
+{
+	GENERATED_BODY()
+public:
+	UTaggedAssetBrowserFilter_Skeleton() = default;
+
+	UE_API virtual FString ToString() const override;
+	UE_API virtual TOptional<FString> GetInstanceIdentifier() const override { return ToString(); }
+	UE_API virtual FText GetTooltip() const override;
+	UE_API virtual FSlateIcon GetIcon() const override;
+
+	UE_API virtual void ModifyARFilterInternal(FARFilter& Filter) const override;
+
+	/** Skeleton to search for. */
+	UPROPERTY(EditAnywhere, Category = "Filter")
+	TObjectPtr<USkeleton> Skeleton;
+
+	/** Name for this class filter, by default will use name of the skeleton */
+	UPROPERTY(EditAnywhere, Category = "Filter")
+	FName FilterName;
+};
+
+#undef UE_API
\ No newline at end of file
diff --git a/Engine/Plugins/Experimental/UAF/UAF/Source/UAFEditor/UAFEditor.Build.cs b/Engine/Plugins/Experimental/UAF/UAF/Source/UAFEditor/UAFEditor.Build.cs
--- a/Engine/Plugins/Experimental/UAF/UAF/Source/UAFEditor/UAFEditor.Build.cs
+++ b/Engine/Plugins/Experimental/UAF/UAF/Source/UAFEditor/UAFEditor.Build.cs
@@ -44,6 +44,7 @@ namespace UnrealBuildTool.Rules
 					"WorkspaceEditor",
 					"ContentBrowser",
 					"ContentBrowserData",
+					"DataHierarchyEditor",
 					"UniversalObjectLocator",
 					"UniversalObjectLocatorEditor",
 					"UserAssetTagsEditor",
diff --git a/Engine/Source/Editor/UserAssetTagsEditor/Private/TaggedAssetBrowserFilters/TaggedAssetBrowser_CommonFilters.cpp b/Engine/Source/Editor/UserAssetTagsEditor/Private/TaggedAssetBrowserFilters/TaggedAssetBrowser_CommonFilters.cpp
--- a/Engine/Source/Editor/UserAssetTagsEditor/Private/TaggedAssetBrowserFilters/TaggedAssetBrowser_CommonFilters.cpp
+++ b/Engine/Source/Editor/UserAssetTagsEditor/Private/TaggedAssetBrowserFilters/TaggedAssetBrowser_CommonFilters.cpp
@@ -358,6 +358,7 @@ void UTaggedAssetBrowserFilter_Class::ModifyARFilterInternal(FARFilter& Filter)
 			Filter.ClassPaths.Add(Class->GetClassPathName());
 		}
 	}
+	Filter.bRecursiveClasses = bAllowChildClasses;
 }
 
 #undef LOCTEXT_NAMESPACE
diff --git a/Engine/Source/Editor/UserAssetTagsEditor/Public/TaggedAssetBrowserFilters/TaggedAssetBrowser_CommonFilters.h b/Engine/Source/Editor/UserAssetTagsEditor/Public/TaggedAssetBrowserFilters/TaggedAssetBrowser_CommonFilters.h
--- a/Engine/Source/Editor/UserAssetTagsEditor/Public/TaggedAssetBrowserFilters/TaggedAssetBrowser_CommonFilters.h
+++ b/Engine/Source/Editor/UserAssetTagsEditor/Public/TaggedAssetBrowserFilters/TaggedAssetBrowser_CommonFilters.h
@@ -172,9 +172,13 @@ public:
 	
 	virtual void ModifyARFilterInternal(FARFilter& Filter) const override;
 
-	UPROPERTY(EditAnywhere, Category="Filter")
+	UPROPERTY(EditAnywhere, Category="Filter", meta=(AllowAbstract))
 	TArray<TObjectPtr<UClass>> Classes;
 
+	/** True if we should allow child classes in our search. */
+	UPROPERTY(EditAnywhere, Category = "Filter")
+	bool bAllowChildClasses = true;
+
 	/** Name for this class filter, by default will use first class in list as name if not set */
 	UPROPERTY(EditAnywhere, Category = "Filter")
 	FName FilterName;
diff --git a/Engine/Source/Editor/UserAssetTagsEditor/Public/TaggedAssetBrowserMenuFilters.h b/Engine/Source/Editor/UserAssetTagsEditor/Public/TaggedAssetBrowserMenuFilters.h
--- a/Engine/Source/Editor/UserAssetTagsEditor/Public/TaggedAssetBrowserMenuFilters.h
+++ b/Engine/Source/Editor/UserAssetTagsEditor/Public/TaggedAssetBrowserMenuFilters.h
@@ -25,7 +25,7 @@ struct FTaggedAssetBrowserContext
  *  While filters are generally part of the hierarchy, sections can also specify additional filters
  */
 UCLASS(Abstract, EditInlineNew)
-class UTaggedAssetBrowserFilterBase : public UHierarchyItem
+class USERASSETTAGSEDITOR_API UTaggedAssetBrowserFilterBase : public UHierarchyItem
 {
 	GENERATED_BODY()
 	
