Add default UAF Browser Tagged Asset Browser Config & Enable UAF Browser by default

#tests editor
#virtualized

diff --git a/Engine/Plugins/Experimental/UAF/UAF/Source/UAFEditor/Private/AnimNextEditorModule.cpp b/Engine/Plugins/Experimental/UAF/UAF/Source/UAFEditor/Private/AnimNextEditorModule.cpp
--- a/Engine/Plugins/Experimental/UAF/UAF/Source/UAFEditor/Private/AnimNextEditorModule.cpp
+++ b/Engine/Plugins/Experimental/UAF/UAF/Source/UAFEditor/Private/AnimNextEditorModule.cpp
@@ -130,7 +130,7 @@ static FAutoConsoleCommand ToggleHiddenUAFContent(
     TEXT("Toggles whether internal UAF content is shown in the editor"),
     ToggleUAFContentHidden);
 
-static bool bEnableUAFBrowser = false;
+static bool bEnableUAFBrowser = true;
 static FAutoConsoleVariableRef CVarEnableUAFBrowser(
 	TEXT("a.UAF.EnableUAFBrowser"),
 	bEnableUAFBrowser,
diff --git a/Engine/Plugins/Experimental/UAF/UAF/Source/UAFEditor/Private/Settings/UAFEditorProjectSettings.cpp b/Engine/Plugins/Experimental/UAF/UAF/Source/UAFEditor/Private/Settings/UAFEditorProjectSettings.cpp
--- a/Engine/Plugins/Experimental/UAF/UAF/Source/UAFEditor/Private/Settings/UAFEditorProjectSettings.cpp
+++ b/Engine/Plugins/Experimental/UAF/UAF/Source/UAFEditor/Private/Settings/UAFEditorProjectSettings.cpp
@@ -6,4 +6,5 @@
 
 UUAFEditorProjectSettings::UUAFEditorProjectSettings()
 {
+	UAFBrowserCategoryConfiguration = "/UAFAnimGraph/Internal/TABC_UAFBrowserDefaultConfig.TABC_UAFBrowserDefaultConfig";
 }
