Fix renaming fallout - UAFProvider doesn't have updated name

#rnx

diff --git a/Engine/Plugins/Experimental/UAF/UAF/Source/UAFEditor/Private/RewindDebugger/AnimNextProvider.cpp b/Engine/Plugins/Experimental/UAF/UAF/Source/UAFEditor/Private/RewindDebugger/AnimNextProvider.cpp
--- a/Engine/Plugins/Experimental/UAF/UAF/Source/UAFEditor/Private/RewindDebugger/AnimNextProvider.cpp
+++ b/Engine/Plugins/Experimental/UAF/UAF/Source/UAFEditor/Private/RewindDebugger/AnimNextProvider.cpp
@@ -4,9 +4,9 @@
 #if UAF_TRACE_ENABLED
 #include "ObjectTrace.h"
 
-FName FUAFProvider::ProviderName("AnimNextProvider");
+FName FUAFProvider::ProviderName("UAFProvider");
 
-#define LOCTEXT_NAMESPACE "AnimNextProvider"
+#define LOCTEXT_NAMESPACE "UAFProvider"
 
 FUAFProvider::FUAFProvider(TraceServices::IAnalysisSession& InSession)
 	: Session(InSession)
