Another RigVM notify requiring refresh of Functions outliner (removing/adding collapse node (graphs))

diff --git a/Engine/Plugins/Experimental/UAF/UAF/Source/UAFEditor/Private/Functions/Outliner/FunctionsOutlinerView.cpp b/Engine/Plugins/Experimental/UAF/UAF/Source/UAFEditor/Private/Functions/Outliner/FunctionsOutlinerView.cpp
--- a/Engine/Plugins/Experimental/UAF/UAF/Source/UAFEditor/Private/Functions/Outliner/FunctionsOutlinerView.cpp
+++ b/Engine/Plugins/Experimental/UAF/UAF/Source/UAFEditor/Private/Functions/Outliner/FunctionsOutlinerView.cpp
@@ -99,7 +99,7 @@ void SFunctionsView::OnRigVMModified(ERigVMGraphNotifType NotifType, URigVMGraph
 	case ERigVMGraphNotifType::NodeRemoved:
 	case ERigVMGraphNotifType::NodeCategoryChanged:
 		{
-			if (Cast<URigVMFunctionLibrary>(Graph))
+			if (Cast<URigVMFunctionLibrary>(Graph) || Cast<URigVMFunctionInterfaceNode>(Subject))
 			{
 				FunctionsOutliner->FullRefresh();
 			}
