//UE5/Main - Build Engine Localization - Text conflict from LOCTEXT macro for namespace 'FFunctionsOutlinerMode' and key 'RemoveFuncFromCategoryFormat'

diff --git a/Engine/Plugins/Experimental/UAF/UAF/Source/UAFEditor/Private/Functions/Outliner/FunctionsOutlinerMode.cpp b/Engine/Plugins/Experimental/UAF/UAF/Source/UAFEditor/Private/Functions/Outliner/FunctionsOutlinerMode.cpp
--- a/Engine/Plugins/Experimental/UAF/UAF/Source/UAFEditor/Private/Functions/Outliner/FunctionsOutlinerMode.cpp
+++ b/Engine/Plugins/Experimental/UAF/UAF/Source/UAFEditor/Private/Functions/Outliner/FunctionsOutlinerMode.cpp
@@ -171,7 +171,7 @@ FFunctionsOutlinerMode::FFunctionsOutlinerMode(SCommonOutliner* InFunctionsOutli
 
 			if (FuncEditorData == AssetEditorData)
 			{
-				const FText FormattedMessage = FText::Format(LOCTEXT("RemoveFuncFromCategoryFormat", "Removing {0} from Category {1}"), OtherFuncName, ThisFuncCategoryText);	
+				const FText FormattedMessage = FText::Format(LOCTEXT("RemovingFuncFromCategoryFormat", "Removing {0} from Category {1}"), OtherFuncName, ThisFuncCategoryText);	
 				FScopedTransaction Transaction(FormattedMessage);
 				
 				if (URigVMGraph* Model = CollapseNode->GetGraph())
