Introducing new abstract skeleton runtime for UAF (disabled by default)

diff --git a/Engine/Plugins/Experimental/MovieSceneAnimMixer/Source/MovieSceneAnimMixer/Private/Systems/MovieSceneAnimMixerSystem.cpp b/Engine/Plugins/Experimental/MovieSceneAnimMixer/Source/MovieSceneAnimMixer/Private/Systems/MovieSceneAnimMixerSystem.cpp
--- a/Engine/Plugins/Experimental/MovieSceneAnimMixer/Source/MovieSceneAnimMixer/Private/Systems/MovieSceneAnimMixerSystem.cpp
+++ b/Engine/Plugins/Experimental/MovieSceneAnimMixer/Source/MovieSceneAnimMixer/Private/Systems/MovieSceneAnimMixerSystem.cpp
@@ -346,6 +346,14 @@ void FAnimNextBlendTwoKeyframesPreserveRootMotionTask::Execute(UE::UAF::FEvaluat
 	using namespace UE::UAF;
 	using namespace UE::MovieScene;
 
+	if (VM.GetContext().IsValid())
+	{
+		// TODO: Implement with new attribute runtime
+		TUniquePtr<FAttributeCollection> DiscardedCollection;
+		(void)VM.PopValue(ATTRIBUTE_STACK_NAME, DiscardedCollection);
+		return;
+	}
+
 	const UE::Anim::FAttributeId RootMotionDeltaAttributeId(UE::Anim::IAnimRootMotionProvider::AttributeName, FCompactPoseBoneIndex(0));
 
 	const FAnimMixerComponentTypes* AnimMixerComponents = FAnimMixerComponentTypes::Get();
@@ -473,6 +481,16 @@ void FMovieSceneAccumulateAbsoluteBlendTask::Execute(UE::UAF::FEvaluationVM& VM)
 {
 	using namespace UE::UAF;
 
+	if (VM.GetContext().IsValid())
+	{
+		// TODO: Implement with new attribute runtime
+		TUniquePtr<FAttributeCollection> DiscardedCollection;
+		(void)VM.PopValue(ATTRIBUTE_STACK_NAME, DiscardedCollection);
+		return;
+
+		// TODO: Is this task any different from the normal Accumulate task?
+	}
+
 	// Pop our top two poses, we'll re-use the top keyframe for our result
 
 	TUniquePtr<FKeyframeState> KeyframeA;
diff --git a/Engine/Plugins/Experimental/MovieSceneAnimMixer/Source/MovieSceneAnimMixer/Private/Systems/MovieSceneMixedSkeletalAnimationSystem.cpp b/Engine/Plugins/Experimental/MovieSceneAnimMixer/Source/MovieSceneAnimMixer/Private/Systems/MovieSceneMixedSkeletalAnimationSystem.cpp
--- a/Engine/Plugins/Experimental/MovieSceneAnimMixer/Source/MovieSceneAnimMixer/Private/Systems/MovieSceneMixedSkeletalAnimationSystem.cpp
+++ b/Engine/Plugins/Experimental/MovieSceneAnimMixer/Source/MovieSceneAnimMixer/Private/Systems/MovieSceneMixedSkeletalAnimationSystem.cpp
@@ -405,6 +405,16 @@ void FMovieSceneSkeletalAnimationEvaluationTask::Execute(UE::UAF::FEvaluationVM&
 {
 	using namespace UE::UAF;
 	using namespace UE::MovieScene;
+
+	if (VM.GetContext().IsValid())
+	{
+		// TODO: Implement with new attribute runtime
+		constexpr bool bIsAdditive = false;
+		constexpr bool bInitializeWithDefaultValues = true;
+		VM.PushValue(ATTRIBUTE_STACK_NAME, MakeUnique<FAttributeCollection>(VM.MakeAttributeCollection(bIsAdditive, bInitializeWithDefaultValues)));
+		return;
+	}
+
 	FAnimMixerComponentTypes* AnimMixerComponents = FAnimMixerComponentTypes::Get();
 	if (const UAnimSequence* AnimSequencePtr = AnimationData.AnimSequence.Get())
 	{
@@ -417,7 +427,7 @@ void FMovieSceneSkeletalAnimationEvaluationTask::Execute(UE::UAF::FEvaluationVM&
 
 		FKeyframeState Keyframe = VM.MakeUninitializedKeyframe(bIsAdditive);
 
-		const bool bUseRawData = FDecompressionTools::ShouldUseRawData(AnimSequencePtr, Keyframe.Pose);
+		const bool bUseRawData = FDecompressionTools::ShouldUseRawData(AnimSequencePtr, ExtractionContext);
 		const bool bUseRootMotion = AnimSequencePtr->bEnableRootMotion;
 
 		TOptional<FTransform> RootBoneTransform;
diff --git a/Engine/Plugins/Experimental/MovieSceneAnimMixer/Source/MovieSceneAnimMixer/Private/Systems/MovieSceneRootMotionSystem.cpp b/Engine/Plugins/Experimental/MovieSceneAnimMixer/Source/MovieSceneAnimMixer/Private/Systems/MovieSceneRootMotionSystem.cpp
--- a/Engine/Plugins/Experimental/MovieSceneAnimMixer/Source/MovieSceneAnimMixer/Private/Systems/MovieSceneRootMotionSystem.cpp
+++ b/Engine/Plugins/Experimental/MovieSceneAnimMixer/Source/MovieSceneAnimMixer/Private/Systems/MovieSceneRootMotionSystem.cpp
@@ -270,6 +270,12 @@ void FAnimNextConvertRootMotionToWorldSpaceTask::Execute(UE::UAF::FEvaluationVM&
 	using namespace UE::UAF; 
 	using namespace UE::MovieScene;
 
+	if (VM.GetContext().IsValid())
+	{
+		// TODO: Implement with new attribute runtime
+		return;
+	}
+
 	if (!EnumHasAnyFlags(VM.GetFlags(), EEvaluationFlags::Attributes | EEvaluationFlags::Trajectory))
 	{
 		return;
@@ -357,6 +363,12 @@ void FAnimNextStoreRootTransformTask::Execute(UE::UAF::FEvaluationVM& VM) const
 	using namespace UE::UAF;
 	using namespace UE::MovieScene;
 
+	if (VM.GetContext().IsValid())
+	{
+		// TODO: Implement with new attribute runtime
+		return;
+	}
+
 	FAnimMixerComponentTypes* AnimMixerComponents = FAnimMixerComponentTypes::Get();
 
 	TSharedPtr<FMovieSceneMixerRootMotionComponentData> RootMotionData = WeakRootMotionData.Pin();
diff --git a/Engine/Plugins/Experimental/MovieScenePoseSearchTracks/Source/MovieScenePoseSearchTracks/Private/Sections/MovieSceneStitchAnimSection.cpp b/Engine/Plugins/Experimental/MovieScenePoseSearchTracks/Source/MovieScenePoseSearchTracks/Private/Sections/MovieSceneStitchAnimSection.cpp
--- a/Engine/Plugins/Experimental/MovieScenePoseSearchTracks/Source/MovieScenePoseSearchTracks/Private/Sections/MovieSceneStitchAnimSection.cpp
+++ b/Engine/Plugins/Experimental/MovieScenePoseSearchTracks/Source/MovieScenePoseSearchTracks/Private/Sections/MovieSceneStitchAnimSection.cpp
@@ -59,6 +59,15 @@ void FMovieSceneStitchAnimEvaluationTask::Execute(UE::UAF::FEvaluationVM& VM) co
 	using namespace UE::PoseSearch;
 	using namespace UE::Anim;
 
+	if (VM.GetContext().IsValid())
+	{
+		// TODO: Implement with new attribute runtime
+		constexpr bool bIsAdditive = false;
+		constexpr bool bInitializeWithDefaultValues = true;
+		VM.PushValue(ATTRIBUTE_STACK_NAME, MakeUnique<FAttributeCollection>(VM.MakeAttributeCollection(bIsAdditive, bInitializeWithDefaultValues)));
+		return;
+	}
+
 	const FAnimMixerComponentTypes* AnimMixerComponents = FAnimMixerComponentTypes::Get();
 
 	// We rely on the pose history having been stored first.
@@ -230,7 +239,7 @@ void FMovieSceneStitchAnimEvaluationTask::Execute(UE::UAF::FEvaluationVM& VM) co
 	FDeltaTimeRecord DeltaTime;
 	DeltaTime.Set(PreviousAnimationTime, AnimationTime-PreviousAnimationTime);
 	const FAnimExtractContext ExtractionContext((double)AnimationTime, false, DeltaTime, false);
-	const bool bUseRawData = FDecompressionTools::ShouldUseRawData(AnimSequence, Keyframe.Pose);
+	const bool bUseRawData = FDecompressionTools::ShouldUseRawData(AnimSequence, ExtractionContext);
 
 	if (EnumHasAnyFlags(VM.GetFlags(), EEvaluationFlags::Bones))
 	{
diff --git a/Engine/Plugins/Experimental/RigLogicUAF/Source/RigLogicUAF/Private/RigLogicTask.cpp b/Engine/Plugins/Experimental/RigLogicUAF/Source/RigLogicUAF/Private/RigLogicTask.cpp
--- a/Engine/Plugins/Experimental/RigLogicUAF/Source/RigLogicUAF/Private/RigLogicTask.cpp
+++ b/Engine/Plugins/Experimental/RigLogicUAF/Source/RigLogicUAF/Private/RigLogicTask.cpp
@@ -28,6 +28,12 @@ void FUAFRigLogicTask::Execute(UE::UAF::FEvaluationVM& VM) const
 {
 	SCOPE_CYCLE_COUNTER(STAT_UAF_Task_RigLogic);
 
+	if (VM.GetContext().IsValid())
+	{
+		// TODO: Implement with new attribute runtime
+		return;
+	}
+
 	// Pop pose, we'll re-use it for our result
 	TUniquePtr<UE::UAF::FKeyframeState> Keyframe;
 	if (!VM.PopValue(UE::UAF::KEYFRAME_STACK_NAME, Keyframe))
diff --git a/Engine/Plugins/Experimental/UAF/UAF/Source/UAF/Private/AnimNextModuleImpl.cpp b/Engine/Plugins/Experimental/UAF/UAF/Source/UAF/Private/AnimNextModuleImpl.cpp
--- a/Engine/Plugins/Experimental/UAF/UAF/Source/UAF/Private/AnimNextModuleImpl.cpp
+++ b/Engine/Plugins/Experimental/UAF/UAF/Source/UAF/Private/AnimNextModuleImpl.cpp
@@ -31,6 +31,14 @@
 #include "Features/IModularFeatures.h"
 #include "RewindDebuggerRuntimeInterface/IRewindDebuggerRuntimeExtension.h"
 #include "RewindDebugger/RewindDebuggerUAFRuntime.h"
+#include "UAF/Attributes/AttributeRuntimeRegistry.h"
+#include "UAF/Attributes/Transformers/Accumulate.h"
+#include "UAF/Attributes/Transformers/Interpolate.h"
+#include "UAF/Attributes/Transformers/Overwrite.h"
+#include "UAF/Attributes/Transformers/Sanitize.h"
+
+#include "Animation/BuiltInAttributeTypes.h"
+#include "UAF/Attributes/EngineAttributes.h"
 
 #define LOCTEXT_NAMESPACE "AnimNextModule"
 
@@ -83,6 +91,22 @@ namespace UE::UAF
 
 		FSystemFactory::Init();
 
+		FAttributeRuntimeRegistry::Init();
+
+		{
+			FAttributeRuntimeRegistry& RuntimeRegistry = FAttributeRuntimeRegistry::Get();
+
+			RuntimeRegistry.RegisterSetMapInitializer<FFloatAnimationAttribute>();
+			RuntimeRegistry.RegisterSetMapInitializer<FIntegerAnimationAttribute>();
+			RuntimeRegistry.RegisterSetMapInitializer<FStringAnimationAttribute>();
+			RuntimeRegistry.RegisterSetMapInitializer<FBoneTransformAnimationAttribute>();
+
+			RuntimeRegistry.RegisterAttributeValueTransformer<Transformers::FAccumulate_BoneTransformAttribute>();
+			RuntimeRegistry.RegisterAttributeValueTransformer<Transformers::FInterpolate_BoneTransformAttribute>();
+			RuntimeRegistry.RegisterAttributeValueTransformer<Transformers::FOverwrite_BoneTransformAttribute>();
+			RuntimeRegistry.RegisterAttributeValueTransformer<Transformers::FSanitize_BoneTransformAttribute>();
+		}
+
 #if UAF_TRACE_ENABLED
 		IModularFeatures::Get().RegisterModularFeature(IRewindDebuggerRuntimeExtension::ModularFeatureName, &GRewindDebuggerUAFRuntime);
 #endif
@@ -126,6 +150,7 @@ namespace UE::UAF
 		IModularFeatures::Get().UnregisterModularFeature(IRewindDebuggerRuntimeExtension::ModularFeatureName, &GRewindDebuggerUAFRuntime);
 #endif
 		
+		FAttributeRuntimeRegistry::Destroy();
 		FRigVMRuntimeDataRegistry::Destroy();
 		FDataRegistry::Destroy();
 		FSystemFactory::Destroy();
diff --git a/Engine/Plugins/Experimental/UAF/UAF/Source/UAF/Private/DecompressionTools.cpp b/Engine/Plugins/Experimental/UAF/UAF/Source/UAF/Private/DecompressionTools.cpp
--- a/Engine/Plugins/Experimental/UAF/UAF/Source/UAF/Private/DecompressionTools.cpp
+++ b/Engine/Plugins/Experimental/UAF/UAF/Source/UAF/Private/DecompressionTools.cpp
@@ -4,6 +4,7 @@
 #include "Animation/AnimSequenceDecompressionContext.h"
 #include "Animation/AnimCompressionTypes.h"
 #include "Animation/AnimBoneCompressionCodec.h"
+#include "Animation/AnimCurveCompressionCodec.h"
 
 #include "AnimationRuntime.h"
 #include "AnimEncoding.h"
@@ -14,6 +15,10 @@
 #include "Animation/AnimSequence.h"
 #include "HAL/ConsoleManager.h"
 
+#include "UAF/Attributes/AttributeCollection.h"
+#include "UAF/Attributes/AttributeEvaluationContext.h"
+#include "UAF/Attributes/EngineAttributes.h"
+
 CSV_DECLARE_CATEGORY_MODULE_EXTERN(ENGINE_API, Animation);
 
 DECLARE_CYCLE_STAT(TEXT("Build Anim Track Pairs"), STAT_BuildAnimTrackPairs, STATGROUP_Anim);
@@ -59,33 +64,493 @@ static bool CanEvaluateRawAnimationData(const UAnimSequence* AnimSequence)
 #endif
 }
 
-bool FDecompressionTools::UseRawDataForPoseExtraction(const UAnimSequence* AnimSequence, const FLODPose& AnimationPoseData)
+static bool UseRawDataForPoseExtraction(const UAnimSequence* AnimSequence, const FAnimExtractContext& ExtractionContext)
 {
 	return CanEvaluateRawAnimationData(AnimSequence) && (
 #if WITH_EDITOR
 		!AnimSequence->IsCompressedDataValid() ||
 #endif
-		(AnimSequence->GetSkeletonVirtualBoneGuid() != AnimSequence->GetSkeleton()->GetVirtualBoneGuid()) || AnimationPoseData.GetDisableRetargeting() || AnimationPoseData.ShouldUseRawData() ||
+		// TODO
+		(AnimSequence->GetSkeletonVirtualBoneGuid() != AnimSequence->GetSkeleton()->GetVirtualBoneGuid())// || ExtractionContext.bDisableRetargeting || ExtractionContext.bUseRawData
 #if WITH_EDITOR
-		GetForceRawData() ||
+		|| GetForceRawData()
 #endif // WITH_EDITOR
-		AnimationPoseData.ShouldUseSourceData());
+		);
 }
 
-bool FDecompressionTools::ShouldUseRawData(const UAnimSequence* AnimSequence, const FLODPose& AnimationPoseData)
+bool FDecompressionTools::ShouldUseRawData(const UAnimSequence* AnimSequence, const FAnimExtractContext& ExtractionContext)
 {
 	const bool bUseRawData =
 #if WITH_EDITOR
 		GetForceRawData() ||
 #endif // WITH_EDITOR
-		UseRawDataForPoseExtraction(AnimSequence, AnimationPoseData);
+		UE::UAF::UseRawDataForPoseExtraction(AnimSequence, ExtractionContext) ||
+		// TODO
+		false; //ExtractionContext.bUseSourceData;
 
 	return bUseRawData;
 }
 
 // --- ---
+namespace Private
+{
+	static void DecompressBoneTransforms(
+		TAttributeSetMap<FBoneTransformAnimationAttribute>& OutBoneTransforms,
+		const FCompressedAnimSequence& CompressedData,
+		const FAnimExtractContext& ExtractionContext,
+		const FAttributeEvaluationContext& VMContext,
+		FAnimSequenceDecompressionContext& DecompressionContext,
+		const TArray<FTransform>& RetargetTransforms,
+		const FRootMotionReset& RootMotionReset)
+	{
+		const FAttributeTypedSetPtr& BoneSet = OutBoneTransforms.GetTypedSet();
+		const int32 NumBoneTransforms = OutBoneTransforms.Num();
+
+		const FTransform RootRefTransform = OutBoneTransforms.GetData()->Value;
+
+		const int32 NumTracks = CompressedData.CompressedTrackToSkeletonMapTable.Num();
+
+		const USkeleton* SourceSkeleton = DecompressionContext.GetSourceSkeleton();
+		const USkeleton* TargetSkeleton = VMContext.GetSkeleton();
+		const FSkeletonRemapping& SkeletonRemapping = UE::Anim::FSkeletonRemappingRegistry::Get().GetRemapping(SourceSkeleton, TargetSkeleton);
+		const bool bIsSkeletonRemappingValid = SkeletonRemapping.IsValid();
+
+		const bool bUseSourceRetargetModes = TargetSkeleton->GetUseRetargetModesFromCompatibleSkeleton();
+		// TODO
+		const bool bDisableRetargeting = false;// ExtractionContext.bDisableRetargeting;
+
+		FGetBonePoseScratchArea& ScratchArea = FGetBonePoseScratchArea::Get();
+		BoneTrackArray& RotationScalePairs = ScratchArea.RotationScalePairs;
+		BoneTrackArray& TranslationPairs = ScratchArea.TranslationPairs;
+		BoneTrackArray& AnimScaleRetargetingPairs = ScratchArea.AnimScaleRetargetingPairs;
+		BoneTrackArray& AnimRelativeRetargetingPairs = ScratchArea.AnimRelativeRetargetingPairs;
+		BoneTrackArray& OrientAndScaleRetargetingPairs = ScratchArea.OrientAndScaleRetargetingPairs;
+
+		// build a list of desired bones
+		RotationScalePairs.Reset();
+		TranslationPairs.Reset();
+		AnimScaleRetargetingPairs.Reset();
+		AnimRelativeRetargetingPairs.Reset();
+		OrientAndScaleRetargetingPairs.Reset();
+
+		const bool bIsMeshSpaceAdditive = DecompressionContext.GetAdditiveType() == AAT_RotationOffsetMeshSpace;
+		TBitArray<>& AnimatedCompactRotations = ScratchArea.AnimatedCompactRotations;
+		if (bIsMeshSpaceAdditive)
+		{
+			AnimatedCompactRotations.Init(false, NumBoneTransforms);
+			//OutAnimationPoseData.Flags |= ELODPoseFlags::MeshSpaceAdditive;			// TODO: Should be set when we create the collection by the caller
+		}
+		else if (DecompressionContext.GetAdditiveType() == AAT_LocalSpaceBase)
+		{
+			//OutAnimationPoseData.Flags |= ELODPoseFlags::LocalSpaceAdditive;
+		}
+
+		// this is not guaranteed for AnimSequences though... If Root is not animated, Track will not exist.
+		const bool bIsFirstBoneRoot = BoneSet->GetBindingIndex(FAttributeSetIndex(0)).IsRootBone();
+
+		{
+			SCOPE_CYCLE_COUNTER(STAT_BuildAnimTrackPairs);
+
+			for (int32 TrackIndex = 0; TrackIndex < NumTracks; TrackIndex++)
+			{
+				const int32 SourceSkeletonBoneIndex = CompressedData.GetSkeletonIndexFromTrackIndex(TrackIndex);
+				const int32 TargetSkeletonBoneIndex = bIsSkeletonRemappingValid ? SkeletonRemapping.GetTargetSkeletonBoneIndex(SourceSkeletonBoneIndex) : SourceSkeletonBoneIndex;
+
+				if (TargetSkeletonBoneIndex != INDEX_NONE)
+				{
+					const FAttributeSetIndex BoneSetIndex = BoneSet->GetIndex(FAttributeBindingIndex(TargetSkeletonBoneIndex));
+
+					// Skip bones not in current set
+					if (BoneSetIndex.IsValid())
+					{
+						RotationScalePairs.Add(BoneTrackPair(BoneSetIndex.GetInt(), TrackIndex));
+
+						if (bIsMeshSpaceAdditive)
+						{
+							AnimatedCompactRotations[BoneSetIndex.GetInt()] = true;
+						}
+
+						// Check what retarget mode to use for the translational retargeting for this specific bone.
+						const EBoneTranslationRetargetingMode::Type RetargetMode = FAnimationRuntime::GetBoneTranslationRetargetingMode(
+							bUseSourceRetargetModes,
+							SourceSkeletonBoneIndex,
+							TargetSkeletonBoneIndex,
+							SourceSkeleton,
+							TargetSkeleton,
+							bDisableRetargeting);
+
+						// Skip extracting translation component for EBoneTranslationRetargetingMode::Skeleton.
+						switch (RetargetMode)
+						{
+						case EBoneTranslationRetargetingMode::Animation:
+							TranslationPairs.Add(BoneTrackPair(BoneSetIndex.GetInt(), TrackIndex));
+							break;
+						case EBoneTranslationRetargetingMode::AnimationScaled:
+							TranslationPairs.Add(BoneTrackPair(BoneSetIndex.GetInt(), TrackIndex));
+							AnimScaleRetargetingPairs.Add(BoneTrackPair(BoneSetIndex.GetInt(), SourceSkeletonBoneIndex));
+							break;
+						case EBoneTranslationRetargetingMode::AnimationRelative:
+							TranslationPairs.Add(BoneTrackPair(BoneSetIndex.GetInt(), TrackIndex));
+
+							// With baked additives, we can skip 'AnimationRelative' tracks, as the relative transform gets canceled out.
+							// (A1 + Rel) - (A2 + Rel) = A1 - A2.
+							if (!DecompressionContext.IsAdditiveAnimation())
+							{
+								AnimRelativeRetargetingPairs.Add(BoneTrackPair(BoneSetIndex.GetInt(), SourceSkeletonBoneIndex));
+							}
+							break;
+						case EBoneTranslationRetargetingMode::OrientAndScale:
+							TranslationPairs.Add(BoneTrackPair(BoneSetIndex.GetInt(), TrackIndex));
+
+							// Additives remain additives, they're not retargeted.
+							if (!DecompressionContext.IsAdditiveAnimation())
+							{
+								OrientAndScaleRetargetingPairs.Add(BoneTrackPair(BoneSetIndex.GetInt(), SourceSkeletonBoneIndex));
+							}
+							break;
+						}
+					}
+				}
+			}
+		}
+
+		{
+			SCOPE_CYCLE_COUNTER(STAT_ExtractPoseFromAnimData);
+			CSV_SCOPED_TIMING_STAT(Animation, ExtractPoseFromAnimData);
+			CSV_CUSTOM_STAT(Animation, NumberOfExtractedAnimations, 1, ECsvCustomStatOp::Accumulate);
+
+			DecompressionContext.Seek(ExtractionContext.CurrentTime);
+
+			if (!RotationScalePairs.IsEmpty())
+			{
+#if DEFAULT_SOA
+				// get the remaining bone atoms
+				//CompressedData.BoneCompressionCodec->DecompressPose(DecompressionContext, UE::Anim::FAnimPoseDecompressionData(RotationScalePairs, TranslationPairs, RotationScalePairs, OutAnimationPoseData.LocalTransformsView.Rotations, OutAnimationPoseData.LocalTransformsView.Translations, OutAnimationPoseData.LocalTransformsView.Scales3D));
+#else
+				// get the remaining bone atoms
+				//TArrayView<FTransform> OutPoseBones = OutAnimationPoseData.LocalTransforms.Transforms;
+				//CompressedData.BoneCompressionCodec->DecompressPose(DecompressionContext, RotationScalePairs, TranslationPairs, RotationScalePairs, OutPoseBones);
+#endif
+
+				static_assert(sizeof(FBoneTransformAnimationAttribute) == sizeof(FTransform), "Sizes should match since we coerce between the two");
+
+				TArrayView<FTransform> TransformsView(reinterpret_cast<FTransform*>(OutBoneTransforms.GetData()), NumBoneTransforms);
+				CompressedData.BoneCompressionCodec->DecompressPose(DecompressionContext, RotationScalePairs, TranslationPairs, RotationScalePairs, TransformsView);
+			}
+		}
+
+		// Retarget the pose onto the target skeleton (correcting for differences in rest poses)
+		if (SkeletonRemapping.RequiresReferencePoseRetarget())
+		{
+			if (DecompressionContext.IsAdditiveAnimation())
+			{
+				// Skip root since it might contain root motion
+				for (FAttributeSetIndex BoneSetIndex(bIsFirstBoneRoot ? 1 : 0); BoneSetIndex < NumBoneTransforms; ++BoneSetIndex)
+				{
+					const int32 TargetSkeletonBoneIndex = BoneSet->GetBindingIndex(BoneSetIndex).GetInt();
+
+					// Mesh space additives do not require fix-up
+					if (DecompressionContext.GetAdditiveType() == AAT_LocalSpaceBase)
+					{
+						OutBoneTransforms[BoneSetIndex].Value.SetRotation(SkeletonRemapping.RetargetAdditiveRotationToTargetSkeleton(TargetSkeletonBoneIndex, OutBoneTransforms[BoneSetIndex].Value.GetRotation()));
+					}
+
+					// Check what retarget mode to use for the translational retargeting for this specific bone.
+					const int32 SourceSkeletonBoneIndex = SkeletonRemapping.GetSourceSkeletonBoneIndex(TargetSkeletonBoneIndex);
+					const EBoneTranslationRetargetingMode::Type RetargetMode = FAnimationRuntime::GetBoneTranslationRetargetingMode(
+						bUseSourceRetargetModes,
+						SourceSkeletonBoneIndex,
+						TargetSkeletonBoneIndex,
+						SourceSkeleton,
+						TargetSkeleton,
+						bDisableRetargeting);
+
+					if (RetargetMode != EBoneTranslationRetargetingMode::Skeleton)
+					{
+						OutBoneTransforms[BoneSetIndex].Value.SetTranslation(SkeletonRemapping.RetargetAdditiveTranslationToTargetSkeleton(TargetSkeletonBoneIndex, OutBoneTransforms[BoneSetIndex].Value.GetTranslation()));
+					}
+				}
+			}
+			else
+			{
+				// Skip root since it might contain root motion
+				for (FAttributeSetIndex BoneSetIndex(bIsFirstBoneRoot ? 1 : 0); BoneSetIndex < NumBoneTransforms; ++BoneSetIndex)
+				{
+					const int32 TargetSkeletonBoneIndex = BoneSet->GetBindingIndex(BoneSetIndex).GetInt();
+
+					OutBoneTransforms[BoneSetIndex].Value.SetRotation(SkeletonRemapping.RetargetBoneRotationToTargetSkeleton(TargetSkeletonBoneIndex, OutBoneTransforms[BoneSetIndex].Value.GetRotation()));
+
+					// Check what retarget mode to use for the translational retargeting for this specific bone.
+					const int32 SourceSkeletonBoneIndex = SkeletonRemapping.GetSourceSkeletonBoneIndex(TargetSkeletonBoneIndex);
+					const EBoneTranslationRetargetingMode::Type RetargetMode = FAnimationRuntime::GetBoneTranslationRetargetingMode(
+						bUseSourceRetargetModes,
+						SourceSkeletonBoneIndex,
+						TargetSkeletonBoneIndex,
+						SourceSkeleton,
+						TargetSkeleton,
+						bDisableRetargeting);
+
+					if (RetargetMode != EBoneTranslationRetargetingMode::Skeleton)
+					{
+						OutBoneTransforms[BoneSetIndex].Value.SetTranslation(SkeletonRemapping.RetargetBoneTranslationToTargetSkeleton(TargetSkeletonBoneIndex, OutBoneTransforms[BoneSetIndex].Value.GetTranslation()));
+					}
+				}
+			}
+		}
+
+		if (bIsFirstBoneRoot)
+		{
+			// Once pose has been extracted, snap root bone back to first frame if we are extracting root motion.
+			if ((ExtractionContext.bExtractRootMotion && RootMotionReset.bEnableRootMotion) || RootMotionReset.bForceRootLock)
+			{
+				FTransform RootTransform = OutBoneTransforms.GetData()->Value;
+				RootMotionReset.ResetRootBoneForRootMotion(RootTransform, RootRefTransform);
+				OutBoneTransforms.GetData()->Value = RootTransform;
+			}
+		}
+
+		// Anim Scale Retargeting
+		if (!AnimScaleRetargetingPairs.IsEmpty())
+		{
+			const TArray<FTransform>& AuthoredOnRefSkeleton = RetargetTransforms;
+			const TArray<FTransform>& TargetRefSkeleton = TargetSkeleton->GetRefLocalPoses();
+
+			for (const BoneTrackPair& BonePair : AnimScaleRetargetingPairs)
+			{
+				const int32 SourceSkeletonBoneIndex = BonePair.TrackIndex;
+				if (!AuthoredOnRefSkeleton.IsValidIndex(SourceSkeletonBoneIndex))
+				{
+					continue;
+				}
+
+				// @todo - precache that in FBoneContainer when we have SkeletonIndex->TrackIndex mapping. So we can just apply scale right away.
+				const double SourceTranslationLengthSq = AuthoredOnRefSkeleton[SourceSkeletonBoneIndex].GetTranslation().SizeSquared();
+				if (SourceTranslationLengthSq > (UE_KINDA_SMALL_NUMBER * UE_KINDA_SMALL_NUMBER))
+				{
+					const FAttributeSetIndex BoneSetIndex(BonePair.AtomIndex);
+					const FAttributeBindingIndex BindingIndex = BoneSet->GetBindingIndex(BoneSetIndex);
+					const double TargetTranslationLength = TargetRefSkeleton[BindingIndex.GetInt()].GetTranslation().Size();
+					OutBoneTransforms[BoneSetIndex].Value.ScaleTranslation(TargetTranslationLength / FMath::Sqrt(SourceTranslationLengthSq));
+				}
+			}
+		}
+
+		// Anim Relative Retargeting
+		if (!AnimRelativeRetargetingPairs.IsEmpty())
+		{
+			const TArray<FTransform>& AuthoredOnRefSkeleton = RetargetTransforms;
+			const TArray<FTransform>& TargetRefSkeleton = TargetSkeleton->GetRefLocalPoses();
+
+			for (const BoneTrackPair& BonePair : AnimRelativeRetargetingPairs)
+			{
+				const int32 SourceSkeletonBoneIndex = BonePair.TrackIndex;
+				if (!AuthoredOnRefSkeleton.IsValidIndex(SourceSkeletonBoneIndex))
+				{
+					continue;
+				}
+
+				// Remap the base pose onto the target skeleton so that we are working entirely in target space
+				FTransform BaseTransform = AuthoredOnRefSkeleton[SourceSkeletonBoneIndex];
+				if (SkeletonRemapping.RequiresReferencePoseRetarget())
+				{
+					const int32 TargetSkeletonBoneIndex = SkeletonRemapping.GetTargetSkeletonBoneIndex(SourceSkeletonBoneIndex);
+					BaseTransform = SkeletonRemapping.RetargetBoneTransformToTargetSkeleton(TargetSkeletonBoneIndex, BaseTransform);
+				}
+
+				const FAttributeSetIndex BoneSetIndex(BonePair.AtomIndex);
+				const FAttributeBindingIndex BindingIndex = BoneSet->GetBindingIndex(BoneSetIndex);
+				const FTransform& RefPoseTransform = TargetRefSkeleton[BindingIndex.GetInt()];
+
+				// Apply the retargeting as if it were an additive difference between the current skeleton and the retarget skeleton. 
+				OutBoneTransforms[BoneSetIndex].Value.SetRotation(OutBoneTransforms[BoneSetIndex].Value.GetRotation() * BaseTransform.GetRotation().Inverse() * RefPoseTransform.GetRotation());
+				OutBoneTransforms[BoneSetIndex].Value.SetTranslation(OutBoneTransforms[BoneSetIndex].Value.GetTranslation() + (RefPoseTransform.GetTranslation() - BaseTransform.GetTranslation()));
+				OutBoneTransforms[BoneSetIndex].Value.SetScale3D(OutBoneTransforms[BoneSetIndex].Value.GetScale3D() * (RefPoseTransform.GetScale3D() * BaseTransform.GetSafeScaleReciprocal(BaseTransform.GetScale3D())));
+				OutBoneTransforms[BoneSetIndex].Value.NormalizeRotation();
+			}
+		}
+
+		// TODO : Have to recreate GetRetargetSourceCachedData
+		// Translation 'Orient and Scale' Translation Retargeting
+		//const int32 NumBonesToOrientAndScaleRetarget = OrientAndScaleRetargetingPairs.Num();
+		//if (NumBonesToOrientAndScaleRetarget > 0)
+		//{
+		//	const FRetargetSourceCachedData& RetargetSourceCachedData = RequiredBones.GetRetargetSourceCachedData(DecompressionContext.AnimName, SkeletonRemapping, RetargetTransforms);
+		//	const TArray<FOrientAndScaleRetargetingCachedData>& OrientAndScaleDataArray = RetargetSourceCachedData.OrientAndScaleData;
+		//	const TArray<int32>& CompactPoseIndexToOrientAndScaleIndex = RetargetSourceCachedData.CompactPoseIndexToOrientAndScaleIndex;
+
+		//	// If we have any cached retargeting data.
+		//	if ((OrientAndScaleDataArray.Num() > 0) && (CompactPoseIndexToOrientAndScaleIndex.Num() == RequiredBones.GetCompactPoseNumBones()))
+		//	{
+		//		for (int32 Index = 0; Index < NumBonesToOrientAndScaleRetarget; Index++)
+		//		{
+		//			const BoneTrackPair& BonePair = OrientAndScaleRetargetingPairs[Index];
+		//			const FCompactPoseBoneIndex CompactPoseBoneIndex(BonePair.AtomIndex);
+		//			const int32 OrientAndScaleIndex = CompactPoseIndexToOrientAndScaleIndex[CompactPoseBoneIndex.GetInt()];
+		//			if (OrientAndScaleIndex != INDEX_NONE)
+		//			{
+		//				const FOrientAndScaleRetargetingCachedData& OrientAndScaleData = OrientAndScaleDataArray[OrientAndScaleIndex];
+		//				FTransform& BoneTransform = OutAnimationPoseData.LocalTransforms[CompactPoseBoneIndex];
+		//				const FVector AnimatedTranslation = BoneTransform.GetTranslation();
+
+		//				// If Translation is not animated, we can just copy the TargetTranslation. No retargeting needs to be done.
+		//				const FVector NewTranslation = (AnimatedTranslation - OrientAndScaleData.SourceTranslation).IsNearlyZero(BONE_TRANS_RT_ORIENT_AND_SCALE_PRECISION) ?
+		//					OrientAndScaleData.TargetTranslation :
+		//					OrientAndScaleData.TranslationDeltaOrient.RotateVector(AnimatedTranslation) * OrientAndScaleData.TranslationScale;
+
+		//				BoneTransform.SetTranslation(NewTranslation);
+		//			}
+		//		}
+		//	}
+		//}
+		ensureMsgf(OrientAndScaleRetargetingPairs.IsEmpty(), TEXT("OrientAndScaleRetargetingPairs found for %s (%s -> %s), Orient and Scale retargeting is not yet supported in AnimNext"), *DecompressionContext.AnimName.ToString(), SourceSkeleton ? *SourceSkeleton->GetName() : TEXT("None"), TargetSkeleton ? *TargetSkeleton->GetName() : TEXT("None"));
+
+		if (bIsMeshSpaceAdditive)
+		{
+			// When an animation is a mesh-space additive, bones that aren't animated will end up with some non-identity
+			// delta relative to the base used to create the additive. This is because the delta is calculated in mesh-space
+			// unlike regular additive animations where bones that aren't animated has an identity delta. For rotations,
+			// this mesh-space delta will be the parent bone rotation.
+			// However, if a bone isn't animated in the sequence but present on the target skeleton, we have no data for it
+			// and the output pose will contain an identity delta which isn't what we want. As such, bones missing from
+			// the sequence have their rotation set to their parent.
+
+			for (FAttributeSetIndex BoneSetIndex(0); BoneSetIndex < NumBoneTransforms; ++BoneSetIndex)
+			{
+				if (!AnimatedCompactRotations[BoneSetIndex.GetInt()])
+				{
+					// This bone wasn't animated in the sequence, fix it up
+					const FAttributeSetIndex ParentBoneSetIndex = BoneSet->GetParentIndex(BoneSetIndex);
+					if (ParentBoneSetIndex.IsValid())
+					{
+						// Propagate the rotation of our parent to maintain the same orientation in mesh space
+						const FQuat ParentRotation = OutBoneTransforms[ParentBoneSetIndex].Value.GetRotation();
+						OutBoneTransforms[BoneSetIndex].Value.SetRotation(ParentRotation);
+					}
+				}
+			}
+		}
+	}
+
+	static void GetBoneTransforms(
+		const UAnimSequence* AnimSequence,
+		const FAnimExtractContext& ExtractionContext,
+		const FAttributeEvaluationContext& VMContext,
+		TAttributeSetMap<FBoneTransformAnimationAttribute>& OutBoneTransforms,
+		bool bForceUseRawData)
+	{
+		SCOPE_CYCLE_COUNTER(STAT_AnimSeq_GetBonePose);
+		CSV_SCOPED_TIMING_STAT(Animation, AnimSeq_GetBonePose);
+
+		check(!bForceUseRawData || CanEvaluateRawAnimationData(AnimSequence));
+
+		const bool bIsBakedAdditive = !bForceUseRawData && AnimSequence->IsValidAdditive();
+
+		const USkeleton* MySkeleton = AnimSequence->GetSkeleton();
+		if (!MySkeleton)
+		{
+			return;
+		}
+
+		// TODO
+		const bool bDisableRetargeting = false;// ExtractionContext.bDisableRetargeting;
+
+		// if retargeting is disabled, we initialize pose with 'Retargeting Source' ref pose.
+		if (!bIsBakedAdditive && bDisableRetargeting)
+		{
+			const TArray<FTransform>& AuthoredOnRefSkeleton = AnimSequence->GetRetargetTransforms();
+
+			const int32 NumBoneTransforms = OutBoneTransforms.Num();
+			const int32 NumRawSkeletonBones = AnimSequence->GetSkeleton()->GetReferenceSkeleton().GetRawBoneNum();
+			const FAttributeTypedSetPtr& BoneSet = OutBoneTransforms.GetTypedSet();
+
+			for (FAttributeSetIndex AttributeSetIndex(0); AttributeSetIndex < NumBoneTransforms; ++AttributeSetIndex)
+			{
+				const FAttributeBindingIndex BindingIndex = BoneSet->GetBindingIndex(AttributeSetIndex);
+				const int32 SkeletonBoneIndex = BindingIndex.GetInt();
+
+				// Virtual bones are part of the retarget transform pose, so if the pose has not been updated (recently) there might be a mismatch
+				if (SkeletonBoneIndex < NumRawSkeletonBones || AuthoredOnRefSkeleton.IsValidIndex(SkeletonBoneIndex))
+				{
+					OutBoneTransforms[AttributeSetIndex].Value = AuthoredOnRefSkeleton[SkeletonBoneIndex];
+				}
+			}
+		}
+
+		UAnimSequence::FScopedCompressedAnimSequence PlatformCompressedData = AnimSequence->GetCompressedData(ExtractionContext);
+		int32 NumTracks = 0;
+		{
+#if WITH_EDITOR
+			NumTracks = bForceUseRawData ? AnimSequence->GetDataModelInterface()->GetNumBoneTracks() : PlatformCompressedData.Get().CompressedTrackToSkeletonMapTable.Num();
+#else
+			NumTracks = PlatformCompressedData.Get().CompressedTrackToSkeletonMapTable.Num();
+#endif
+		}
+		const bool bTreatAnimAsAdditive = AnimSequence->IsValidAdditive() && !bForceUseRawData; // Raw data is never additive
+		const FRootMotionReset RootMotionReset(
+			AnimSequence->bEnableRootMotion,
+			AnimSequence->RootMotionRootLock,
+			AnimSequence->bForceRootLock,
+			AnimSequence->ExtractRootTrackTransform(FAnimExtractContext(0.0), /*&RequiredBones*/nullptr),
+			bTreatAnimAsAdditive);
+
+#if WITH_EDITOR
+		// Evaluate raw (source) curve and bone data
+		if (bForceUseRawData)
+		{
+			// TODO : Curves support
+			//{
+			//	FSkeletonRemappingCurve RemappedCurve(OutAnimationPoseData.GetCurve(), RequiredBones, AnimSequence->GetSkeleton());
+			//	FAnimationPoseData RemappedPoseData(OutAnimationPoseData.GetPose(), RemappedCurve.GetCurve(), OutAnimationPoseData.GetAttributes());
+
+			//	const UE::Anim::DataModel::FEvaluationContext EvaluationContext(ExtractionContext.CurrentTime, DataModelInterface->GetFrameRate(), GetRetargetTransformsSourceName(), GetRetargetTransforms(), Interpolation);
+			//	DataModelInterface->Evaluate(RemappedPoseData, EvaluationContext);
+			//}
+
+			//if ((ExtractionContext.bExtractRootMotion && RootMotionReset.bEnableRootMotion) || RootMotionReset.bForceRootLock)
+			//{
+			//	RootMotionReset.ResetRootBoneForRootMotion(OutPose[FCompactPoseBoneIndex(0)], RequiredBones);
+			//}
+
+			return;
+		}
+		else
+#endif // WITH_EDITOR
+		{
+			// Only try and evaluate compressed bone data if the animation contains any bone tracks
+			if (NumTracks != 0)
+			{
+				// Evaluate compressed bone data
+				FAnimSequenceDecompressionContext DecompContext(AnimSequence->GetSamplingFrameRate()
+					, AnimSequence->GetSamplingFrameRate().AsFrameTime(AnimSequence->GetPlayLength()).RoundToFrame().Value
+					, AnimSequence->Interpolation
+					, AnimSequence->GetRetargetTransformsSourceName()
+					, *PlatformCompressedData.Get().CompressedDataStructure
+					, AnimSequence->GetSkeleton()->GetRefLocalPoses()
+					, PlatformCompressedData.Get().CompressedTrackToSkeletonMapTable
+					, AnimSequence->GetSkeleton()
+					, AnimSequence->IsValidAdditive()
+					, AnimSequence->GetAdditiveAnimType());
+
+				DecompressBoneTransforms(OutBoneTransforms, PlatformCompressedData.Get(), ExtractionContext, VMContext, DecompContext, AnimSequence->GetRetargetTransforms(), RootMotionReset);
+			}
+		}
+	}
+
+	struct FBlendedCurveWrapper : public FBlendedCurve
+	{
+		UE::Anim::FCurveElement& operator[](int32 Index)
+		{
+			return Elements[Index];
+		}
+
+		const UE::Anim::FCurveElement& operator[](int32 Index) const
+		{
+			return Elements[Index];
+		}
+	};
+}
 
-/*static*/ void FDecompressionTools::GetAnimationPose(const UAnimSequence* AnimSequence, const FAnimExtractContext& ExtractionContext, FLODPose& OutAnimationPoseData, bool bForceUseRawData)
+void FDecompressionTools::GetAnimationPose(const UAnimSequence* AnimSequence, const FAnimExtractContext& ExtractionContext, FLODPose& OutAnimationPoseData, bool bForceUseRawData)
 {
 	if (OutAnimationPoseData.GetRefPose().IsValid() == false)
 	{
@@ -125,6 +590,52 @@ bool FDecompressionTools::ShouldUseRawData(const UAnimSequence* AnimSequence, co
 
 }
 
+void FDecompressionTools::GetAnimationPose(const UAnimSequence* AnimSequence, const FAnimExtractContext& ExtractionContext, const FAttributeEvaluationContext& VMContext, FAttributeCollection& OutCollection, bool bForceUseRawData)
+{
+	using namespace UE::UAF;
+
+	const FAttributeMappingKey MappingKey = FAttributeMappingKey::MakeFromTo<FBoneTransformAnimationAttribute>();
+	TAttributeSetMap<FBoneTransformAnimationAttribute>* BoneTransforms = OutCollection.FindSetMap<FBoneTransformAnimationAttribute>(MappingKey);
+	if (BoneTransforms == nullptr || BoneTransforms->IsEmpty())
+	{
+		// Not requesting bone transforms
+		return;
+	}
+
+	check(VMContext.IsValid());
+
+	// @todo anim: if compressed and baked in the future, we don't have to do this 
+	if (bForceUseRawData && AnimSequence->IsValidAdditive())
+	{
+		switch (AnimSequence->GetAdditiveAnimType())
+		{
+		case AAT_LocalSpaceBase:
+		{
+			//GetBonePose_Additive(AnimSequence, ExtractionContext, OutAnimationPoseData);
+		}
+		break;
+
+		case AAT_RotationOffsetMeshSpace:
+		{
+			//GetBonePose_AdditiveMeshRotationOnly(AnimSequence, ExtractionContext, OutAnimationPoseData);
+		}
+		break;
+
+		default:
+			break;
+		}
+	}
+	else
+	{
+		Private::GetBoneTransforms(AnimSequence, ExtractionContext, VMContext, *BoneTransforms, bForceUseRawData);
+	}
+
+	// Check that all bone atoms coming from animation are normalized
+#if DO_CHECK && WITH_EDITORONLY_DATA
+	//check(OutAnimationPoseData.LocalTransformsView.IsValid());
+#endif
+}
+
 void FDecompressionTools::GetBonePose(const UAnimSequence* AnimSequence, const FAnimExtractContext& ExtractionContext, FLODPose& OutAnimationPoseData, bool bForceUseRawData /*= false*/)
 {
 	SCOPE_CYCLE_COUNTER(STAT_AnimSeq_GetBonePose);
@@ -252,6 +763,81 @@ void FDecompressionTools::GetAnimationCurves(const UAnimSequence* AnimSequence,
 	AnimSequence->EvaluateCurveData(OutCurves, ExtractionContext.CurrentTime, bForceUseRawData);
 }
 
+void FDecompressionTools::GetAnimationCurves(const UAnimSequence* AnimSequence, const FAnimExtractContext& ExtractionContext, const FAttributeEvaluationContext& VMContext, FAttributeCollection& OutCollection, bool bForceUseRawData)
+{
+	using namespace UE::UAF;
+
+	SCOPE_CYCLE_COUNTER(STAT_AnimSeq_EvalCurveData);
+
+	if (!VMContext.IsValid())
+	{
+		return;
+	}
+
+	UAnimSequence::FScopedCompressedAnimSequence PlatformCompressedData = AnimSequence->GetCompressedData(ExtractionContext);
+	const FCompressedAnimSequence& CompressedData = PlatformCompressedData.Get();
+
+	const bool bEvaluateRawData = bForceUseRawData || !CompressedData.IsCurveDataValid(AnimSequence);
+	check(!bForceUseRawData || AnimSequence->CanEvaluateRawAnimationData());
+	if (AnimSequence->CanEvaluateRawAnimationData() && bEvaluateRawData)
+	{
+		// TODO: Raw curve support
+#if WITH_EDITOR
+		//UE::Anim::EvaluateFloatCurvesFromModel(DataModelInterface.GetInterface(), OutCurve, ExtractionContext.CurrentTime);
+#else
+		//Super::EvaluateCurveData(OutCurve, ExtractionContext, bForceUseRawData);
+#endif
+	}
+	else if (CompressedData.IsCurveDataValid(AnimSequence))
+	{
+		if (CompressedData.CurveCompressionCodec)
+		{
+			CSV_SCOPED_TIMING_STAT(Animation, EvaluateCurveData);
+
+			Private::FBlendedCurveWrapper Curves;
+			CompressedData.CurveCompressionCodec->DecompressCurves(CompressedData, Curves, ExtractionContext.CurrentTime);
+
+			const FAttributeMappingKey MappingKey = FAttributeMappingKey::MakeFromTo<FFloatAnimationAttribute>();
+			TAttributeSetMap<FFloatAnimationAttribute>* FloatAttributes = OutCollection.FindSetMap<FFloatAnimationAttribute>(MappingKey);
+			if (FloatAttributes != nullptr && !FloatAttributes->IsEmpty())
+			{
+				// Extract our set defined floats out into our set map
+				const FAttributeTypedSetPtr& FloatSet = FloatAttributes->GetTypedSet();
+
+				// TODO: We can perform a tape merge since both inputs are in sorted order (in reverse order since we remove values from our temp curve array)
+				
+				const int32 NumDecompressedCurves = Curves.Num();
+				for (int32 DecompressedCurveIndex = NumDecompressedCurves - 1; DecompressedCurveIndex >= 0; --DecompressedCurveIndex)
+				{
+					const FName CurveName = Curves[DecompressedCurveIndex].Name;
+
+					if (const FAttributeBindingIndex BindingIndex = FloatSet->FindBindingIndex(CurveName))
+					{
+						// This curve is mapped within our binding, it is statically defined
+						if (const FAttributeSetIndex AttributeIndex = FloatSet->GetIndex(BindingIndex))
+						{
+							// This curve is within our current named set, keep it
+							(*FloatAttributes)[AttributeIndex].Value = Curves[DecompressedCurveIndex].Value;
+						}
+						else
+						{
+							// This curve is stripped and isn't within our current named set
+						}
+					}
+					else
+					{
+						// This curve isn't mapped within our binding, it is dynamically defined
+
+						// TODO: Append dynamic attribute
+					}
+				}
+			}
+
+			// TODO: Append our dynamic float values (could do it in loop above, one value at a time)
+		}
+	}
+}
+
 void FDecompressionTools::GetAnimationAttributes(const UAnimSequence* AnimSequence, const FAnimExtractContext& ExtractionContext, const FReferencePose& RefPose, UE::Anim::FStackAttributeContainer& OutAttributes, bool bForceUseRawData)
 {
 	QUICK_SCOPE_CYCLE_COUNTER(STAT_EvaluateAttributes);
@@ -288,6 +874,123 @@ void FDecompressionTools::GetAnimationAttributes(const UAnimSequence* AnimSequen
 	}
 }
 
+void FDecompressionTools::GetAnimationAttributes(const UAnimSequence* AnimSequence, const FAnimExtractContext& ExtractionContext, const FAttributeEvaluationContext& VMContext, FAttributeCollection& OutCollection, bool bForceUseRawData)
+{
+	QUICK_SCOPE_CYCLE_COUNTER(STAT_EvaluateAttributes);
+
+	if (!VMContext.IsValid())
+	{
+		return;
+	}
+
+#if WITH_EDITOR && 0	// TODO: Add support for raw attributes
+	if (bForceUseRawData)
+	{
+		AnimSequence->ValidateModel();
+
+		for (const FAnimatedBoneAttribute& Attribute : AnimSequence->GetDataModel()->GetAttributes())
+		{
+			const int32 LODBoneIndex = RefPose.GetLODBoneIndexFromSkeletonBoneIndex(Attribute.Identifier.GetBoneIndex());
+			// Only add attribute if the bone its tied to exists in the currently evaluated set of bones
+			if (LODBoneIndex != INDEX_NONE)
+			{
+				UE::Anim::Attributes::GetAttributeValue(OutAttributes, FCompactPoseBoneIndex(LODBoneIndex), Attribute, ExtractionContext.CurrentTime);
+			}
+		}
+	}
+	else
+#endif // WITH_EDITOR
+	{
+		if (AnimSequence->AttributeCurves.IsEmpty())
+		{
+			return;
+		}
+
+		FAttributeNamedSetPtr NamedSet = VMContext.GetNamedSet();
+		FAttributeTypedSetPtr BoneTypedSet = NamedSet->FindTypedSet<FBoneTransformAnimationAttribute>();
+
+		for (const TPair<FAnimationAttributeIdentifier, FAttributeCurve>& BakedAttribute : AnimSequence->AttributeCurves)
+		{
+			UScriptStruct* AttributeType = BakedAttribute.Key.GetType();
+
+			if (const FSkeletonPoseBoneIndex SkeletonBoneIndex = FSkeletonPoseBoneIndex(BakedAttribute.Key.GetBoneIndex()))
+			{
+				// This attribute is attached to a bone
+
+				const FAttributeBindingIndex BoneBindingIndex(SkeletonBoneIndex);
+				if (BoneTypedSet && BoneTypedSet->GetIndex(BoneBindingIndex).IsValid())
+				{
+					// Our parent bone is present within the current named set
+					if (FAttributeTypedSetPtr AttributeTypedSet = NamedSet->FindTypedSet(AttributeType))
+					{
+						const FName AttributeName = BakedAttribute.Key.GetName();
+
+						if (const FAttributeSetIndex SetIndex = AttributeTypedSet->FindIndex(AttributeName))
+						{
+							// This attribute is present within the current named set
+
+							FAttributeSetMap* SetMap = OutCollection.FindSetMap(FAttributeMappingKey::MakeFromTo(AttributeType));
+							checkf(SetMap != nullptr, TEXT("This attribute exists within our named set but isn't present in the output collection"));
+							checkf(SetMap->GetValueType() == AttributeType, TEXT("Attribute type does not match the set map"));
+
+							SetMap->SetValueWithSetter(SetIndex, [&ExtractionContext, &BakedAttribute](UScriptStruct* SetType, uint8* OutValuePtr)
+								{
+									BakedAttribute.Value.EvaluateToPtr(SetType, ExtractionContext.CurrentTime, OutValuePtr);
+								});
+						}
+						else
+						{
+							// This attribute is stripped within the current named set
+						}
+					}
+					else
+					{
+						// This attribute type is unknown within our binding, it is dynamic
+
+						// TODO: Append dynamic attribute
+					}
+				}
+				else
+				{
+					// Our parent bone is stripped within the current named set
+				}
+			}
+			else
+			{
+				// This attribute isn't attached to a bone
+
+				if (FAttributeTypedSetPtr AttributeTypedSet = NamedSet->FindTypedSet(AttributeType))
+				{
+					const FName AttributeName = BakedAttribute.Key.GetName();
+
+					if (const FAttributeSetIndex SetIndex = AttributeTypedSet->FindIndex(AttributeName))
+					{
+						// This attribute is present within the current named set
+
+						FAttributeSetMap* SetMap = OutCollection.FindSetMap(FAttributeMappingKey::MakeFromTo(AttributeType));
+						checkf(SetMap != nullptr, TEXT("This attribute exists within our named set but isn't present in the output collection"));
+						checkf(SetMap->GetValueType() == AttributeType, TEXT("Attribute type does not match the set map"));
+
+						SetMap->SetValueWithSetter(SetIndex, [&ExtractionContext, &BakedAttribute](UScriptStruct* SetType, uint8* OutValuePtr)
+							{
+								BakedAttribute.Value.EvaluateToPtr(SetType, ExtractionContext.CurrentTime, OutValuePtr);
+							});
+					}
+					else
+					{
+						// This attribute is stripped within the current named set
+					}
+				}
+				else
+				{
+					// This attribute type is unknown within our binding, it is dynamic
+
+					// TODO: Append dynamic attribute
+				}
+			}
+		}
+	}
+}
 
 // --- ---
 
@@ -472,7 +1175,7 @@ void FDecompressionTools::DecompressPose(FLODPose& OutAnimationPoseData,
 				, 0
 				, LODRootBone
 				, DecompressionContext.IsAdditiveAnimation()
-				, OutAnimationPoseData.GetDisableRetargeting());
+				, bDisableRetargeting);
 
 			OutAnimationPoseData.LocalTransformsView[0] = RootAtom;
 		}
diff --git a/Engine/Plugins/Experimental/UAF/UAF/Source/UAF/Private/UAF/AbstractSkeleton/AbstractSkeletonSetBinding.cpp b/Engine/Plugins/Experimental/UAF/UAF/Source/UAF/Private/UAF/AbstractSkeleton/AbstractSkeletonSetBinding.cpp
--- a/Engine/Plugins/Experimental/UAF/UAF/Source/UAF/Private/UAF/AbstractSkeleton/AbstractSkeletonSetBinding.cpp
+++ b/Engine/Plugins/Experimental/UAF/UAF/Source/UAF/Private/UAF/AbstractSkeleton/AbstractSkeletonSetBinding.cpp
@@ -3,13 +3,54 @@
 #include "UAF/AbstractSkeleton/AbstractSkeletonSetBinding.h"
 
 #include "UAF/AbstractSkeleton/AbstractSkeletonSetCollection.h"
-
-// Bone Bindings
+#include "UAF/Attributes/EngineAttributes.h"
 
 #include UE_INLINE_GENERATED_CPP_BY_NAME(AbstractSkeletonSetBinding)
 
+void UAbstractSkeletonSetBinding::PostLoad()
+{
+	Super::PostLoad();
+
+	RebuildBindingData();
+
+#if WITH_EDITOR
+	if (Skeleton)
+	{
+		Skeleton->RegisterOnSkeletonHierarchyChanged(USkeleton::FOnSkeletonHierarchyChanged::CreateUObject(this, &UAbstractSkeletonSetBinding::OnSkeletonChanged));
+		Skeleton->RegisterOnReferenceSkeletonChanged(USkeleton::FOnReferenceSkeletonChanged::CreateUObject(this, &UAbstractSkeletonSetBinding::OnSkeletonChanged));
+	}
+
+	if (SetCollection)
+	{
+		SetCollectionChangedHandle = SetCollection->RegisterOnSetsChanged(FSimpleMulticastDelegate::FDelegate::CreateUObject(this, &UAbstractSkeletonSetBinding::OnSetCollectionChanged));
+	}
+#endif
+}
+
+void UAbstractSkeletonSetBinding::GetPreloadDependencies(TArray<UObject*>& OutDeps)
+{
+	Super::GetPreloadDependencies(OutDeps);
+
+	if (Skeleton)
+	{
+		OutDeps.Add(Skeleton);
+	}
+
+	if (SetCollection)
+	{
+		OutDeps.Add(SetCollection);
+	}
+}
+
+// Bone Bindings
+
 bool UAbstractSkeletonSetBinding::AddBoneToSet(const FName BoneName, const FName SetName)
 {
+	if (!SetCollection)
+	{
+		return false;
+	}
+
 	if (!SetCollection->HasSet(SetName) && SetName != NAME_None)
 	{
 		return false; // Set does not exist
@@ -23,6 +64,7 @@ bool UAbstractSkeletonSetBinding::AddBoneToSet(const FName BoneName, const FName
 	FAbstractSkeleton_BoneBinding& Binding = BoneBindings.AddDefaulted_GetRef();
 	Binding.BoneName = BoneName;
 	Binding.SetName = SetName;
+	RebuildBindingData();
 
 	return true;
 }
@@ -37,6 +79,7 @@ bool UAbstractSkeletonSetBinding::RemoveBoneFromSet(const FName BoneName)
 	if (BindingIndex != INDEX_NONE)
 	{
 		BoneBindings.RemoveAtSwap(BindingIndex);
+		RebuildBindingData();
 		return true;
 	}
 	else
@@ -105,6 +148,7 @@ bool UAbstractSkeletonSetBinding::AddAttributeToSet(const FAnimationAttributeIde
 		if (ensure(UnboundAttribute))
 		{
 			UnboundAttribute->SetName = SetName;
+			RebuildBindingData();
 			return true;
 		}
 		else
@@ -116,6 +160,7 @@ bool UAbstractSkeletonSetBinding::AddAttributeToSet(const FAnimationAttributeIde
 	FAbstractSkeleton_AttributeBinding& Binding = AttributeBindings.AddDefaulted_GetRef();
 	Binding.Attribute = Attribute;
 	Binding.SetName = SetName;
+	RebuildBindingData();
 
 	return true;
 }
@@ -130,6 +175,7 @@ bool UAbstractSkeletonSetBinding::RemoveAttributeFromSet(const FAnimationAttribu
 	if (BindingIndex != INDEX_NONE)
 	{
 		AttributeBindings.RemoveAtSwap(BindingIndex);
+		RebuildBindingData();
 		return true;
 	}
 	else
@@ -159,13 +205,55 @@ const TConstArrayView<FAbstractSkeleton_AttributeBinding> UAbstractSkeletonSetBi
 	return AttributeBindings;
 }
 
+const UE::UAF::FAttributeBindingDataPtr& UAbstractSkeletonSetBinding::GetAttributeBindingData() const
+{
+	return AttributeBindingData;
+}
+
+void UAbstractSkeletonSetBinding::RebuildBindingData()
+{
+	// Clear our old data
+	AttributeBindingData = nullptr;
+
+	if (!Skeleton)
+	{
+		return;
+	}
+
+	// Build fresh data
+	AttributeBindingData = UE::UAF::MakeAttributeBindingData(*this);
+}
+
 bool UAbstractSkeletonSetBinding::SetSkeleton(const TObjectPtr<USkeleton> InSkeleton)
 {
-	Skeleton = InSkeleton;
+	if (Skeleton == InSkeleton)
+	{
+		// Ignore calls to set identical data
+		return true;
+	}
 
+#if WITH_EDITOR
+	if (Skeleton)
+	{
+		Skeleton->UnregisterOnSkeletonHierarchyChanged(this);
+		Skeleton->UnregisterOnReferenceSkeletonChanged(this);
+	}
+#endif
+
+	Skeleton = InSkeleton;
 	BoneBindings.Empty();
 	AttributeBindings.Empty();
 
+	RebuildBindingData();
+
+#if WITH_EDITOR
+	if (InSkeleton)
+	{
+		InSkeleton->RegisterOnSkeletonHierarchyChanged(USkeleton::FOnSkeletonHierarchyChanged::CreateUObject(this, &UAbstractSkeletonSetBinding::OnSkeletonChanged));
+		InSkeleton->RegisterOnReferenceSkeletonChanged(USkeleton::FOnReferenceSkeletonChanged::CreateUObject(this, &UAbstractSkeletonSetBinding::OnSkeletonChanged));
+	}
+#endif
+
 	return true;
 }
 
@@ -174,13 +262,34 @@ TObjectPtr<USkeleton> UAbstractSkeletonSetBinding::GetSkeleton() const
 	return Skeleton;
 }
 
-bool UAbstractSkeletonSetBinding::SetSetCollection(const TObjectPtr<const UAbstractSkeletonSetCollection> InSetCollection)
+bool UAbstractSkeletonSetBinding::SetSetCollection(const TObjectPtr<UAbstractSkeletonSetCollection> InSetCollection)
 {
-	SetCollection = InSetCollection;
+	if (SetCollection == InSetCollection)
+	{
+		// Ignore calls to set identical data
+		return true;
+	}
 
+#if WITH_EDITOR
+	if (SetCollection)
+	{
+		SetCollection->UnregisterOnSetsChanged(SetCollectionChangedHandle);
+	}
+#endif
+
+	SetCollection = InSetCollection;
 	BoneBindings.Empty();
 	AttributeBindings.Empty();
 
+	RebuildBindingData();
+
+#if WITH_EDITOR
+	if (InSetCollection)
+	{
+		SetCollectionChangedHandle = InSetCollection->RegisterOnSetsChanged(FSimpleMulticastDelegate::FDelegate::CreateUObject(this, &UAbstractSkeletonSetBinding::OnSetCollectionChanged));
+	}
+#endif
+
 	return true;
 }
 
@@ -188,3 +297,15 @@ TObjectPtr<const UAbstractSkeletonSetCollection> UAbstractSkeletonSetBinding::Ge
 {
 	return SetCollection;
 }
+
+#if WITH_EDITOR
+void UAbstractSkeletonSetBinding::OnSkeletonChanged()
+{
+	RebuildBindingData();
+}
+
+void UAbstractSkeletonSetBinding::OnSetCollectionChanged()
+{
+	RebuildBindingData();
+}
+#endif
diff --git a/Engine/Plugins/Experimental/UAF/UAF/Source/UAF/Private/UAF/Attributes/AttributeCollection.cpp b/Engine/Plugins/Experimental/UAF/UAF/Source/UAF/Private/UAF/Attributes/AttributeCollection.cpp
new file mode 100644
--- /dev/null
+++ b/Engine/Plugins/Experimental/UAF/UAF/Source/UAF/Private/UAF/Attributes/AttributeCollection.cpp
@@ -0,0 +1,276 @@
+// Copyright Epic Games, Inc. All Rights Reserved.
+
+#include "UAF/Attributes/AttributeCollection.h"
+
+#include "Algo/BinarySearch.h"
+#include "UAF/Attributes/AttributeRuntimeRegistry.h"
+
+namespace UE::UAF
+{
+	FAttributeCollection::FAttributeCollection(const FAttributeCollection& Other)
+		: NamedSet(Other.NamedSet)
+		, SetMapCount(Other.SetMapCount)
+		, SetMapCapacity(Other.SetMapCapacity)
+		, ReallocFun(Other.ReallocFun)
+	{
+		if (Other.SetMaps != nullptr)
+		{
+			const int32 BufferSize = sizeof(FSetMapEntry) * Other.SetMapCount;
+			SetMaps = reinterpret_cast<FSetMapEntry*>((*Other.ReallocFun)(nullptr, 0, BufferSize));
+			FMemory::Memcpy(SetMaps, Other.SetMaps, BufferSize);
+
+			for (int32 MapIndex = 0; MapIndex < Other.SetMapCount; ++MapIndex)
+			{
+				SetMaps[MapIndex].Map = Other.SetMaps[MapIndex].Map->Duplicate(Other.ReallocFun);
+			}
+		}
+		else
+		{
+			SetMaps = nullptr;
+		}
+	}
+
+	FAttributeCollection::FAttributeCollection(FAttributeCollection&& Other)
+		: NamedSet(MoveTemp(Other.NamedSet))
+		, SetMaps(Other.SetMaps)
+		, SetMapCount(Other.SetMapCount)
+		, SetMapCapacity(Other.SetMapCapacity)
+		, ReallocFun(Other.ReallocFun)
+	{
+		Other.SetMaps = nullptr;
+		Other.SetMapCount = 0;
+		Other.SetMapCapacity = 0;
+	}
+
+	FAttributeCollection::FAttributeCollection(const FAttributeNamedSetPtr& InNamedSet, FReallocFun InReallocFun)
+		: NamedSet(InNamedSet)
+		, SetMaps(nullptr)
+		, SetMapCount(0)
+		, SetMapCapacity(0)
+		, ReallocFun(InReallocFun)
+	{
+		check(InReallocFun != nullptr);
+	}
+
+	FAttributeCollection::~FAttributeCollection()
+	{
+		if (SetMaps != nullptr)
+		{
+			for (int32 MapIndex = 0; MapIndex < SetMapCount; ++MapIndex)
+			{
+				ReleaseAttributeSetMap(SetMaps[MapIndex].Map);
+			}
+
+			(*ReallocFun)(reinterpret_cast<uint8*>(SetMaps), sizeof(FSetMapEntry) * SetMapCapacity, 0);
+			SetMaps = nullptr;
+		}
+	}
+
+	FAttributeCollection& FAttributeCollection::operator=(const FAttributeCollection& Other)
+	{
+		if (this != &Other)
+		{
+			checkf(ReallocFun != nullptr, TEXT("When we copy from another collection, we retain our original allocator"));
+
+			// Free anything we might be holding
+			this->~FAttributeCollection();
+
+			NamedSet = Other.NamedSet;
+			SetMapCount = Other.SetMapCount;
+			SetMapCapacity = Other.SetMapCapacity;
+
+			if (Other.SetMaps != nullptr)
+			{
+				const int32 BufferSize = sizeof(FSetMapEntry) * Other.SetMapCount;
+				SetMaps = reinterpret_cast<FSetMapEntry*>((*ReallocFun)(nullptr, 0, BufferSize));
+				FMemory::Memcpy(SetMaps, Other.SetMaps, BufferSize);
+
+				for (int32 MapIndex = 0; MapIndex < Other.SetMapCount; ++MapIndex)
+				{
+					SetMaps[MapIndex].Map = Other.SetMaps[MapIndex].Map->Duplicate(ReallocFun);
+				}
+			}
+			else
+			{
+				SetMaps = nullptr;
+			}
+		}
+
+		return *this;
+	}
+
+	FAttributeCollection& FAttributeCollection::operator=(FAttributeCollection&& Other)
+	{
+		if (this != &Other)
+		{
+			checkf(ReallocFun == Other.ReallocFun, TEXT("When we move from another collection, our allocators must match"));
+
+			// Free anything we might be holding
+			this->~FAttributeCollection();
+
+			NamedSet = MoveTemp(Other.NamedSet);
+			SetMaps = Other.SetMaps;
+			SetMapCount = Other.SetMapCount;
+			SetMapCapacity = Other.SetMapCapacity;
+
+			Other.SetMaps = nullptr;
+			Other.SetMapCount = 0;
+			Other.SetMapCapacity = 0;
+		}
+
+		return *this;
+	}
+
+	bool FAttributeCollection::IsEmpty() const
+	{
+		return SetMapCount == 0;
+	}
+
+	const FAttributeNamedSetPtr& FAttributeCollection::GetNamedSet() const
+	{
+		return NamedSet;
+	}
+
+	int32 FAttributeCollection::NumSetMaps() const
+	{
+		return SetMapCount;
+	}
+
+	FAttributeSetMap* FAttributeCollection::AddSetMap(const FAttributeMappingKey& MappingKey)
+	{
+		if (!NamedSet)
+		{
+			// We are not initialized
+			return nullptr;
+		}
+
+		FAttributeTypedSetPtr TypedSet = NamedSet->FindTypedSet(MappingKey.GetAttributeType());
+		if (!TypedSet)
+		{
+			// Unknown attribute type, it isn't contained within our named set
+			return nullptr;
+		}
+
+		// Use default construction arguments
+		const FAttributeSetMap::FConstructArgs Args(TypedSet, MappingKey.GetValueType(), ReallocFun);
+
+		// We don't know the concrete type of the set map, use indirection to construct it
+		constexpr bool bConstructSetMap = true;
+		return *AddSetMapImpl(Args, bConstructSetMap);
+	}
+
+	FAttributeSetMap** FAttributeCollection::AddSetMapImpl(const FAttributeSetMap::FConstructArgs& Args, bool bConstructSetMap)
+	{
+		UScriptStruct* AttributeType = Args.TypedSet->GetType();
+		UScriptStruct* ValueType = Args.ValueType;
+
+		const FAttributeMappingKey Key(AttributeType, ValueType);
+
+		const int32 MapIndex = Algo::LowerBoundBy(TArrayView<FSetMapEntry>(SetMaps, SetMapCount), Key, [](const FSetMapEntry& Entry) { return Entry.Key; });
+		if (MapIndex < SetMapCount && SetMaps[MapIndex].Key == Key)
+		{
+			// A map already exists that matches the specified types
+			return nullptr;
+		}
+
+		FAttributeSetMap* SetMap = nullptr;
+		if (bConstructSetMap)
+		{
+			SetMap = FAttributeRuntimeRegistry::Get().ConstructSetMap(Args);
+
+			if (SetMap == nullptr)
+			{
+				// Failed to construct a set map, the type has probably not been registered
+				return nullptr;
+			}
+		}
+
+		if (SetMapCount >= SetMapCapacity)
+		{
+			// The first time we add a set, we reserve space for an entry per typed set plus some extra room
+			const bool bAllowQuantize = true;
+			const int32 DesiredCapasity = SetMapCapacity == 0 ? (NamedSet->NumTypedSets() + 4) : (SetMapCapacity + 1);
+			const int32 NewCapacity = DefaultCalculateSlackGrow(DesiredCapasity, SetMapCapacity, sizeof(FAttributeSetMap*), bAllowQuantize);
+
+			SetMaps = reinterpret_cast<FSetMapEntry*>((*ReallocFun)(
+				reinterpret_cast<uint8*>(SetMaps),
+				sizeof(FSetMapEntry) * SetMapCapacity,
+				sizeof(FSetMapEntry) * NewCapacity));
+
+			SetMapCapacity = NewCapacity;
+		}
+
+		const int32 NumToMove = SetMapCount - MapIndex;
+		FMemory::Memmove(SetMaps + MapIndex + 1, SetMaps + MapIndex, sizeof(FSetMapEntry) * NumToMove);
+		SetMaps[MapIndex].Key = Key;
+		SetMaps[MapIndex].Map = SetMap;
+		SetMapCount++;
+
+		return &SetMaps[MapIndex].Map;
+	}
+
+	bool FAttributeCollection::RemoveSetMap(FAttributeSetMap* Map)
+	{
+		if (Map == nullptr)
+		{
+			// Invalid set map
+			return false;
+		}
+
+		return RemoveSetMap(Map->GetMappingKey());
+	}
+
+	bool FAttributeCollection::RemoveSetMap(const FAttributeMappingKey& MappingKey)
+	{
+		const int32 MapIndex = Algo::LowerBoundBy(TArrayView<FSetMapEntry>(SetMaps, SetMapCount), MappingKey, [](const FSetMapEntry& Entry) { return Entry.Key; });
+		if (MapIndex >= SetMapCount || SetMaps[MapIndex].Key != MappingKey)
+		{
+			// No map found that matches the desired types
+			return false;
+		}
+
+		// Destroy and free our set map
+		ReleaseAttributeSetMap(SetMaps[MapIndex].Map);
+
+		// Remove our entry
+		const int32 NumToMove = SetMapCount - MapIndex - 1;
+		FMemory::Memmove(SetMaps + MapIndex, SetMaps + MapIndex + 1, sizeof(FSetMapEntry) * NumToMove);
+		SetMapCount--;
+
+		return true;
+	}
+
+	FAttributeSetMap* FAttributeCollection::FindSetMap(const FAttributeMappingKey& MappingKey)
+	{
+		const int32 MapIndex = Algo::LowerBoundBy(TArrayView<FSetMapEntry>(SetMaps, SetMapCount), MappingKey, [](const FSetMapEntry& Entry) { return Entry.Key; });
+		if (MapIndex >= SetMapCount || SetMaps[MapIndex].Key != MappingKey)
+		{
+			// No map found that matches the desired types
+			return nullptr;
+		}
+
+		return SetMaps[MapIndex].Map;
+	}
+
+	const FAttributeSetMap* FAttributeCollection::FindSetMap(const FAttributeMappingKey& MappingKey) const
+	{
+		const int32 MapIndex = Algo::LowerBoundBy(TArrayView<FSetMapEntry>(SetMaps, SetMapCount), MappingKey, [](const FSetMapEntry& Entry) { return Entry.Key; });
+		if (MapIndex >= SetMapCount || SetMaps[MapIndex].Key != MappingKey)
+		{
+			// No map found that matches the desired types
+			return nullptr;
+		}
+
+		return SetMaps[MapIndex].Map;
+	}
+
+	FAttributeSetMap* FAttributeCollection::FindOrAddSetMap(const FAttributeMappingKey& MappingKey)
+	{
+		if (FAttributeSetMap* SetMap = FindSetMap(MappingKey))
+		{
+			return SetMap;
+		}
+
+		return AddSetMap(MappingKey);
+	}
+}
diff --git a/Engine/Plugins/Experimental/UAF/UAF/Source/UAF/Private/UAF/Attributes/AttributeEvaluationContext.cpp b/Engine/Plugins/Experimental/UAF/UAF/Source/UAF/Private/UAF/Attributes/AttributeEvaluationContext.cpp
new file mode 100644
--- /dev/null
+++ b/Engine/Plugins/Experimental/UAF/UAF/Source/UAF/Private/UAF/Attributes/AttributeEvaluationContext.cpp
@@ -0,0 +1,55 @@
+// Copyright Epic Games, Inc. All Rights Reserved.
+
+#include "UAF/Attributes/AttributeEvaluationContext.h"
+
+#include "UAF/AbstractSkeleton/AbstractSkeletonSetBinding.h"
+
+namespace UE::UAF
+{
+	FAttributeEvaluationContext::FAttributeEvaluationContext() = default;
+
+	FAttributeEvaluationContext::FAttributeEvaluationContext(const UAbstractSkeletonSetBinding* InSetBinding, FName InSetName)
+	{
+		if (InSetBinding == nullptr)
+		{
+			return;
+		}
+
+		if (FAttributeBindingDataPtr InBindingData = InSetBinding->GetAttributeBindingData())
+		{
+			if (FAttributeNamedSetPtr InNamedSet = InBindingData->FindNamedSet(InSetName))
+			{
+				SetBinding = InSetBinding;
+				Skeleton = InSetBinding->GetSkeleton();
+				BindingData = InBindingData;
+				NamedSet = InNamedSet;
+				DefaultAttributeValues = InNamedSet->GetDefaultAttributeValues();
+			}
+		}
+	}
+
+	const UAbstractSkeletonSetBinding* FAttributeEvaluationContext::GetSetBinding() const
+	{
+		return SetBinding;
+	}
+
+	const USkeleton* FAttributeEvaluationContext::GetSkeleton() const
+	{
+		return Skeleton;
+	}
+
+	const FAttributeNamedSetPtr& FAttributeEvaluationContext::GetNamedSet() const
+	{
+		return NamedSet;
+	}
+
+	const FAttributeCollectionPtr& FAttributeEvaluationContext::GetDefaultAttributeValues() const
+	{
+		return DefaultAttributeValues;
+	}
+
+	bool FAttributeEvaluationContext::IsValid() const
+	{
+		return SetBinding != nullptr;
+	}
+}
diff --git a/Engine/Plugins/Experimental/UAF/UAF/Source/UAF/Private/UAF/Attributes/AttributeRuntimeRegistry.cpp b/Engine/Plugins/Experimental/UAF/UAF/Source/UAF/Private/UAF/Attributes/AttributeRuntimeRegistry.cpp
new file mode 100644
--- /dev/null
+++ b/Engine/Plugins/Experimental/UAF/UAF/Source/UAF/Private/UAF/Attributes/AttributeRuntimeRegistry.cpp
@@ -0,0 +1,95 @@
+// Copyright Epic Games, Inc. All Rights Reserved.
+
+#include "UAF/Attributes/AttributeRuntimeRegistry.h"
+
+#include "Algo/BinarySearch.h"
+
+namespace UE::UAF
+{
+	namespace Private
+	{
+		static FAttributeRuntimeRegistry* GAttributeRuntimeRegistry = nullptr;
+	}
+
+	void FAttributeRuntimeRegistry::Init()
+	{
+		if (Private::GAttributeRuntimeRegistry == nullptr)
+		{
+			Private::GAttributeRuntimeRegistry = new FAttributeRuntimeRegistry();
+		}
+	}
+
+	void FAttributeRuntimeRegistry::Destroy()
+	{
+		if (Private::GAttributeRuntimeRegistry != nullptr)
+		{
+			delete Private::GAttributeRuntimeRegistry;
+			Private::GAttributeRuntimeRegistry = nullptr;
+		}
+	}
+
+	FAttributeRuntimeRegistry& FAttributeRuntimeRegistry::Get()
+	{
+		checkf(Private::GAttributeRuntimeRegistry != nullptr, TEXT("Attribute Runtime Registry is not instanced. It is only valid to access this while the UAF module is loaded."));
+		return *Private::GAttributeRuntimeRegistry;
+	}
+
+	FAttributeRuntimeRegistry::FAttributeRuntimeRegistry() = default;
+
+	FAttributeSetMap* FAttributeRuntimeRegistry::ConstructSetMap(const FAttributeSetMap::FConstructArgs& Args) const
+	{
+		FReadScopeLock Lock(SetMapInitializerLock);
+		if (const SetMapConstructorFun* ConstructorFun = SetMapInitializers.Find(Args.ValueType))
+		{
+			return (*ConstructorFun)(Args);
+		}
+
+		// Value type has not been registered
+		return nullptr;
+	}
+
+	FAttributeValueTransformerMapPtr FAttributeRuntimeRegistry::GetTransformerMap() const
+	{
+		FReadScopeLock Lock(TransformerMapLock);
+		return TransformerMap;
+	}
+
+	void FAttributeRuntimeRegistry::RegisterSetMapInitializer(UScriptStruct* ValueType, SetMapConstructorFun Initializer)
+	{
+		FWriteScopeLock Lock(SetMapInitializerLock);
+		SetMapInitializers.Add(ValueType, Initializer);
+	}
+
+	bool FAttributeRuntimeRegistry::RegisterAttributeValueTransformer(FName TransformerName, UScriptStruct* ValueType, FRawTransformerFunc TransformerFunc)
+	{
+		FWriteScopeLock Lock(TransformerMapLock);
+
+		// If someone is already holding a reference to our named operators, we cannot mutate
+		// When this occurs, we make a copy and swap it, keeping the old one alive until it is no longer needed
+		FAttributeValueTransformerMapPtr MutableTransformerMap = TransformerMap;
+		if (!MutableTransformerMap || MutableTransformerMap.GetSharedReferenceCount() != 2)
+		{
+			MutableTransformerMap = MakeShared<typename FAttributeValueTransformerMapPtr::ElementType>();
+		}
+
+		const uint32 TransformerNameHash = GetTypeHash(TransformerName);
+		FAttributeValueTransformerList* TransformerList = MutableTransformerMap->FindByHash(TransformerNameHash, TransformerName);
+		if (TransformerList == nullptr)
+		{
+			TransformerList = &MutableTransformerMap->AddByHash(TransformerNameHash, TransformerName, FAttributeValueTransformerList(TransformerName));
+		}
+
+		if (!TransformerList->AddTransformer(ValueType, TransformerFunc))
+		{
+			// We failed to add this transformer
+			return false;
+		}
+
+		if (TransformerMap != MutableTransformerMap)
+		{
+			TransformerMap = MutableTransformerMap;
+		}
+
+		return true;
+	}
+}
diff --git a/Engine/Plugins/Experimental/UAF/UAF/Source/UAF/Private/UAF/Attributes/AttributeSetMap.cpp b/Engine/Plugins/Experimental/UAF/UAF/Source/UAF/Private/UAF/Attributes/AttributeSetMap.cpp
new file mode 100644
--- /dev/null
+++ b/Engine/Plugins/Experimental/UAF/UAF/Source/UAF/Private/UAF/Attributes/AttributeSetMap.cpp
@@ -0,0 +1,60 @@
+// Copyright Epic Games, Inc. All Rights Reserved.
+
+#include "UAF/Attributes/AttributeSetMap.h"
+
+namespace UE::UAF
+{
+	FAttributeSetMap::FAttributeSetMap() = default;
+
+	FAttributeSetMap::FAttributeSetMap(const FConstructArgs& Args, int32 InAllocatedSize)
+		: TypedSet(Args.TypedSet)
+		, ValueTypeStruct(Args.ValueType)
+		, AllocatedSize(InAllocatedSize)
+		, NumValues(0)
+		, ValuesPtr(nullptr)
+		, ReallocFun(Args.ReallocFun)
+	{
+		check(Args.TypedSet.IsValid());
+		check(Args.ReallocFun != nullptr);
+	}
+
+	FAttributeSetMap::~FAttributeSetMap()
+	{
+		checkf(ValuesPtr == nullptr, TEXT("Derived types should handle destruction"));
+	}
+
+	bool FAttributeSetMap::IsEmpty() const
+	{
+		return ValuesPtr == nullptr;
+	}
+
+	int32 FAttributeSetMap::Num() const
+	{
+		return NumValues;
+	}
+
+	const FAttributeTypedSetPtr& FAttributeSetMap::GetTypedSet() const
+	{
+		return TypedSet;
+	}
+
+	UScriptStruct* FAttributeSetMap::GetAttributeType() const
+	{
+		return TypedSet ? TypedSet->GetType() : nullptr;
+	}
+
+	UScriptStruct* FAttributeSetMap::GetValueType() const
+	{
+		return ValueTypeStruct;
+	}
+
+	FAttributeMappingKey FAttributeSetMap::GetMappingKey() const
+	{
+		return TypedSet ? FAttributeMappingKey(TypedSet->GetType(), ValueTypeStruct) : FAttributeMappingKey();
+	}
+
+	int32 FAttributeSetMap::GetAllocatedSize() const
+	{
+		return AllocatedSize;
+	}
+}
diff --git a/Engine/Plugins/Experimental/UAF/UAF/Source/UAF/Private/UAF/Attributes/Transformers/Accumulate.cpp b/Engine/Plugins/Experimental/UAF/UAF/Source/UAF/Private/UAF/Attributes/Transformers/Accumulate.cpp
new file mode 100644
--- /dev/null
+++ b/Engine/Plugins/Experimental/UAF/UAF/Source/UAF/Private/UAF/Attributes/Transformers/Accumulate.cpp
@@ -0,0 +1,98 @@
+// Copyright Epic Games, Inc. All Rights Reserved.
+
+#include "UAF/Attributes/Transformers/Accumulate.h"
+
+#include "UAF/Attributes/AttributeCollection.h"
+#include "UAF/Attributes/AttributeSetMap.h"
+#include "UAF/Attributes/IteratorUtils.h"
+
+namespace UE::UAF::Transformers
+{
+	FName FAccumulate::TransformerName = TEXT("UAF::Transformers::Accumulate");
+
+	void FAccumulate::Apply(const FAttributeValueTransformerMapPtr& TransformerMap, FAttributeCollection& CollectionA, const FAttributeCollection& CollectionB, float Alpha)
+	{
+		// Named sets must match as it ensures that our inputs have the same sizes/shapes
+		check(CollectionA.GetNamedSet() == CollectionB.GetNamedSet());
+
+		if (const FAttributeValueTransformerList* TransformerList = TransformerMap->Find(TransformerName))
+		{
+			// Keep a list of set maps that don't have a matching pair
+			// These are duplicated afterwards to avoid mutating the collection while iterating
+			TArray<const FAttributeSetMap*, TInlineAllocator<8>> UnmatchedSetMapsB;
+
+			OuterJoinBy(
+				FAttributeValueTransformerListWithAttributeCollectionJoinOp(),
+				[Alpha, &UnmatchedSetMapsB](FRawTransformerFunc Transformer, FAttributeSetMap* SetMapA, const FAttributeSetMap* SetMapB)
+				{
+					if (SetMapA == nullptr)
+					{
+						if (SetMapB == nullptr)
+						{
+							// No maps found for the transformer type
+						}
+						else
+						{
+							// This set map doesn't have a matching pair
+							UnmatchedSetMapsB.Add(SetMapB);
+						}
+					}
+					else
+					{
+						if (SetMapB == nullptr)
+						{
+							// This set map doesn't have a matching pair, keep it
+						}
+						else
+						{
+							if (Transformer == nullptr)
+							{
+								// No transformer was found for this value type
+							}
+							else
+							{
+								// Transform our matching pair
+								FTransformFunc TypedTransformer = reinterpret_cast<FTransformFunc>(Transformer);
+								(*TypedTransformer)(SetMapA, SetMapB, Alpha);
+							}
+						}
+					}
+				},
+				TransformerList->CreateTransformerIterator(),
+				CollectionA.CreateSetMapIterator(),
+				CollectionB.CreateConstSetMapIterator());
+
+			// Duplicate any unmatched entries
+			for (const FAttributeSetMap* SetMapB : UnmatchedSetMapsB)
+			{
+				const FAttributeMappingKey MappingKey = SetMapB->GetMappingKey();
+				FAttributeSetMap* SetMapA = CollectionA.AddSetMap(MappingKey);
+				checkf(SetMapA != nullptr, TEXT("Failed to add set mapping"));
+
+				SetMapB->CopyTo(*SetMapA);
+			}
+		}
+	}
+
+	void FAccumulate_BoneTransformAttribute::Transform(UE::UAF::FAttributeSetMap* SetMapA, const UE::UAF::FAttributeSetMap* SetMapB, float Alpha)
+	{
+		check(SetMapA->GetTypedSet() == SetMapB->GetTypedSet());
+
+		TAttributeSetMap<FBoneTransformAnimationAttribute>* BoneTransformsA = static_cast<TAttributeSetMap<FBoneTransformAnimationAttribute>*>(SetMapA);
+		const TAttributeSetMap<FBoneTransformAnimationAttribute>* BoneTransformsB = static_cast<const TAttributeSetMap<FBoneTransformAnimationAttribute>*>(SetMapB);
+
+		FBoneTransformAnimationAttribute* TransformPtrA = BoneTransformsA->GetData();
+		const FBoneTransformAnimationAttribute* TransformPtrB = BoneTransformsB->GetData();
+		const FBoneTransformAnimationAttribute* EndTransformPtrA = TransformPtrA + BoneTransformsA->Num();
+
+		const ScalarRegister VAlpha(Alpha);
+
+		while (TransformPtrA < EndTransformPtrA)
+		{
+			TransformPtrA->Value.AccumulateWithShortestRotation(TransformPtrB->Value, VAlpha);
+
+			TransformPtrA++;
+			TransformPtrB++;
+		}
+	}
+}
diff --git a/Engine/Plugins/Experimental/UAF/UAF/Source/UAF/Private/UAF/Attributes/Transformers/Interpolate.cpp b/Engine/Plugins/Experimental/UAF/UAF/Source/UAF/Private/UAF/Attributes/Transformers/Interpolate.cpp
new file mode 100644
--- /dev/null
+++ b/Engine/Plugins/Experimental/UAF/UAF/Source/UAF/Private/UAF/Attributes/Transformers/Interpolate.cpp
@@ -0,0 +1,96 @@
+// Copyright Epic Games, Inc. All Rights Reserved.
+
+#include "UAF/Attributes/Transformers/Interpolate.h"
+
+#include "UAF/Attributes/AttributeCollection.h"
+#include "UAF/Attributes/AttributeSetMap.h"
+#include "UAF/Attributes/IteratorUtils.h"
+
+namespace UE::UAF::Transformers
+{
+	FName FInterpolate::TransformerName = TEXT("UAF::Transformers::Interpolate");
+
+	void FInterpolate::Apply(const FAttributeValueTransformerMapPtr& TransformerMap, FAttributeCollection& CollectionA, const FAttributeCollection& CollectionB, float Alpha)
+	{
+		// Named sets must match as it ensures that our inputs have the same sizes/shapes
+		check(CollectionA.GetNamedSet() == CollectionB.GetNamedSet());
+
+		if (const FAttributeValueTransformerList* TransformerList = TransformerMap->Find(TransformerName))
+		{
+			// Keep a list of set maps that don't have a matching pair
+			// These are duplicated afterwards to avoid mutating the collection while iterating
+			TArray<const FAttributeSetMap*, TInlineAllocator<8>> UnmatchedSetMapsB;
+
+			OuterJoinBy(
+				FAttributeValueTransformerListWithAttributeCollectionJoinOp(),
+				[Alpha, &UnmatchedSetMapsB](FRawTransformerFunc Transformer, FAttributeSetMap* SetMapA, const FAttributeSetMap* SetMapB)
+				{
+					if (SetMapA == nullptr)
+					{
+						if (SetMapB == nullptr)
+						{
+							// No maps found for the transformer type
+						}
+						else
+						{
+							// This set map doesn't have a matching pair
+							UnmatchedSetMapsB.Add(SetMapB);
+						}
+					}
+					else
+					{
+						if (SetMapB == nullptr)
+						{
+							// This set map doesn't have a matching pair, keep it
+						}
+						else
+						{
+							if (Transformer == nullptr)
+							{
+								// No transformer was found for this value type
+							}
+							else
+							{
+								// Transform our matching pair
+								FTransformFunc TypedTransformer = reinterpret_cast<FTransformFunc>(Transformer);
+								(*TypedTransformer)(SetMapA, SetMapB, Alpha);
+							}
+						}
+					}
+				},
+				TransformerList->CreateTransformerIterator(),
+				CollectionA.CreateSetMapIterator(),
+				CollectionB.CreateConstSetMapIterator());
+
+			// Duplicate any unmatched entries
+			for (const FAttributeSetMap* SetMapB : UnmatchedSetMapsB)
+			{
+				const FAttributeMappingKey MappingKey = SetMapB->GetMappingKey();
+				FAttributeSetMap* SetMapA = CollectionA.AddSetMap(MappingKey);
+				checkf(SetMapA != nullptr, TEXT("Failed to add set mapping"));
+
+				SetMapB->CopyTo(*SetMapA);
+			}
+		}
+	}
+
+	void FInterpolate_BoneTransformAttribute::Transform(UE::UAF::FAttributeSetMap* SetMapA, const UE::UAF::FAttributeSetMap* SetMapB, float Alpha)
+	{
+		check(SetMapA->GetTypedSet() == SetMapB->GetTypedSet());
+
+		TAttributeSetMap<FBoneTransformAnimationAttribute>* BoneTransformsA = static_cast<TAttributeSetMap<FBoneTransformAnimationAttribute>*>(SetMapA);
+		const TAttributeSetMap<FBoneTransformAnimationAttribute>* BoneTransformsB = static_cast<const TAttributeSetMap<FBoneTransformAnimationAttribute>*>(SetMapB);
+
+		FBoneTransformAnimationAttribute* TransformPtrA = BoneTransformsA->GetData();
+		const FBoneTransformAnimationAttribute* TransformPtrB = BoneTransformsB->GetData();
+		const FBoneTransformAnimationAttribute* EndTransformPtrA = TransformPtrA + BoneTransformsA->Num();
+
+		while (TransformPtrA < EndTransformPtrA)
+		{
+			TransformPtrA->Value.Blend(TransformPtrA->Value, TransformPtrB->Value, Alpha);
+
+			TransformPtrA++;
+			TransformPtrB++;
+		}
+	}
+}
diff --git a/Engine/Plugins/Experimental/UAF/UAF/Source/UAF/Private/UAF/Attributes/Transformers/Overwrite.cpp b/Engine/Plugins/Experimental/UAF/UAF/Source/UAF/Private/UAF/Attributes/Transformers/Overwrite.cpp
new file mode 100644
--- /dev/null
+++ b/Engine/Plugins/Experimental/UAF/UAF/Source/UAF/Private/UAF/Attributes/Transformers/Overwrite.cpp
@@ -0,0 +1,98 @@
+// Copyright Epic Games, Inc. All Rights Reserved.
+
+#include "UAF/Attributes/Transformers/Overwrite.h"
+
+#include "UAF/Attributes/AttributeCollection.h"
+#include "UAF/Attributes/AttributeSetMap.h"
+#include "UAF/Attributes/IteratorUtils.h"
+
+namespace UE::UAF::Transformers
+{
+	FName FOverwrite::TransformerName = TEXT("UAF::Transformers::Overwrite");
+
+	void FOverwrite::Apply(const FAttributeValueTransformerMapPtr& TransformerMap, FAttributeCollection& CollectionA, const FAttributeCollection& CollectionB, float Alpha)
+	{
+		// Named sets must match as it ensures that our inputs have the same sizes/shapes
+		check(CollectionA.GetNamedSet() == CollectionB.GetNamedSet());
+
+		if (const FAttributeValueTransformerList* TransformerList = TransformerMap->Find(TransformerName))
+		{
+			// Keep a list of set maps that don't have a matching pair
+			// These are duplicated afterwards to avoid mutating the collection while iterating
+			TArray<const FAttributeSetMap*, TInlineAllocator<8>> UnmatchedSetMapsB;
+
+			OuterJoinBy(
+				FAttributeValueTransformerListWithAttributeCollectionJoinOp(),
+				[Alpha, &UnmatchedSetMapsB](FRawTransformerFunc Transformer, FAttributeSetMap* SetMapA, const FAttributeSetMap* SetMapB)
+				{
+					if (SetMapA == nullptr)
+					{
+						if (SetMapB == nullptr)
+						{
+							// No maps found for the transformer type
+						}
+						else
+						{
+							// This set map doesn't have a matching pair
+							UnmatchedSetMapsB.Add(SetMapB);
+						}
+					}
+					else
+					{
+						if (SetMapB == nullptr)
+						{
+							// This set map doesn't have a matching pair, keep it
+						}
+						else
+						{
+							if (Transformer == nullptr)
+							{
+								// No transformer was found for this value type
+							}
+							else
+							{
+								// Transform our matching pair
+								FTransformFunc TypedTransformer = reinterpret_cast<FTransformFunc>(Transformer);
+								(*TypedTransformer)(SetMapA, SetMapB, Alpha);
+							}
+						}
+					}
+				},
+				TransformerList->CreateTransformerIterator(),
+				CollectionA.CreateSetMapIterator(),
+				CollectionB.CreateConstSetMapIterator());
+
+			// Duplicate any unmatched entries
+			for (const FAttributeSetMap* SetMapB : UnmatchedSetMapsB)
+			{
+				const FAttributeMappingKey MappingKey = SetMapB->GetMappingKey();
+				FAttributeSetMap* SetMapA = CollectionA.AddSetMap(MappingKey);
+				checkf(SetMapA != nullptr, TEXT("Failed to add set mapping"));
+
+				SetMapB->CopyTo(*SetMapA);
+			}
+		}
+	}
+
+	void FOverwrite_BoneTransformAttribute::Transform(UE::UAF::FAttributeSetMap* SetMapA, const UE::UAF::FAttributeSetMap* SetMapB, float Alpha)
+	{
+		check(SetMapA->GetTypedSet() == SetMapB->GetTypedSet());
+
+		TAttributeSetMap<FBoneTransformAnimationAttribute>* BoneTransformsA = static_cast<TAttributeSetMap<FBoneTransformAnimationAttribute>*>(SetMapA);
+		const TAttributeSetMap<FBoneTransformAnimationAttribute>* BoneTransformsB = static_cast<const TAttributeSetMap<FBoneTransformAnimationAttribute>*>(SetMapB);
+
+		FBoneTransformAnimationAttribute* TransformPtrA = BoneTransformsA->GetData();
+		const FBoneTransformAnimationAttribute* TransformPtrB = BoneTransformsB->GetData();
+		const FBoneTransformAnimationAttribute* EndTransformPtrA = TransformPtrA + BoneTransformsA->Num();
+
+		const ScalarRegister VAlpha(Alpha);
+
+		while (TransformPtrA < EndTransformPtrA)
+		{
+			TransformPtrA->Value = TransformPtrB->Value * VAlpha;
+
+			TransformPtrA++;
+			TransformPtrB++;
+		}
+	}
+}
diff --git a/Engine/Plugins/Experimental/UAF/UAF/Source/UAF/Private/UAF/Attributes/Transformers/Sanitize.cpp b/Engine/Plugins/Experimental/UAF/UAF/Source/UAF/Private/UAF/Attributes/Transformers/Sanitize.cpp
new file mode 100644
--- /dev/null
+++ b/Engine/Plugins/Experimental/UAF/UAF/Source/UAF/Private/UAF/Attributes/Transformers/Sanitize.cpp
@@ -0,0 +1,77 @@
+// Copyright Epic Games, Inc. All Rights Reserved.
+
+#include "UAF/Attributes/Transformers/Sanitize.h"
+
+#include "UAF/Attributes/AttributeCollection.h"
+#include "UAF/Attributes/AttributeSetMap.h"
+#include "UAF/Attributes/IteratorUtils.h"
+
+// We validate bone transforms in development and debug builds
+#define UE_VALIDATE_BONE_TRANSFORMS !(UE_BUILD_SHIPPING || UE_BUILD_TEST)
+
+namespace UE::UAF::Transformers
+{
+	FName FSanitize::TransformerName = TEXT("UAF::Transformers::Sanitize");
+
+	void FSanitize::Apply(const FAttributeValueTransformerMapPtr& TransformerMap, FAttributeCollection& Collection)
+	{
+		if (const FAttributeValueTransformerList* TransformerList = TransformerMap->Find(TransformerName))
+		{
+			InnerJoinBy(
+				FAttributeValueTransformerListWithAttributeCollectionJoinOp(),
+				[](FRawTransformerFunc Transformer, FAttributeSetMap* SetMap)
+				{
+					FTransformFunc TypedTransformer = reinterpret_cast<FTransformFunc>(Transformer);
+					(*TypedTransformer)(SetMap);
+				},
+				TransformerList->CreateTransformerIterator(),
+				Collection.CreateSetMapIterator());
+		}
+	}
+
+#if UE_VALIDATE_BONE_TRANSFORMS
+	// Don't inline this function to keep the stack usage down
+	static FORCENOINLINE void IsBoneTransformSetMapValid(const TAttributeSetMap<FBoneTransformAnimationAttribute>* BoneTransforms)
+	{
+		const FBoneTransformAnimationAttribute* TransformPtr = BoneTransforms->GetData();
+		const FBoneTransformAnimationAttribute* EndTransformPtr = TransformPtr + BoneTransforms->Num();
+
+		while (TransformPtr < EndTransformPtr)
+		{
+			if (TransformPtr->Value.ContainsNaN())
+			{
+				ensureMsgf(false, TEXT("Bone (%s) contains NaN with Value:[%s]"),
+					*BoneTransforms->GetTypedSet()->GetName(FAttributeSetIndex(TransformPtr - BoneTransforms->GetData())).ToString(),
+					*TransformPtr->Value.ToString());
+			}
+			else if (!TransformPtr->Value.IsRotationNormalized())
+			{
+				ensureMsgf(false, TEXT("Bone (%s) Rotation not normalized with Rotation:[%s]"),
+					*BoneTransforms->GetTypedSet()->GetName(FAttributeSetIndex(TransformPtr - BoneTransforms->GetData())).ToString(),
+					*TransformPtr->Value.GetRotation().ToString());
+			}
+
+			TransformPtr++;
+		}
+	}
+#endif
+
+	void FSanitize_BoneTransformAttribute::Transform(FAttributeSetMap* SetMap)
+	{
+		TAttributeSetMap<FBoneTransformAnimationAttribute>* BoneTransforms = static_cast<TAttributeSetMap<FBoneTransformAnimationAttribute>*>(SetMap);
+
+		FBoneTransformAnimationAttribute* TransformPtr = BoneTransforms->GetData();
+		const FBoneTransformAnimationAttribute* EndTransformPtr = TransformPtr + BoneTransforms->Num();
+
+		while (TransformPtr < EndTransformPtr)
+		{
+			TransformPtr->Value.NormalizeRotation();
+
+			TransformPtr++;
+		}
+
+#if UE_VALIDATE_BONE_TRANSFORMS
+		IsBoneTransformSetMapValid(BoneTransforms);
+#endif
+	}
+}
diff --git a/Engine/Plugins/Experimental/UAF/UAF/Source/UAF/Private/UAF/Attributes/Transformers/Transformer.cpp b/Engine/Plugins/Experimental/UAF/UAF/Source/UAF/Private/UAF/Attributes/Transformers/Transformer.cpp
new file mode 100644
--- /dev/null
+++ b/Engine/Plugins/Experimental/UAF/UAF/Source/UAF/Private/UAF/Attributes/Transformers/Transformer.cpp
@@ -0,0 +1,28 @@
+// Copyright Epic Games, Inc. All Rights Reserved.
+
+#include "UAF/Attributes/Transformers/Transformer.h"
+
+namespace UE::UAF
+{
+	FAttributeValueTransformerList::FAttributeValueTransformerList() = default;
+
+	FAttributeValueTransformerList::FAttributeValueTransformerList(FName InTransformerName)
+		: TransformerName(InTransformerName)
+	{
+	}
+
+	bool FAttributeValueTransformerList::AddTransformer(UScriptStruct* ValueType, FRawTransformerFunc TransformerFunc)
+	{
+		const int32 EntryIndex = Algo::LowerBound(ValueTypes, ValueType);
+		if (EntryIndex < ValueTypes.Num() && ValueTypes[EntryIndex] == ValueType)
+		{
+			// A transformer for this value type has already been registered
+			return false;
+		}
+
+		ValueTypes.Insert(ValueType, EntryIndex);
+		Transformers.Insert(TransformerFunc, EntryIndex);
+
+		return true;
+	}
+}
diff --git a/Engine/Plugins/Experimental/UAF/UAF/Source/UAF/Public/DecompressionTools.h b/Engine/Plugins/Experimental/UAF/UAF/Source/UAF/Public/DecompressionTools.h
--- a/Engine/Plugins/Experimental/UAF/UAF/Source/UAF/Public/DecompressionTools.h
+++ b/Engine/Plugins/Experimental/UAF/UAF/Source/UAF/Public/DecompressionTools.h
@@ -19,16 +19,19 @@ class UAnimSequence;
 
 namespace UE::UAF
 {
+	class FAttributeCollection;
+	class FAttributeEvaluationContext;
 
 class FDecompressionTools
 {
 public:
 	// Returns whether decompression should use raw data or not
-	static UE_API bool ShouldUseRawData(const UAnimSequence* AnimSequence, const FLODPose& AnimationPoseData);
+	static UE_API bool ShouldUseRawData(const UAnimSequence* AnimSequence, const FAnimExtractContext& ExtractionContext);
 
 	// Extracted from UAnimSequence
 	// Extracts Animation Data from the provided AnimSequence, using AnimExtractContext extraction parameters
 	static UE_API void GetAnimationPose(const UAnimSequence* AnimSequence, const FAnimExtractContext& ExtractionContext, FLODPose& OutAnimationPoseData, bool bForceUseRawData = false);
+	static UE_API void GetAnimationPose(const UAnimSequence* AnimSequence, const FAnimExtractContext& ExtractionContext, const FAttributeEvaluationContext& VMContext, FAttributeCollection& OutCollection, bool bForceUseRawData = false);
 
 	/**
 	* Get Bone Transform of the Time given, relative to Parent for all RequiredBones
@@ -40,7 +43,10 @@ public:
 	static UE_API void GetBonePose_AdditiveMeshRotationOnly(const UAnimSequence* AnimSequence, const FAnimExtractContext& ExtractionContext, FLODPose& OutAnimationPoseData);
 
 	static UE_API void GetAnimationCurves(const UAnimSequence* AnimSequence, const FAnimExtractContext& ExtractionContext, FBlendedCurve& OutCurves, bool bForceUseRawData = false);
+	static UE_API void GetAnimationCurves(const UAnimSequence* AnimSequence, const FAnimExtractContext& ExtractionContext, const FAttributeEvaluationContext& VMContext, FAttributeCollection& OutCollection, bool bForceUseRawData = false);
+
 	static UE_API void GetAnimationAttributes(const UAnimSequence* AnimSequence, const FAnimExtractContext& ExtractionContext, const FReferencePose& RefPose, UE::Anim::FStackAttributeContainer& OutAttributes, bool bForceUseRawData = false);
+	static UE_API void GetAnimationAttributes(const UAnimSequence* AnimSequence, const FAnimExtractContext& ExtractionContext, const FAttributeEvaluationContext& VMContext, FAttributeCollection& OutCollection, bool bForceUseRawData = false);
 
 	// Decompress and retarget animation data using provided RetargetTransforms
 	static UE_API void DecompressPose(FLODPose& OutAnimationPoseData,
@@ -78,8 +84,6 @@ public:
 		, const int32 LODBoneIndex
 		, const bool bIsBakedAdditive
 		, const bool bDisableRetargeting);
-
-	static UE_API bool UseRawDataForPoseExtraction(const UAnimSequence* AnimSequence, const FLODPose& AnimationPoseData);
 };
 
 } // end namespace UE::UAF
diff --git a/Engine/Plugins/Experimental/UAF/UAF/Source/UAF/Public/UAF/AbstractSkeleton/AbstractSkeletonSetBinding.h b/Engine/Plugins/Experimental/UAF/UAF/Source/UAF/Public/UAF/AbstractSkeleton/AbstractSkeletonSetBinding.h
--- a/Engine/Plugins/Experimental/UAF/UAF/Source/UAF/Public/UAF/AbstractSkeleton/AbstractSkeletonSetBinding.h
+++ b/Engine/Plugins/Experimental/UAF/UAF/Source/UAF/Public/UAF/AbstractSkeleton/AbstractSkeletonSetBinding.h
@@ -4,6 +4,10 @@
 
 #include "CoreMinimal.h"
 #include "Animation/AnimData/AttributeIdentifier.h"
+#include "Animation/BuiltInAttributeTypes.h"
+#include "UAF/Attributes/AttributeBindingIndex.h"
+#include "UAF/Attributes/AttributeCollection.h"
+#include "UAF/Attributes/AttributeBindingData.h"
 
 #include "AbstractSkeletonSetBinding.generated.h"
 
@@ -43,11 +47,15 @@ class UAbstractSkeletonSetBinding : public UObject
 	GENERATED_BODY()
 
 public:
+	// UObject implementation
+	virtual void PostLoad() override;
+	virtual void GetPreloadDependencies(TArray<UObject*>& OutDeps) override;
+
 	UE_API bool SetSkeleton(const TObjectPtr<USkeleton> InSkeleton);
 
 	UE_API TObjectPtr<USkeleton> GetSkeleton() const;
 
-	UE_API bool SetSetCollection(const TObjectPtr<const UAbstractSkeletonSetCollection> InSetCollection);
+	UE_API bool SetSetCollection(const TObjectPtr<UAbstractSkeletonSetCollection> InSetCollection);
 	
 	UE_API TObjectPtr<const UAbstractSkeletonSetCollection> GetSetCollection() const;
 	
@@ -80,9 +88,19 @@ public:
 
 	UE_API const TConstArrayView<FAbstractSkeleton_AttributeBinding> GetAttributeBindings() const;
 
+	// Returns the baked binding data
+	UE_API const UE::UAF::FAttributeBindingDataPtr& GetAttributeBindingData() const;
+
 private:
+#if WITH_EDITOR
+	void OnSkeletonChanged();
+	void OnSetCollectionChanged();
+#endif
+
+	void RebuildBindingData();
+
 	UPROPERTY()
-	TObjectPtr<const UAbstractSkeletonSetCollection> SetCollection;
+	TObjectPtr<UAbstractSkeletonSetCollection> SetCollection;
 	
 	UPROPERTY(AssetRegistrySearchable)
 	TObjectPtr<USkeleton> Skeleton;
@@ -92,6 +110,16 @@ private:
 
 	UPROPERTY()
 	TArray<FAbstractSkeleton_AttributeBinding> AttributeBindings;
+
+	// TODO: Cache sets in PostLoad for cooked builds
+
+	// Attribute binding data, computed on demand and cached in editor, built in PostLoad in cooked builds
+	// This is a baked and immutable representation of the binding
+	UE::UAF::FAttributeBindingDataPtr AttributeBindingData;
+
+#if WITH_EDITOR
+	FDelegateHandle SetCollectionChangedHandle;
+#endif
 };
 
 #undef UE_API
