Introducing new abstract skeleton runtime for UAF (disabled by default)

diff --git a/Engine/Plugins/Experimental/UAF/UAF/Source/UAF/Public/UAF/Attributes/AttributeCollection.h b/Engine/Plugins/Experimental/UAF/UAF/Source/UAF/Public/UAF/Attributes/AttributeCollection.h
new file mode 100644
--- /dev/null
+++ b/Engine/Plugins/Experimental/UAF/UAF/Source/UAF/Public/UAF/Attributes/AttributeCollection.h
@@ -0,0 +1,384 @@
+// Copyright Epic Games, Inc. All Rights Reserved.
+
+#pragma once
+
+#include "CoreMinimal.h"
+#include "Misc/MemStack.h"
+#include "UAF/Attributes/AttributeNamedSet.h"
+#include "UAF/Attributes/AttributeMappingKey.h"
+#include "UAF/Attributes/AttributeSetMap.h"
+#include "UAF/Attributes/IndirectAllocator.h"
+
+#include "AttributeCollection.generated.h"
+
+#define UE_API UAF_API
+
+namespace UE::UAF
+{
+	// A iterator that iterates over all set maps within an attribute collection
+	template<class ContainerType>
+	class TAttributeCollectionSetMapIterator
+	{
+	public:
+		static constexpr bool bIsContainerConst = std::is_const_v<ContainerType>;
+
+		using EntryPtrType = typename std::conditional<bIsContainerConst, const typename ContainerType::FSetMapEntry*, typename ContainerType::FSetMapEntry*>::type;
+		using SetMapType = typename std::conditional<bIsContainerConst, const FAttributeSetMap*, FAttributeSetMap*>::type;
+
+		// Creates an iterator over the set map entries within the specified collection
+		[[nodiscard]] TAttributeCollectionSetMapIterator(ContainerType& Collection);
+
+		// Increments and moves the iterator to the next set map entry
+		TAttributeCollectionSetMapIterator& operator++();
+
+		// Returns whether or not the iterator still contains values
+		[[nodiscard]] explicit operator bool() const;
+
+		// Returns the mapping key of the current set map entry
+		[[nodiscard]] const FAttributeMappingKey& GetKey() const;
+
+		// Returns the attribute type of the current set map entry
+		[[nodiscard]] UScriptStruct* GetAttributeType() const;
+
+		// Returns the value type of the current set map entry
+		[[nodiscard]] UScriptStruct* GetValueType() const;
+
+		// Returns the set map from the current entry
+		[[nodiscard]] SetMapType GetMap() const;
+
+	private:
+		// Current set map entry
+		EntryPtrType EntryPtr = nullptr;
+
+		// One past the last set map entry
+		EntryPtrType EntryEndPtr = nullptr;
+	};
+
+	/*
+	 * Attribute Collection
+	 *
+	 * This struct holds attributes produced by UAF. They come in two flavors: static (set map) and dynamic (map).
+	 * Static attributes are fast and defined up front in the skeleton set binding asset.
+	 * Dynamic attributes are slower but produced on demand during evaluation.
+	 */
+	class FAttributeCollection
+	{
+	private:
+		struct FSetMapEntry
+		{
+			// The mapping key that identifies this entry
+			FAttributeMappingKey Key;
+
+			// The set map containing our values
+			FAttributeSetMap* Map = nullptr;
+		};
+
+	public:
+		UE_API FAttributeCollection(const FAttributeCollection& Other);
+		UE_API FAttributeCollection(FAttributeCollection&& Other);
+		UE_API ~FAttributeCollection();
+
+		UE_API FAttributeCollection& operator=(const FAttributeCollection& Other);
+		UE_API FAttributeCollection& operator=(FAttributeCollection&& Other);
+
+		// Returns whether or not we are empty
+		[[nodiscard]] UE_API bool IsEmpty() const;
+
+		// Returns the named set this collection maps to
+		[[nodiscard]] UE_API const FAttributeNamedSetPtr& GetNamedSet() const;
+
+		// Returns the number of set maps contained within
+		[[nodiscard]] UE_API int32 NumSetMaps() const;
+
+		// Constructs and adds a new mapping of ValueType to a set of AttributeType to this collection.
+		// Returns nullptr if we fail to add an entry (e.g. duplicate mapping)
+		template <class ValueType, typename... TArgs> [[nodiscard]] TAttributeSetMap<ValueType>* AddSetMap(const FAttributeMappingKey& MappingKey, TArgs... Args);
+
+		// Constructs and adds a new set mapping to this collection.
+		// Returns nullptr if we fail to add an entry (e.g. duplicate mapping)
+		[[nodiscard]] UE_API FAttributeSetMap* AddSetMap(const FAttributeMappingKey& MappingKey);
+
+		// Removes the specified value set mapping from this collection
+		// Returns whether or not the entry was successfully removed
+		UE_API bool RemoveSetMap(FAttributeSetMap* Map);
+
+		// Removes the specified set mapping from this collection
+		// Returns whether or not the entry was successfully removed
+		UE_API bool RemoveSetMap(const FAttributeMappingKey& MappingKey);
+
+		// Finds the mapping of ValueType to a set of AttributeType in this collection
+		// Returns nullptr if this collection does not contain such a mapping
+		template <class ValueType> [[nodiscard]] TAttributeSetMap<ValueType>* FindSetMap(const FAttributeMappingKey& MappingKey);
+
+		// Finds the mapping of ValueType to a set of AttributeType in this collection
+		// Returns nullptr if this collection does not contain such a mapping
+		template <class ValueType> [[nodiscard]] const TAttributeSetMap<ValueType>* FindSetMap(const FAttributeMappingKey& MappingKey) const;
+
+		// Finds the mapping of ValueType to a set of AttributeType in this collection
+		// Returns nullptr if this collection does not contain such a mapping
+		[[nodiscard]] UE_API FAttributeSetMap* FindSetMap(const FAttributeMappingKey& MappingKey);
+
+		// Finds the mapping of ValueType to a set of AttributeType in this collection
+		// Returns nullptr if this collection does not contain such a mapping
+		[[nodiscard]] UE_API const FAttributeSetMap* FindSetMap(const FAttributeMappingKey& MappingKey) const;
+
+		// Finds or adds a mapping of ValueType to a set of AttributeType in this collection
+		// Returns nullptr if we fail to add an entry (e.g. duplicate mapping)
+		template <class ValueType, typename... TArgs> [[nodiscard]] TAttributeSetMap<ValueType>* FindOrAddSetMap(const FAttributeMappingKey& MappingKey, TArgs... Args);
+
+		// Finds or adds a mapping of ValueType to a set of AttributeType in this collection
+		// Returns nullptr if we fail to add an entry (e.g. duplicate mapping)
+		[[nodiscard]] UE_API FAttributeSetMap* FindOrAddSetMap(const FAttributeMappingKey& MappingKey);
+
+		// Iterators
+		using TSetMapIterator = TAttributeCollectionSetMapIterator<FAttributeCollection>;
+		using TConstSetMapIterator = TAttributeCollectionSetMapIterator<const FAttributeCollection>;
+
+		// Returns an iterator over all set maps contained within, sorted by their key (value type first, then attribute type)
+		[[nodiscard]] TSetMapIterator CreateSetMapIterator();
+		[[nodiscard]] TConstSetMapIterator CreateConstSetMapIterator() const;
+
+	protected:
+		UE_API FAttributeCollection(const FAttributeNamedSetPtr& NamedSet, FReallocFun InReallocFun);
+
+	private:
+		UE_API FAttributeSetMap** AddSetMapImpl(const FAttributeSetMap::FConstructArgs& Args, bool bConstructSetMap);
+
+		// The named set we are based on, every set mapping uses a typed set within it
+		FAttributeNamedSetPtr NamedSet;
+
+		// An array of set mappings (sorted by mapping key)
+		FSetMapEntry* SetMaps = nullptr;
+
+		// How many set mappings are contained within this collection
+		int32 SetMapCount = 0;
+
+		// How many set mappings we can contain within this collection
+		int32 SetMapCapacity = 0;
+
+		// The indirect allocator function pointer
+		FReallocFun ReallocFun = nullptr;
+
+		friend TSetMapIterator;
+		friend TConstSetMapIterator;
+	};
+
+	// A heap allocated attribute collection
+	class FAttributeCollectionHeap : public FAttributeCollection
+	{
+	public:
+		// Constructs an empty collection
+		FAttributeCollectionHeap();
+
+		// Constructs an empty collection based on the specified named set
+		explicit FAttributeCollectionHeap(const FAttributeNamedSetPtr& NamedSet);
+	};
+
+	// A Mem-Stack allocated attribute collection
+	class FAttributeCollectionStack : public FAttributeCollection
+	{
+	public:
+		// Constructs an empty collection
+		FAttributeCollectionStack();
+
+		// Constructs an empty collection based on the specified named set
+		explicit FAttributeCollectionStack(const FAttributeNamedSetPtr& NamedSet);
+	};
+
+	using FAttributeCollectionPtr = TSharedPtr<FAttributeCollection, ESPMode::ThreadSafe>;
+
+	// An operator structure to join attribute collections using InnerJoin/OuterJoin/etc
+	struct FAttributeCollectionSetMapJoinOp
+	{
+		const FAttributeMappingKey& GetLargestKey() const { return FAttributeMappingKey::LARGEST_VALUE; }
+
+		const FAttributeMappingKey& GetKey(const FAttributeCollection::TSetMapIterator& It) const { return It.GetKey(); }
+		const FAttributeMappingKey& GetKey(const FAttributeCollection::TConstSetMapIterator& It) const { return It.GetKey(); }
+
+		FAttributeSetMap* GetValue(const FAttributeCollection::TSetMapIterator& It) const { return It.GetMap(); }
+		const FAttributeSetMap* GetValue(const FAttributeCollection::TConstSetMapIterator& It) const { return It.GetMap(); }
+	};
+}
+
+/*
+ * Attribute Collection
+ *
+ * This struct holds attributes produced by UAF. They come in two flavors: static (set map) and dynamic (map).
+ * Static attributes are fast and defined up front in the skeleton set binding asset.
+ * Dynamic attributes are slower but produced on demand during evaluation.
+ */
+USTRUCT(BlueprintType, meta = (DisplayName = "Attribute Collection"))
+struct FUAFAttributeCollection
+#if CPP
+	: public UE::UAF::FAttributeCollectionHeap
+#endif
+{
+	GENERATED_BODY()
+
+	// Equivalent to FAnimNextGraphLODPose
+};
+
+//////////////////////////////////////////////////////////////////////////
+// Implementation
+
+namespace UE::UAF
+{
+	template <class ValueType, typename... TArgs>
+	inline TAttributeSetMap<ValueType>* FAttributeCollection::AddSetMap(const FAttributeMappingKey& MappingKey, TArgs... Args)
+	{
+		if (!NamedSet)
+		{
+			// We are not initialized
+			return nullptr;
+		}
+
+		UScriptStruct* ValueTypeStruct = ValueType::StaticStruct();
+		if (!ensureMsgf(ValueTypeStruct == MappingKey.GetValueType(), TEXT("Mismatch between the return value type and mapping key value type")))
+		{
+			return nullptr;
+		}
+
+		FAttributeTypedSetPtr TypedSet = NamedSet->FindTypedSet(MappingKey.GetAttributeType());
+		if (!TypedSet)
+		{
+			// Unknown attribute type, it isn't contained within our named set
+			return nullptr;
+		}
+
+		// Forward the extra arguments
+		typename TAttributeSetMap<ValueType>::FConstructArgs ConstructArgs(TypedSet, ValueTypeStruct, ReallocFun, Args...);
+
+		// Don't construct the set map internally, it'll have to use indirection to do so and here we know the concrete type
+		constexpr bool bConstructSetMap = false;
+
+		FAttributeSetMap** SetMapPtr = AddSetMapImpl(ConstructArgs, bConstructSetMap);
+		if (SetMapPtr == nullptr)
+		{
+			return nullptr;
+		}
+
+		TAttributeSetMap<ValueType>* SetMap = MakeAttributeSetMap<ValueType>(ConstructArgs);
+
+		check(*SetMapPtr == nullptr);
+		*SetMapPtr = SetMap;
+
+		return SetMap;
+	}
+
+	template <class ValueType>
+	inline TAttributeSetMap<ValueType>* FAttributeCollection::FindSetMap(const FAttributeMappingKey& MappingKey)
+	{
+		UScriptStruct* ValueTypeStruct = ValueType::StaticStruct();
+		if (!ensureMsgf(ValueTypeStruct == MappingKey.GetValueType(), TEXT("Mismatch between the return value type and mapping key value type")))
+		{
+			return nullptr;
+		}
+
+		return static_cast<TAttributeSetMap<ValueType>*>(FindSetMap(MappingKey));
+	}
+
+	template <class ValueType>
+	inline const TAttributeSetMap<ValueType>* FAttributeCollection::FindSetMap(const FAttributeMappingKey& MappingKey) const
+	{
+		UScriptStruct* ValueTypeStruct = ValueType::StaticStruct();
+		if (!ensureMsgf(ValueTypeStruct == MappingKey.GetValueType(), TEXT("Mismatch between the return value type and mapping key value type")))
+		{
+			return nullptr;
+		}
+
+		return static_cast<const TAttributeSetMap<ValueType>*>(FindSetMap(MappingKey));
+	}
+
+	template <class ValueType, typename... TArgs>
+	inline TAttributeSetMap<ValueType>* FAttributeCollection::FindOrAddSetMap(const FAttributeMappingKey& MappingKey, TArgs... Args)
+	{
+		UScriptStruct* ValueTypeStruct = ValueType::StaticStruct();
+		if (!ensureMsgf(ValueTypeStruct == MappingKey.GetValueType(), TEXT("Mismatch between the return value type and mapping key value type")))
+		{
+			return nullptr;
+		}
+
+		if (FAttributeSetMap* SetMap = FindSetMap(MappingKey))
+		{
+			return static_cast<TAttributeSetMap<ValueType>*>(SetMap);
+		}
+
+		return AddSetMap<ValueType>(MappingKey, Args...);
+	}
+
+	inline typename FAttributeCollection::TSetMapIterator FAttributeCollection::CreateSetMapIterator()
+	{
+		return TSetMapIterator(*this);
+	}
+
+	inline typename FAttributeCollection::TConstSetMapIterator FAttributeCollection::CreateConstSetMapIterator() const
+	{
+		return TConstSetMapIterator(*this);
+	}
+
+	inline FAttributeCollectionHeap::FAttributeCollectionHeap()
+		: FAttributeCollection(nullptr, &FAllocatorTypeTrait<FDefaultAllocator>::Realloc)
+	{
+	}
+
+	inline FAttributeCollectionHeap::FAttributeCollectionHeap(const FAttributeNamedSetPtr& NamedSet)
+		: FAttributeCollection(NamedSet, &FAllocatorTypeTrait<FDefaultAllocator>::Realloc)
+	{
+	}
+
+	inline FAttributeCollectionStack::FAttributeCollectionStack()
+		: FAttributeCollection(nullptr, &FAllocatorTypeTrait<TMemStackAllocator<>>::Realloc)
+	{
+	}
+
+	inline FAttributeCollectionStack::FAttributeCollectionStack(const FAttributeNamedSetPtr& NamedSet)
+		: FAttributeCollection(NamedSet, &FAllocatorTypeTrait<TMemStackAllocator<>>::Realloc)
+	{
+	}
+
+	template<class ContainerType>
+	inline TAttributeCollectionSetMapIterator<ContainerType>::TAttributeCollectionSetMapIterator(ContainerType& Collection)
+		: EntryPtr(Collection.SetMaps)
+		, EntryEndPtr(Collection.SetMaps + Collection.SetMapCount)
+	{
+	}
+
+	template<class ContainerType>
+	inline TAttributeCollectionSetMapIterator<ContainerType>& TAttributeCollectionSetMapIterator<ContainerType>::operator++()
+	{
+		++EntryPtr;
+		return *this;
+	}
+
+	template<class ContainerType>
+	inline TAttributeCollectionSetMapIterator<ContainerType>::operator bool() const
+	{
+		return EntryPtr < EntryEndPtr;
+	}
+
+	template<class ContainerType>
+	inline const FAttributeMappingKey& TAttributeCollectionSetMapIterator<ContainerType>::GetKey() const
+	{
+		return EntryPtr->Key;
+	}
+
+	template<class ContainerType>
+	inline UScriptStruct* TAttributeCollectionSetMapIterator<ContainerType>::GetAttributeType() const
+	{
+		return EntryPtr->Key.GetAttributeType();
+	}
+
+	template<class ContainerType>
+	inline UScriptStruct* TAttributeCollectionSetMapIterator<ContainerType>::GetValueType() const
+	{
+		return EntryPtr->Key.GetValueType();
+	}
+
+	template<class ContainerType>
+	inline typename TAttributeCollectionSetMapIterator<ContainerType>::SetMapType TAttributeCollectionSetMapIterator<ContainerType>::GetMap() const
+	{
+		return EntryPtr->Map;
+	}
+}
+
+#undef UE_API
diff --git a/Engine/Plugins/Experimental/UAF/UAF/Source/UAF/Public/UAF/Attributes/AttributeEvaluationContext.h b/Engine/Plugins/Experimental/UAF/UAF/Source/UAF/Public/UAF/Attributes/AttributeEvaluationContext.h
new file mode 100644
--- /dev/null
+++ b/Engine/Plugins/Experimental/UAF/UAF/Source/UAF/Public/UAF/Attributes/AttributeEvaluationContext.h
@@ -0,0 +1,60 @@
+// Copyright Epic Games, Inc. All Rights Reserved.
+
+#pragma once
+
+#include "CoreMinimal.h"
+#include "UAF/Attributes/AttributeBindingData.h"
+#include "UAF/Attributes/AttributeCollection.h"
+
+#define UE_API UAF_API
+
+class UAbstractSkeletonSetBinding;
+class USkeleton;
+
+namespace UE::UAF
+{
+	/*
+	 * Attribute Evaluation Context
+	 *
+	 * Encapsulates the state that controls what attributes we should produce.
+	 */
+	class FAttributeEvaluationContext final
+	{
+	public:
+		UE_API FAttributeEvaluationContext();
+		UE_API FAttributeEvaluationContext(const UAbstractSkeletonSetBinding* SetBinding, FName SetName);
+
+		// Returns whether or not this context is valid
+		[[nodiscard]] UE_API bool IsValid() const;
+
+		// Returns the abstract skeleton set binding asset we are evaluating with (or nullptr if we aren't valid)
+		[[nodiscard]] UE_API const UAbstractSkeletonSetBinding* GetSetBinding() const;
+
+		// Returns the concrete skeleton we are evaluating with (or nullptr if we aren't valid)
+		[[nodiscard]] UE_API const USkeleton* GetSkeleton() const;
+
+		// Returns the named set we are currently evaluating with (or nullptr if we aren't valid)
+		[[nodiscard]] UE_API const FAttributeNamedSetPtr& GetNamedSet() const;
+
+		// Returns the default attribute values for the current named set (or nullptr if we aren't valid)
+		[[nodiscard]] UE_API const FAttributeCollectionPtr& GetDefaultAttributeValues() const;
+
+	private:
+		// The set binding we are evaluating with
+		const UAbstractSkeletonSetBinding* SetBinding = nullptr;
+
+		// The target skeleton we are evaluating with
+		const USkeleton* Skeleton = nullptr;
+
+		// Binding data owned by the set binding
+		FAttributeBindingDataPtr BindingData;
+
+		// Named set instance to produce an output for
+		FAttributeNamedSetPtr NamedSet;
+
+		// Default values as specified by the current set binding and named set
+		FAttributeCollectionPtr DefaultAttributeValues;
+	};
+}
+
+#undef UE_API
diff --git a/Engine/Plugins/Experimental/UAF/UAF/Source/UAF/Public/UAF/Attributes/AttributeRuntimeRegistry.h b/Engine/Plugins/Experimental/UAF/UAF/Source/UAF/Public/UAF/Attributes/AttributeRuntimeRegistry.h
new file mode 100644
--- /dev/null
+++ b/Engine/Plugins/Experimental/UAF/UAF/Source/UAF/Public/UAF/Attributes/AttributeRuntimeRegistry.h
@@ -0,0 +1,104 @@
+// Copyright Epic Games, Inc. All Rights Reserved.
+
+#pragma once
+
+#include "CoreMinimal.h"
+#include "UAF/Attributes/AttributeSetMap.h"
+#include "UAF/Attributes/Transformers/Transformer.h"
+
+#define UE_API UAF_API
+
+namespace UE::UAF
+{
+	// A global registry to store attribute runtime metadata
+	class FAttributeRuntimeRegistry
+	{
+	public:
+		UE_NONCOPYABLE(FAttributeRuntimeRegistry);	// We disallow copy/move semantics
+
+		using SetMapConstructorFun = TFunctionRef<FAttributeSetMap*(const FAttributeSetMap::FConstructArgs& Args)>;
+
+		// Returns the set map registry instance
+		[[nodiscard]] static UE_API FAttributeRuntimeRegistry& Get();
+
+		// Constructs a set map from its UStruct value type
+		[[nodiscard]] UE_API FAttributeSetMap* ConstructSetMap(const FAttributeSetMap::FConstructArgs& Args) const;
+
+		// Returns the map of known attribute value transformers
+		[[nodiscard]] UE_API FAttributeValueTransformerMapPtr GetTransformerMap() const;
+
+		// Registers a set map of the specified value type
+		template<class ValueType>
+		void RegisterSetMapInitializer();
+
+		template<class TransformerSpecializationType>
+		bool RegisterAttributeValueTransformer();
+
+	private:
+		UE_API FAttributeRuntimeRegistry();
+		UE_API void RegisterSetMapInitializer(UScriptStruct* ValueType, SetMapConstructorFun Initializer);
+		UE_API bool RegisterAttributeValueTransformer(FName TransformerName, UScriptStruct* ValueType, FRawTransformerFunc TransformerFunc);
+
+		// Initialize/destroy the global registry
+		static UE_API void Init();
+		static UE_API void Destroy();
+
+		// A map of ValueType -> Set Map constructors
+		// This is used for indirect construction when we have a UScriptStruct value type
+		TMap<UScriptStruct*, SetMapConstructorFun> SetMapInitializers;
+		mutable FRWLock SetMapInitializerLock;
+
+		// A shared pointer to a map of transformer name -> transformer list (a transformer per value type)
+		FAttributeValueTransformerMapPtr TransformerMap;
+		mutable FRWLock TransformerMapLock;
+
+		friend class FAnimNextModuleImpl;
+	};
+
+	//////////////////////////////////////////////////////////////////////////
+	// Implementation
+
+	template<class ValueType>
+	inline void FAttributeRuntimeRegistry::RegisterSetMapInitializer()
+	{
+		RegisterSetMapInitializer(ValueType::StaticStruct(), [](const FAttributeSetMap::FConstructArgs& Args)
+			{
+				typename TAttributeSetMap<ValueType>::FConstructArgs TypedArgs(Args.TypedSet, Args.ValueType, Args.ReallocFun);
+
+				return MakeAttributeSetMap<ValueType>(Args);
+			});
+	}
+
+	template<class TransformerSpecializationType>
+	inline bool FAttributeRuntimeRegistry::RegisterAttributeValueTransformer()
+	{
+		using TransformerType = typename TransformerSpecializationType::TransformerType;
+		using ValueType = typename TransformerSpecializationType::ValueType;
+		using FTransformFunc = typename TransformerType::FTransformFunc;
+
+		static_assert(std::is_base_of_v<FAttributeValueTransformer, TransformerType>, "Specified transformer must derive from FAttributeValueTransformer and provide the necessary API");
+
+		const FName TransformerName = TransformerType::TransformerName;
+		UScriptStruct* ValueTypeStruct = ValueType::StaticStruct();
+
+		// Signature must match our transformer
+		FTransformFunc SpecializedTransformFunc = &TransformerSpecializationType::Transform;
+
+#ifdef _MSC_VER
+	#pragma warning(push)
+	#pragma warning(disable : 4191) // unsafe conversion from 'function pointer type A' to 'function pointer type B'
+#endif
+
+		// Cast to a raw dummy function pointer type that we exclusively use for storage
+		// The transformer is responsible for casting it back to the original type it specifies
+		auto RawSpecializedTransformFunc = reinterpret_cast<FRawTransformerFunc>(SpecializedTransformFunc);
+
+#ifdef _MSC_VER
+	#pragma warning(pop)
+#endif
+
+		return RegisterAttributeValueTransformer(TransformerName, ValueTypeStruct, RawSpecializedTransformFunc);
+	}
+}
+
+#undef UE_API
diff --git a/Engine/Plugins/Experimental/UAF/UAF/Source/UAF/Public/UAF/Attributes/AttributeSetMap.h b/Engine/Plugins/Experimental/UAF/UAF/Source/UAF/Public/UAF/Attributes/AttributeSetMap.h
new file mode 100644
--- /dev/null
+++ b/Engine/Plugins/Experimental/UAF/UAF/Source/UAF/Public/UAF/Attributes/AttributeSetMap.h
@@ -0,0 +1,332 @@
+// Copyright Epic Games, Inc. All Rights Reserved.
+
+#pragma once
+
+#include "CoreMinimal.h"
+#include "Misc/MemStack.h"
+#include "Templates/MemoryOps.h"
+#include "UAF/Attributes/AttributeMappingKey.h"
+#include "UAF/Attributes/AttributeTypedSet.h"
+#include "UAF/Attributes/IndirectAllocator.h"
+
+#define UE_API UAF_API
+
+namespace UE::UAF
+{
+	// Untyped base class for TAttributeSetMap
+	class FAttributeSetMap
+	{
+	public:
+		UE_NONCOPYABLE(FAttributeSetMap);	// We disallow copy/move semantics
+
+		struct FDefaultConstructArgs
+		{
+			const FAttributeTypedSetPtr& TypedSet;
+			UScriptStruct* ValueType = nullptr;
+			FReallocFun ReallocFun = nullptr;
+
+			FDefaultConstructArgs(const FAttributeTypedSetPtr& InTypedSet, UScriptStruct* InValueType, FReallocFun InReallocFun)
+				: TypedSet(InTypedSet)
+				, ValueType(InValueType)
+				, ReallocFun(InReallocFun)
+			{
+			}
+		};
+
+		using FConstructArgs = FDefaultConstructArgs;
+
+		// Returns whether or not this map has any values
+		[[nodiscard]] UE_API bool IsEmpty() const;
+
+		// Returns the number of mapped values
+		[[nodiscard]] UE_API int32 Num() const;
+
+		// Returns the attribute typed set we map to
+		[[nodiscard]] UE_API const FAttributeTypedSetPtr& GetTypedSet() const;
+
+		// Returns the attribute type we map to
+		[[nodiscard]] UE_API UScriptStruct* GetAttributeType() const;
+
+		// Returns the mapped value type
+		[[nodiscard]] UE_API UScriptStruct* GetValueType() const;
+
+		// Returns the mapping key represented by this set
+		[[nodiscard]] UE_API FAttributeMappingKey GetMappingKey() const;
+
+		// Returns the total allocated size for this set map
+		[[nodiscard]] UE_API int32 GetAllocatedSize() const;
+
+		// Returns a copy of this set map using the specified allocator
+		[[nodiscard]] virtual FAttributeSetMap* Duplicate(FReallocFun InReallocFun) const = 0;
+
+		// Copies the content of this set map into another
+		// Both must have matching typed sets and value types
+		virtual void CopyTo(FAttributeSetMap& Other) const = 0;
+
+		// Sets a value at the specified index using a value setter function
+		// The value setter (and/or caller) is responsible for validating the value type
+		virtual void SetValueWithSetter(FAttributeSetIndex Index, const TFunctionRef<void(UScriptStruct* ValueType, uint8* OutValuePtr)>& ValueSetter) = 0;
+
+		// Sets a value at the specified index
+		// The called is responsible for validating the value type matches the one contained within this set map
+		template <class ValueType>
+		void SetValue(FAttributeSetIndex Index, const ValueType& Value);
+
+		// Gets a value from the specified index
+		// The called is responsible for validating the value type matches the one contained within this set map
+		template <class ValueType>
+		[[nodiscard]] ValueType GetValue(FAttributeSetIndex Index) const;
+
+	protected:
+		// Default construction for an empty set map
+		UE_API FAttributeSetMap();
+
+		// Constructs a set map but does not allocate its memory, derived types are responsible for it
+		UE_API FAttributeSetMap(const FConstructArgs& Args, int32 AllocatedSize);
+
+		// Destroys the set map and frees its memory
+		UE_API virtual ~FAttributeSetMap();
+
+		// Sets a value at the specified index
+		virtual void SetValueImpl(FAttributeSetIndex Index, const void* ValuePtr) = 0;
+
+		// Gets a value from the specified index
+		virtual void GetValueImpl(FAttributeSetIndex Index, void* OutValuePtr) const = 0;
+
+		// The set whose attributes we map values to
+		FAttributeTypedSetPtr TypedSet;
+
+		// The type of values this set map contains
+		UScriptStruct* ValueTypeStruct = nullptr;
+
+		// The total allocated size of this set map
+		int32 AllocatedSize = 0;
+
+		// The number of values allocated in our buffer
+		int32 NumValues = 0;
+
+		// The base untyped pointer for our values
+		uint8* ValuesPtr = nullptr;
+
+		// The indirect allocator function pointer
+		FReallocFun ReallocFun = nullptr;
+
+		friend class FAttributeCollection;
+		friend class FAttributeRuntimeRegistry;
+		friend void ReleaseAttributeSetMap(FAttributeSetMap* SetMap);
+	};
+
+	// Encapsulates a mapping of a value type to a specific attribute set (of any attribute type)
+	// Our keys are held by the attribute set we point to (keys are shared) while we hold the
+	// mapped values internally. Set maps and their values are allocated together in a contiguous buffer.
+	// Value types must be a UStruct
+	template<class ValueType>
+	class TAttributeSetMap : public FAttributeSetMap
+	{
+	public:
+		// Returns a copy of this set map using the specified allocator
+		[[nodiscard]] virtual FAttributeSetMap* Duplicate(FReallocFun InReallocFun) const override;
+
+		// Copies the content of this set map into another
+		// Both must have matching typed sets and value types
+		virtual void CopyTo(FAttributeSetMap& Other) const override;
+
+		// Value accessors by index
+		[[nodiscard]] ValueType& operator[](FAttributeSetIndex Index);
+
+		// Value accessors by index
+		[[nodiscard]] const ValueType& operator[](FAttributeSetIndex Index) const;
+
+		// Returns a pointer to the contiguous values held within
+		[[nodiscard]] ValueType* GetData();
+
+		// Returns a pointer to the contiguous values held within
+		[[nodiscard]] const ValueType* GetData() const;
+
+		// Returns the total size required for an instance of a set map of this type
+		[[nodiscard]] static int32 CalculateAllocationSize(int32 NumAttributes);
+
+	protected:
+		// Constructs a set map
+		TAttributeSetMap(const FConstructArgs& Args, int32 AllocatedSize);
+
+		// Destroys the set map and frees its memory
+		virtual ~TAttributeSetMap() override;
+
+		virtual void SetValueWithSetter(FAttributeSetIndex Index, const TFunctionRef<void(UScriptStruct* InValueType, uint8* OutValuePtr)>& ValueSetter) override;
+		virtual void SetValueImpl(FAttributeSetIndex Index, const void* ValuePtr) override;
+		virtual void GetValueImpl(FAttributeSetIndex Index, void* OutValuePtr) const override;
+
+		template<class T>
+		friend TAttributeSetMap<T>* MakeAttributeSetMap(const typename TAttributeSetMap<T>::FConstructArgs& Args);
+	};
+
+	// Allocates and constructs a new attribute set map for the specified value type
+	// Set maps and their values are contiguous in memory
+	template<class ValueType>
+	TAttributeSetMap<ValueType>* MakeAttributeSetMap(const typename TAttributeSetMap<ValueType>::FConstructArgs& Args);
+
+	// Releases and destroys the specified attribute set map
+	void ReleaseAttributeSetMap(FAttributeSetMap* SetMap);
+
+	//////////////////////////////////////////////////////////////////////////
+	// Implementation
+
+	template<class ValueType>
+	inline void FAttributeSetMap::SetValue(FAttributeSetIndex Index, const ValueType& Value)
+	{
+		checkf(ValueType::StaticStruct() == ValueTypeStruct, TEXT("Value type does not match the set map"));
+		SetValueImpl(Index, &Value);
+	}
+
+	template<class ValueType>
+	inline ValueType FAttributeSetMap::GetValue(FAttributeSetIndex Index) const
+	{
+		checkf(ValueType::StaticStruct() == ValueTypeStruct, TEXT("Value type does not match the set map"));
+
+		ValueType Result;
+		GetValueImpl(Index, &Result);
+
+		return Result;
+	}
+
+	template<class ValueType>
+	inline TAttributeSetMap<ValueType>::TAttributeSetMap(const FConstructArgs& Args, int32 InAllocatedSize)
+		: FAttributeSetMap(Args, InAllocatedSize)
+	{
+		const int32 NumAttributes = Args.TypedSet->Num();
+
+		check(InAllocatedSize == TAttributeSetMap<ValueType>::CalculateAllocationSize(NumAttributes));
+
+		if (NumAttributes != 0)
+		{
+			ValuesPtr = Align(reinterpret_cast<uint8*>(this) + sizeof(TAttributeSetMap<ValueType>), alignof(ValueType));
+			NumValues = NumAttributes;
+
+			DefaultConstructItems<ValueType>(ValuesPtr, NumAttributes);
+		}
+	}
+
+	template<class ValueType>
+	inline TAttributeSetMap<ValueType>::~TAttributeSetMap()
+	{
+		if (ValuesPtr != nullptr)
+		{
+			DestructItems<ValueType>(GetData(), NumValues);
+
+			ValuesPtr = nullptr;
+			NumValues = 0;
+		}
+	}
+
+	template<class ValueType>
+	inline FAttributeSetMap* TAttributeSetMap<ValueType>::Duplicate(FReallocFun InReallocFun) const
+	{
+		TAttributeSetMap<ValueType>* Copy = MakeAttributeSetMap<ValueType>(FConstructArgs(TypedSet, ValueTypeStruct, InReallocFun));
+
+		if (NumValues != 0)
+		{
+			ConstructItems<ValueType>(Copy->ValuesPtr, GetData(), NumValues);
+		}
+
+		return Copy;
+	}
+
+	template<class ValueType>
+	inline void TAttributeSetMap<ValueType>::CopyTo(FAttributeSetMap& Other) const
+	{
+		checkf(TypedSet == Other.GetTypedSet() && ValueTypeStruct == Other.GetValueType(), TEXT("Typed set and value type must match to allow copying"));
+		check(Other.Num() == Num());	// If the typed sets match, their sizes does too
+
+		TAttributeSetMap<ValueType>& Copy = reinterpret_cast<TAttributeSetMap<ValueType>&>(Other);
+
+		CopyAssignItems<ValueType>(Copy.GetData(), GetData(), NumValues);
+	}
+
+	template<class ValueType>
+	inline ValueType& TAttributeSetMap<ValueType>::operator[](FAttributeSetIndex Index)
+	{
+		checkf((Index.GetInt() >= 0) & (Index.GetInt() < NumValues), TEXT("Map index out of bounds: %lld into an map of size %lld"), (long long)Index.GetInt(), (long long)NumValues); // & for one branch
+		return GetData()[Index.GetInt()];
+	}
+
+	template<class ValueType>
+	inline const ValueType& TAttributeSetMap<ValueType>::operator[](FAttributeSetIndex Index) const
+	{
+		checkf((Index.GetInt() >= 0) & (Index.GetInt() < NumValues), TEXT("Map index out of bounds: %lld into an map of size %lld"), (long long)Index.GetInt(), (long long)NumValues); // & for one branch
+		return GetData()[Index.GetInt()];
+	}
+
+	template<class ValueType>
+	inline ValueType* TAttributeSetMap<ValueType>::GetData()
+	{
+		return reinterpret_cast<ValueType*>(ValuesPtr);
+	}
+
+	template<class ValueType>
+	inline const ValueType* TAttributeSetMap<ValueType>::GetData() const
+	{
+		return reinterpret_cast<const ValueType*>(ValuesPtr);
+	}
+
+	template<class ValueType>
+	inline void TAttributeSetMap<ValueType>::SetValueWithSetter(FAttributeSetIndex Index, const TFunctionRef<void(UScriptStruct* InValueType, uint8* OutValuePtr)>& ValueSetter)
+	{
+		checkf((Index.GetInt() >= 0) & (Index.GetInt() < NumValues), TEXT("Map index out of bounds: %lld into an map of size %lld"), (long long)Index.GetInt(), (long long)NumValues); // & for one branch
+		ValueType* ValuePtr = GetData() + Index.GetInt();
+		ValueSetter(ValueTypeStruct, reinterpret_cast<uint8*>(ValuePtr));
+	}
+
+	template<class ValueType>
+	inline void TAttributeSetMap<ValueType>::SetValueImpl(FAttributeSetIndex Index, const void* ValuePtr)
+	{
+		checkf((Index.GetInt() >= 0) & (Index.GetInt() < NumValues), TEXT("Map index out of bounds: %lld into an map of size %lld"), (long long)Index.GetInt(), (long long)NumValues); // & for one branch
+		const ValueType* Value = reinterpret_cast<const ValueType*>(ValuePtr);
+		GetData()[Index.GetInt()] = *Value;
+	}
+
+	template<class ValueType>
+	inline void TAttributeSetMap<ValueType>::GetValueImpl(FAttributeSetIndex Index, void* OutValuePtr) const
+	{
+		checkf((Index.GetInt() >= 0) & (Index.GetInt() < NumValues), TEXT("Map index out of bounds: %lld into an map of size %lld"), (long long)Index.GetInt(), (long long)NumValues); // & for one branch
+		ValueType* OutValue = reinterpret_cast<ValueType*>(OutValuePtr);
+		*OutValue = GetData()[Index.GetInt()];
+	}
+
+	template<class ValueType>
+	inline int32 TAttributeSetMap<ValueType>::CalculateAllocationSize(int32 NumAttributes)
+	{
+		return Align(sizeof(TAttributeSetMap<ValueType>), alignof(ValueType)) + sizeof(ValueType) * NumAttributes;
+	}
+
+	template<class ValueType>
+	inline TAttributeSetMap<ValueType>* MakeAttributeSetMap(const typename TAttributeSetMap<ValueType>::FConstructArgs& Args)
+	{
+		// We allocate the set map and its values in a contiguous buffer part of a single allocation
+		// As such, the set map struct ends up being a small header
+		const int32 AllocationSize = TAttributeSetMap<ValueType>::CalculateAllocationSize(Args.TypedSet->Num());
+
+		// Allocate our single buffer
+		uint8* Ptr = (*Args.ReallocFun)(nullptr, 0, AllocationSize);
+
+		// Construct our set map over it, the specialized constructor is responsible for partitioning the rest of the buffer
+		return new(Ptr) TAttributeSetMap<ValueType>(Args, AllocationSize);
+	}
+
+	inline void ReleaseAttributeSetMap(FAttributeSetMap* SetMap)
+	{
+		if (SetMap == nullptr)
+		{
+			return;
+		}
+
+		const int32 AllocationSize = SetMap->AllocatedSize;
+		const FReallocFun ReallocFun = SetMap->ReallocFun;
+
+		SetMap->~FAttributeSetMap();
+
+		(*ReallocFun)(reinterpret_cast<uint8*>(SetMap), AllocationSize, 0);
+	}
+}
+
+#undef UE_API
diff --git a/Engine/Plugins/Experimental/UAF/UAF/Source/UAF/Public/UAF/Attributes/IndirectAllocator.h b/Engine/Plugins/Experimental/UAF/UAF/Source/UAF/Public/UAF/Attributes/IndirectAllocator.h
new file mode 100644
--- /dev/null
+++ b/Engine/Plugins/Experimental/UAF/UAF/Source/UAF/Public/UAF/Attributes/IndirectAllocator.h
@@ -0,0 +1,68 @@
+// Copyright Epic Games, Inc. All Rights Reserved.
+
+#pragma once
+
+#include "CoreMinimal.h"
+#include "Containers/ContainerAllocationPolicies.h"
+#include "Misc/MemStack.h"
+
+namespace UE::UAF
+{
+	// Use allocator pointer instead of template argument
+	// The allocator should behave as follows:
+	//    - When Ptr is nullptr (and OldSize is 0), we are allocating memory
+	//    - When Ptr is valid but NewSize is 0, we are freeing memory
+	//    - When Ptr is valid and NewSize is not 0, we are re-allocating memory
+	//    - It assumes that the elements that live in said memory can be trivially copyable
+	//    - Allocations use default alignment
+	typedef uint8* (*FReallocFun)(uint8* OldPtr, uint32 OldSize, uint32 NewSize);
+
+	// This type can be specialized for template style allocators for easy interop
+	// with the engine allocator types
+	template<class AllocatorType>
+	struct FAllocatorTypeTrait
+	{
+		// Specialize this struct for your allocator by implementing this function
+		static uint8* Realloc(uint8* OldPtr, uint32 OldSize, uint32 NewSize) = delete;
+	};
+
+	template<>
+	struct FAllocatorTypeTrait<FDefaultAllocator>
+	{
+		static uint8* Realloc(uint8* OldPtr, uint32 OldSize, uint32 NewSize)
+		{
+			if (OldPtr == nullptr && NewSize == 0)
+			{
+				// Nothing to do if we weren't previously allocated and we aren't requesting a valid size
+				return nullptr;
+			}
+
+			return (uint8*)FMemory::Realloc(OldPtr, NewSize, DEFAULT_ALIGNMENT);
+		}
+	};
+
+	template<>
+	struct FAllocatorTypeTrait<TMemStackAllocator<>>
+	{
+		static uint8* Realloc(uint8* OldPtr, uint32 OldSize, uint32 NewSize)
+		{
+			if (NewSize == 0)
+			{
+				// Nothing to do if we are freeing memory
+				return nullptr;
+			}
+
+			// Allocate memory from the stack.
+			uint8* NewPtr = FMemStack::Get().PushBytes(NewSize, DEFAULT_ALIGNMENT);
+
+			// If we are growing an existing allocation, copy from the old
+			if (OldPtr != nullptr)
+			{
+				const uint32 NumToCopy = FMath::Min(NewSize, OldSize);
+				FMemory::Memcpy(NewPtr, OldPtr, NumToCopy);
+			}
+
+			return NewPtr;
+		}
+	};
+}
diff --git a/Engine/Plugins/Experimental/UAF/UAF/Source/UAF/Public/UAF/Attributes/Transformers/Accumulate.h b/Engine/Plugins/Experimental/UAF/UAF/Source/UAF/Public/UAF/Attributes/Transformers/Accumulate.h
new file mode 100644
--- /dev/null
+++ b/Engine/Plugins/Experimental/UAF/UAF/Source/UAF/Public/UAF/Attributes/Transformers/Accumulate.h
@@ -0,0 +1,40 @@
+// Copyright Epic Games, Inc. All Rights Reserved.
+
+#pragma once
+
+#include "CoreMinimal.h"
+
+#include "UAF/Attributes/EngineAttributes.h"
+#include "UAF/Attributes/Transformers/Transformer.h"
+
+#define UE_API UAF_API
+
+namespace UE::UAF::Transformers
+{
+	// Attribute value transformer to accumulate values
+	// Accumulate is the repeating step within an N-step interpolation: Overwrite, Accumulate0, Accumulate1, ...
+	// Make sure to also implement variations for Overwrite and Interpolate
+	struct FAccumulate final : public FAttributeValueTransformer
+	{
+		static UE_API FName TransformerName;
+
+		using FTransformFunc = void (*)(UE::UAF::FAttributeSetMap* SetMapA, const UE::UAF::FAttributeSetMap* SetMapB, float Alpha);
+
+		// Accumulate attribute values
+		// accumulate(A, B, Alpha) -> A += (B * Alpha)
+		static UE_API void Apply(const FAttributeValueTransformerMapPtr& TransformerMap, FAttributeCollection& CollectionA, const FAttributeCollection& CollectionB, float Alpha);
+	};
+
+	//////////////////////////////////////////////////////////////////////////
+	// Specific attribute value type specializations
+
+	struct FAccumulate_BoneTransformAttribute
+	{
+		using TransformerType = FAccumulate;
+		using ValueType = FBoneTransformAnimationAttribute;
+
+		static void Transform(UE::UAF::FAttributeSetMap* SetMapA, const UE::UAF::FAttributeSetMap* SetMapB, float Alpha);
+	};
+}
+
+#undef UE_API
diff --git a/Engine/Plugins/Experimental/UAF/UAF/Source/UAF/Public/UAF/Attributes/Transformers/Interpolate.h b/Engine/Plugins/Experimental/UAF/UAF/Source/UAF/Public/UAF/Attributes/Transformers/Interpolate.h
new file mode 100644
--- /dev/null
+++ b/Engine/Plugins/Experimental/UAF/UAF/Source/UAF/Public/UAF/Attributes/Transformers/Interpolate.h
@@ -0,0 +1,39 @@
+// Copyright Epic Games, Inc. All Rights Reserved.
+
+#pragma once
+
+#include "CoreMinimal.h"
+
+#include "UAF/Attributes/EngineAttributes.h"
+#include "UAF/Attributes/Transformers/Transformer.h"
+
+#define UE_API UAF_API
+
+namespace UE::UAF::Transformers
+{
+	// Attribute value transformer to interpolate values
+	// Make sure to also implement variations for Overwrite and Accumulate
+	struct FInterpolate final : public FAttributeValueTransformer
+	{
+		static UE_API FName TransformerName;
+
+		using FTransformFunc = void (*)(UE::UAF::FAttributeSetMap* SetMapA, const UE::UAF::FAttributeSetMap* SetMapB, float Alpha);
+
+		// Interpolate attribute values
+		// lerp(A, B, Alpha) -> A = (A * (1.0 - Alpha)) + (B * Alpha)
+		static UE_API void Apply(const FAttributeValueTransformerMapPtr& TransformerMap, FAttributeCollection& CollectionA, const FAttributeCollection& CollectionB, float Alpha);
+	};
+
+	//////////////////////////////////////////////////////////////////////////
+	// Specific attribute value type specializations
+
+	struct FInterpolate_BoneTransformAttribute
+	{
+		using TransformerType = FInterpolate;
+		using ValueType = FBoneTransformAnimationAttribute;
+
+		static void Transform(UE::UAF::FAttributeSetMap* SetMapA, const UE::UAF::FAttributeSetMap* SetMapB, float Alpha);
+	};
+}
+
+#undef UE_API
diff --git a/Engine/Plugins/Experimental/UAF/UAF/Source/UAF/Public/UAF/Attributes/Transformers/Overwrite.h b/Engine/Plugins/Experimental/UAF/UAF/Source/UAF/Public/UAF/Attributes/Transformers/Overwrite.h
new file mode 100644
--- /dev/null
+++ b/Engine/Plugins/Experimental/UAF/UAF/Source/UAF/Public/UAF/Attributes/Transformers/Overwrite.h
@@ -0,0 +1,40 @@
+// Copyright Epic Games, Inc. All Rights Reserved.
+
+#pragma once
+
+#include "CoreMinimal.h"
+
+#include "UAF/Attributes/EngineAttributes.h"
+#include "UAF/Attributes/Transformers/Transformer.h"
+
+#define UE_API UAF_API
+
+namespace UE::UAF::Transformers
+{
+	// Attribute value transformer to overwrite values
+	// Accumulate is the initial step within an N-step interpolation: Overwrite, Accumulate0, Accumulate1, ...
+	// Make sure to also implement variations for Accumulate and Interpolate
+	struct FOverwrite final : public FAttributeValueTransformer
+	{
+		static UE_API FName TransformerName;
+
+		using FTransformFunc = void (*)(UE::UAF::FAttributeSetMap* SetMapA, const UE::UAF::FAttributeSetMap* SetMapB, float Alpha);
+
+		// Overwrite attribute values
+		// overwrite(A, B, Alpha) -> A = (B * Alpha)
+		static UE_API void Apply(const FAttributeValueTransformerMapPtr& TransformerMap, FAttributeCollection& CollectionA, const FAttributeCollection& CollectionB, float Alpha);
+	};
+
+	//////////////////////////////////////////////////////////////////////////
+	// Specific attribute value type specializations
+
+	struct FOverwrite_BoneTransformAttribute
+	{
+		using TransformerType = FOverwrite;
+		using ValueType = FBoneTransformAnimationAttribute;
+
+		static void Transform(UE::UAF::FAttributeSetMap* SetMapA, const UE::UAF::FAttributeSetMap* SetMapB, float Alpha);
+	};
+}
+
+#undef UE_API
diff --git a/Engine/Plugins/Experimental/UAF/UAF/Source/UAF/Public/UAF/Attributes/Transformers/Sanitize.h b/Engine/Plugins/Experimental/UAF/UAF/Source/UAF/Public/UAF/Attributes/Transformers/Sanitize.h
new file mode 100644
--- /dev/null
+++ b/Engine/Plugins/Experimental/UAF/UAF/Source/UAF/Public/UAF/Attributes/Transformers/Sanitize.h
@@ -0,0 +1,40 @@
+// Copyright Epic Games, Inc. All Rights Reserved.
+
+#pragma once
+
+#include "CoreMinimal.h"
+
+#include "UAF/Attributes/EngineAttributes.h"
+#include "UAF/Attributes/Transformers/Transformer.h"
+
+#define UE_API UAF_API
+
+namespace UE::UAF::Transformers
+{
+	// Attribute value transformer to sanitize values
+	// This can be used to normalize rotations and other like values as well as check for NaN/Inf values
+	struct FSanitize final : public FAttributeValueTransformer
+	{
+		static UE_API FName TransformerName;
+
+		using FTransformFunc = void (*)(UE::UAF::FAttributeSetMap* SetMap);
+
+		// Sanitize attribute values
+		// This can be used to normalize rotations and other like values as well as check for NaN/Inf values
+		static UE_API void Apply(const FAttributeValueTransformerMapPtr& TransformerMap, FAttributeCollection& Collection);
+	};
+
+	//////////////////////////////////////////////////////////////////////////
+	// Specific attribute value type specializations
+
+	// Sanitize bone transforms (normalize rotation, ensure finite)
+	struct FSanitize_BoneTransformAttribute
+	{
+		using TransformerType = FSanitize;
+		using ValueType = FBoneTransformAnimationAttribute;
+
+		static void Transform(UE::UAF::FAttributeSetMap* SetMap);
+	};
+}
+
+#undef UE_API
diff --git a/Engine/Plugins/Experimental/UAF/UAF/Source/UAF/Public/UAF/Attributes/Transformers/Transformer.h b/Engine/Plugins/Experimental/UAF/UAF/Source/UAF/Public/UAF/Attributes/Transformers/Transformer.h
new file mode 100644
--- /dev/null
+++ b/Engine/Plugins/Experimental/UAF/UAF/Source/UAF/Public/UAF/Attributes/Transformers/Transformer.h
@@ -0,0 +1,164 @@
+// Copyright Epic Games, Inc. All Rights Reserved.
+
+#pragma once
+
+#include "CoreMinimal.h"
+#include "UAF/Attributes/AttributeCollection.h"
+
+#define UE_API UAF_API
+
+namespace UE::UAF
+{
+	class FAttributeRuntimeRegistry;
+	class FAttributeValueTransformerList;
+
+	// A raw function pointer for all transformers (static functions)
+	// Only used for storage purposes, NEVER CALL THIS!
+	using FRawTransformerFunc = void(*)();
+
+	// Attribute value transformer base type
+	// All transformer types must derive from this
+	struct FAttributeValueTransformer
+	{
+		// The name of the transformer (exposed in derived types)
+		//static FName TransformerName;
+
+		// The function signature of the transformed (exposed in derived types)
+		//using FTransformFunc = void (*)(...);
+	};
+
+	// An iterator over a transformer list within the attribute runtime registry
+	class FAttributeValueTransformerListIterator
+	{
+	public:
+		[[nodiscard]] FAttributeValueTransformerListIterator(const FAttributeValueTransformerList& TransformerList);
+
+		// Increments and moves the iterator to the next transformer entry
+		FAttributeValueTransformerListIterator& operator++();
+
+		// Returns whether or not the iterator still contains values
+		[[nodiscard]] explicit operator bool() const;
+
+		// Returns the value type that the transformer operates on for the current entry
+		[[nodiscard]] UScriptStruct* GetValueType() const;
+
+		// Returns the transformer for the current entry
+		[[nodiscard]] FRawTransformerFunc GetTransformerFunc() const;
+
+	private:
+		// Iterator over our value types
+		TArray<UScriptStruct*>::TConstIterator ValueTypesIt;
+
+		// Iterator over our transformer instances
+		TArray<FRawTransformerFunc>::TConstIterator TransformerFuncIt;
+	};
+
+	// A list of transformer specializations that operators over various value types
+	// Transformers have a base type (e.g. interpolate, sanitize) and can be specialized to operate on arbitrary
+	// value types (e.g. FFloatAnimationAttribute)
+	class FAttributeValueTransformerList
+	{
+	public:
+		// Creates an empty transformer list
+		UE_API FAttributeValueTransformerList();
+
+		// Creates a transformer list for the specified transformer type
+		[[nodiscard]] UE_API explicit FAttributeValueTransformerList(FName TransformerName);
+
+		// Adds a new transformer entry
+		UE_API bool AddTransformer(UScriptStruct* ValueType, FRawTransformerFunc TransformerFunc);
+
+		// Returns an iterator over all transformer specializations within
+		[[nodiscard]] FAttributeValueTransformerListIterator CreateTransformerIterator() const { return FAttributeValueTransformerListIterator(*this); }
+
+	private:
+		// The type of the transformer (e.g. interpolate, sanitize)
+		FName TransformerName;
+
+		// List of attribute value types for which we have transformers registered (e.g. FTransformAnimationAttribute, FFloatAnimationAttribute) (sorted by value)
+		TArray<UScriptStruct*> ValueTypes;
+
+		// List of transformer functions for each value type (sorted by value type)
+		TArray<FRawTransformerFunc> Transformers;
+
+		friend FAttributeRuntimeRegistry;
+		friend FAttributeValueTransformerListIterator;
+	};
+
+	using FAttributeValueTransformerMapPtr = TSharedPtr<TMap<FName, FAttributeValueTransformerList>>;
+
+	// An operator structure to join attribute value transformer lists and attribute collections using InnerJoin/OuterJoin/etc
+	// Ensures that the predicate is called for each set map with a value type that has a corresponding transformer
+	struct FAttributeValueTransformerListWithAttributeCollectionJoinOp
+	{
+		// Key sorting uses signed comparison, use the largest signed value regardless of pointer size
+		static constexpr intptr_t WILDCARD_ATTRIBUTE_TYPE_VALUE = static_cast<intptr_t>(static_cast<uintptr_t>(~0ULL) >> 1);
+
+		// We enable fuzzy matching by implementing its API
+		bool IsMatch(FAttributeMappingKey ReferenceKey, FAttributeMappingKey TestKey) const
+		{
+			// Our wildcard key ensures that it is larger than all other keys meaning it will only ever be the
+			// reference key if we've matched all other smaller values (all valid ones)
+			if (reinterpret_cast<intptr_t>(TestKey.GetAttributeType()) == WILDCARD_ATTRIBUTE_TYPE_VALUE)
+			{
+				// If our test value is our wildcard value, we match if our value types match
+				return ReferenceKey.GetValueType() == TestKey.GetValueType();
+			}
+
+			// Otherwise we test normally
+			return ReferenceKey == TestKey;
+		}
+
+		FAttributeMappingKey GetLargestKey() const { return FAttributeMappingKey::LARGEST_VALUE; }
+
+		// FAttributeValueTransformerListIterator entries
+		FAttributeMappingKey GetKey(const FAttributeValueTransformerListIterator& It) const
+		{
+			// We use a wildcard entry for the attribute type within our key to ensure it is always larger than all other valid values
+			return FAttributeMappingKey(reinterpret_cast<UScriptStruct*>(WILDCARD_ATTRIBUTE_TYPE_VALUE), It.GetValueType());
+		}
+
+		FRawTransformerFunc GetValue(const FAttributeValueTransformerListIterator& It) const { return It.GetTransformerFunc(); }
+
+		// Normal FAttributeCollectionSetMapJoinOp
+		FAttributeMappingKey GetKey(const FAttributeCollection::TSetMapIterator& It) const { return It.GetKey(); }
+		FAttributeMappingKey GetKey(const FAttributeCollection::TConstSetMapIterator& It) const { return It.GetKey(); }
+
+		FAttributeSetMap* GetValue(const FAttributeCollection::TSetMapIterator& It) const { return It.GetMap(); }
+		const FAttributeSetMap* GetValue(const FAttributeCollection::TConstSetMapIterator& It) const { return It.GetMap(); }
+	};
+
+	//////////////////////////////////////////////////////////////////////////
+	// Implementation
+
+	inline FAttributeValueTransformerListIterator::FAttributeValueTransformerListIterator(const FAttributeValueTransformerList& TransformerList)
+		: ValueTypesIt(TransformerList.ValueTypes.CreateConstIterator())
+		, TransformerFuncIt(TransformerList.Transformers.CreateConstIterator())
+	{
+	}
+
+	inline FAttributeValueTransformerListIterator& FAttributeValueTransformerListIterator::operator++()
+	{
+		++ValueTypesIt;
+		++TransformerFuncIt;
+
+		return *this;
+	}
+
+	inline FAttributeValueTransformerListIterator::operator bool() const
+	{
+		return (bool)ValueTypesIt;
+	}
+
+	inline UScriptStruct* FAttributeValueTransformerListIterator::GetValueType() const
+	{
+		return *ValueTypesIt;
+	}
+
+	inline FRawTransformerFunc FAttributeValueTransformerListIterator::GetTransformerFunc() const
+	{
+		return *TransformerFuncIt;
+	}
+}
+
+#undef UE_API
diff --git a/Engine/Plugins/Experimental/UAF/UAF/Source/UAFEditor/Internal/UAF/AbstractSkeleton/SetBindingFactory.h b/Engine/Plugins/Experimental/UAF/UAF/Source/UAFEditor/Internal/UAF/AbstractSkeleton/SetBindingFactory.h
--- a/Engine/Plugins/Experimental/UAF/UAF/Source/UAFEditor/Internal/UAF/AbstractSkeleton/SetBindingFactory.h
+++ b/Engine/Plugins/Experimental/UAF/UAF/Source/UAFEditor/Internal/UAF/AbstractSkeleton/SetBindingFactory.h
@@ -7,7 +7,7 @@
 
 #include "SetBindingFactory.generated.h"
 
-UCLASS()
+UCLASS(MinimalAPI)
 class UAbstractSkeletonSetBindingFactory : public UFactory
 {
 	GENERATED_BODY()
diff --git a/Engine/Plugins/Experimental/UAF/UAF/Source/UAFEditor/Private/UAF/AbstractSkeleton/Sets/SetBindingFactory.cpp b/Engine/Plugins/Experimental/UAF/UAF/Source/UAFEditor/Private/UAF/AbstractSkeleton/Sets/SetBindingFactory.cpp
--- a/Engine/Plugins/Experimental/UAF/UAF/Source/UAFEditor/Private/UAF/AbstractSkeleton/Sets/SetBindingFactory.cpp
+++ b/Engine/Plugins/Experimental/UAF/UAF/Source/UAFEditor/Private/UAF/AbstractSkeleton/Sets/SetBindingFactory.cpp
@@ -4,6 +4,8 @@
 #include "UAF/AbstractSkeleton/AbstractSkeletonSetBinding.h"
 #include "UObject/Package.h"
 
+#include UE_INLINE_GENERATED_CPP_BY_NAME(SetBindingFactory)
+
 UAbstractSkeletonSetBindingFactory::UAbstractSkeletonSetBindingFactory()
 {
 	bCreateNew = true;
diff --git a/Engine/Plugins/Experimental/UAF/UAF/Source/UAFEditor/Private/UAF/AbstractSkeleton/Sets/SetCollectionFactory.cpp b/Engine/Plugins/Experimental/UAF/UAF/Source/UAFEditor/Private/UAF/AbstractSkeleton/Sets/SetCollectionFactory.cpp
--- a/Engine/Plugins/Experimental/UAF/UAF/Source/UAFEditor/Private/UAF/AbstractSkeleton/Sets/SetCollectionFactory.cpp
+++ b/Engine/Plugins/Experimental/UAF/UAF/Source/UAFEditor/Private/UAF/AbstractSkeleton/Sets/SetCollectionFactory.cpp
@@ -5,6 +5,8 @@
 #include "UAF/AbstractSkeleton/AbstractSkeletonSetCollection.h"
 #include "UObject/Package.h"
 
+#include UE_INLINE_GENERATED_CPP_BY_NAME(SetCollectionFactory)
+
 UAbstractSkeletonSetCollectionFactory::UAbstractSkeletonSetCollectionFactory()
 {
 	bCreateNew = true;
diff --git a/Engine/Plugins/Experimental/UAF/UAFAnimGraph/Source/UAFAnimGraph/Private/EvaluationVM/EvaluationVM.cpp b/Engine/Plugins/Experimental/UAF/UAFAnimGraph/Source/UAFAnimGraph/Private/EvaluationVM/EvaluationVM.cpp
--- a/Engine/Plugins/Experimental/UAF/UAFAnimGraph/Source/UAFAnimGraph/Private/EvaluationVM/EvaluationVM.cpp
+++ b/Engine/Plugins/Experimental/UAF/UAFAnimGraph/Source/UAFAnimGraph/Private/EvaluationVM/EvaluationVM.cpp
@@ -9,6 +9,7 @@
 namespace UE::UAF
 {
 	const FEvaluationVMStackName KEYFRAME_STACK_NAME = FName(TEXT("KeyframeStack"));
+	const FEvaluationVMStackName ATTRIBUTE_STACK_NAME = FName(TEXT("AttributeStack"));
 
 	FEvaluationVMStack::~FEvaluationVMStack()
 	{
@@ -57,9 +58,27 @@ namespace UE::UAF
 		}
 	}
 
+	FEvaluationVM::FEvaluationVM(EEvaluationFlags InEvaluationFlags, const UAbstractSkeletonSetBinding& InSetBinding, FName InSetName, int32 InCurrentLOD)
+		: InternalContext(&InSetBinding, InSetName)
+		, AttributeValueTransformerMap(FAttributeRuntimeRegistry::Get().GetTransformerMap())
+		, CurrentLOD(InCurrentLOD)
+		, EvaluationFlags(InEvaluationFlags)
+	{
+	}
+
+	// REMOVE ME: Exclusively for testing that before/after yields the same results
+	FEvaluationVM::FEvaluationVM(EEvaluationFlags InEvaluationFlags, const FReferencePose& InReferencePose, const UAbstractSkeletonSetBinding* InSetBinding, FName InSetName, int32 InCurrentLOD)
+		: ReferencePose(&InReferencePose)
+		, InternalContext(InSetBinding, InSetName)
+		, AttributeValueTransformerMap(FAttributeRuntimeRegistry::Get().GetTransformerMap())
+		, CurrentLOD(InCurrentLOD)
+		, EvaluationFlags(InEvaluationFlags)
+	{
+	}
+
 	bool FEvaluationVM::IsValid() const
 	{
-		return ReferencePose != nullptr;
+		return ReferencePose != nullptr || InternalContext.IsValid();
 	}
 
 	void FEvaluationVM::Shrink()
@@ -111,6 +130,52 @@ namespace UE::UAF
 		return Keyframe;
 	}
 
+	FAttributeCollectionStack FEvaluationVM::MakeAttributeCollection(bool bAdditive, bool bInitialize) const
+	{
+		if (!InternalContext.IsValid())
+		{
+			return FAttributeCollectionStack();
+		}
+
+		FAttributeCollectionStack Collection(InternalContext.GetNamedSet());
+
+		// TODO: How to handle additive attributes? We need to store an alternate set of default values
+		// Perhaps that should be a parameter we pass to AddSetMap, an initialization space (or to the collection)
+		// and when we create the map, instead of copying defaults, we default construct instead using that space to decide on our value
+
+		// Iterate over our named set attribute types and initialize our collection with its default values
+		for (auto It = InternalContext.GetNamedSet()->CreateTypedSetIterator(); It; ++It)
+		{
+			FAttributeTypedSetPtr TypedSet = *It;
+			UScriptStruct* AttributeType = TypedSet->GetType();
+
+			// By default, attributes map to values of their type
+			const FAttributeMappingKey MappingKey = FAttributeMappingKey::MakeFromTo(AttributeType);
+
+			FAttributeSetMap* Copy = Collection.AddSetMap(MappingKey);
+
+			if (bInitialize)
+			{
+				if (const FAttributeSetMap* DefaultValuesMap = InternalContext.GetDefaultAttributeValues()->FindSetMap(MappingKey))
+				{
+					DefaultValuesMap->CopyTo(*Copy);
+				}
+			}
+		}
+
+		return Collection;
+	}
+
+	const FAttributeEvaluationContext& FEvaluationVM::GetContext() const
+	{
+		return InternalContext;
+	}
+
+	const FAttributeValueTransformerMapPtr& FEvaluationVM::GetTransformerMap() const
+	{
+		return AttributeValueTransformerMap;
+	}
+
 	FEvaluationVMStack& FEvaluationVM::GetOrCreateStack(const FEvaluationVMStackName& StackName, uint32 TypeID)
 	{
 		if (FEvaluationVMStack* Stack = InternalStacks.FindByHash(StackName.NameHash, StackName.Name))
diff --git a/Engine/Plugins/Experimental/UAF/UAFAnimGraph/Source/UAFAnimGraph/Private/EvaluationVM/Tasks/ApplyAdditiveKeyframe.cpp b/Engine/Plugins/Experimental/UAF/UAFAnimGraph/Source/UAFAnimGraph/Private/EvaluationVM/Tasks/ApplyAdditiveKeyframe.cpp
--- a/Engine/Plugins/Experimental/UAF/UAFAnimGraph/Source/UAFAnimGraph/Private/EvaluationVM/Tasks/ApplyAdditiveKeyframe.cpp
+++ b/Engine/Plugins/Experimental/UAF/UAFAnimGraph/Source/UAFAnimGraph/Private/EvaluationVM/Tasks/ApplyAdditiveKeyframe.cpp
@@ -29,6 +29,14 @@ void FAnimNextApplyAdditiveKeyframeTask::Execute(UE::UAF::FEvaluationVM& VM) con
 {
 	using namespace UE::UAF;
 
+	if (VM.GetContext().IsValid())
+	{
+		// TODO: Implement with new attribute runtime
+		TUniquePtr<FAttributeCollection> DiscardedCollection;
+		(void)VM.PopValue(ATTRIBUTE_STACK_NAME, DiscardedCollection);
+		return;
+	}
+
 	// Pop our top two poses, we'll re-use the top keyframe for our result
 
 	TUniquePtr<FKeyframeState> AdditiveKeyframe;
diff --git a/Engine/Plugins/Experimental/UAF/UAFAnimGraph/Source/UAFAnimGraph/Private/EvaluationVM/Tasks/BlendKeyframes.cpp b/Engine/Plugins/Experimental/UAF/UAFAnimGraph/Source/UAFAnimGraph/Private/EvaluationVM/Tasks/BlendKeyframes.cpp
--- a/Engine/Plugins/Experimental/UAF/UAFAnimGraph/Source/UAFAnimGraph/Private/EvaluationVM/Tasks/BlendKeyframes.cpp
+++ b/Engine/Plugins/Experimental/UAF/UAFAnimGraph/Source/UAFAnimGraph/Private/EvaluationVM/Tasks/BlendKeyframes.cpp
@@ -6,6 +6,11 @@
 #include "EvaluationVM/KeyframeState.h"
 #include "TransformArrayOperations.h"
 #include "Animation/InputScaleBias.h"
+#include "UAF/Attributes/AttributeCollection.h"
+#include "UAF/Attributes/Transformers/Accumulate.h"
+#include "UAF/Attributes/Transformers/Interpolate.h"
+#include "UAF/Attributes/Transformers/Overwrite.h"
+#include "UAF/Attributes/Transformers/Sanitize.h"
 
 #include UE_INLINE_GENERATED_CPP_BY_NAME(BlendKeyframes)
 
@@ -30,67 +35,111 @@ void FAnimNextBlendTwoKeyframesTask::Execute(UE::UAF::FEvaluationVM& VM) const
 {
 	using namespace UE::UAF;
 
-	// Pop our top two poses, we'll re-use the top keyframe for our result
-
-	TUniquePtr<FKeyframeState> KeyframeB;
-	if (!VM.PopValue(KEYFRAME_STACK_NAME, KeyframeB))
+	if (VM.GetContext().IsValid())
 	{
-		// We have no inputs, nothing to do
-		return;
-	}
+		// Pop our top two poses, we'll re-use the top keyframe for our result
 
-	TUniquePtr<FKeyframeState> KeyframeA;
-	if (!VM.PopValue(KEYFRAME_STACK_NAME, KeyframeA))
-	{
-		// We have a single input, leave it on top of the stack
-		VM.PushValue(KEYFRAME_STACK_NAME, MoveTemp(KeyframeB));
-		return;
-	}
+		TUniquePtr<FAttributeCollection> CollectionB;
+		if (!VM.PopValue(ATTRIBUTE_STACK_NAME, CollectionB))
+		{
+			// We have no inputs, nothing to do
+			return;
+		}
+
+		TUniquePtr<FAttributeCollection> CollectionA;
+		if (!VM.PopValue(ATTRIBUTE_STACK_NAME, CollectionA))
+		{
+			// We have a single input, leave it on top of the stack
+			VM.PushValue(ATTRIBUTE_STACK_NAME, MoveTemp(CollectionB));
+			return;
+		}
 
-	const float WeightOfPoseB = GetInterpolationAlpha(KeyframeA.Get(), KeyframeB.Get());
+		const float WeightOfPoseB = GetInterpolationAlpha(*CollectionA, *CollectionB);
 
-	if (!FAnimWeight::IsRelevant(WeightOfPoseB))
-	{
-		VM.PushValue(KEYFRAME_STACK_NAME, MoveTemp(KeyframeA));
-	}
-	else if (FAnimWeight::IsFullWeight(WeightOfPoseB))
-	{
-		VM.PushValue(KEYFRAME_STACK_NAME, MoveTemp(KeyframeB));
+		if (!FAnimWeight::IsRelevant(WeightOfPoseB))
+		{
+			VM.PushValue(ATTRIBUTE_STACK_NAME, MoveTemp(CollectionA));
+		}
+		else if (FAnimWeight::IsFullWeight(WeightOfPoseB))
+		{
+			VM.PushValue(ATTRIBUTE_STACK_NAME, MoveTemp(CollectionB));
+		}
+		else
+		{
+			// Interpolate in place the result within A
+			// Named sets must match as it ensures that our inputs have the same sizes/shapes
+			check(CollectionA->GetNamedSet() == CollectionB->GetNamedSet());
+
+			Transformers::FInterpolate::Apply(VM.GetTransformerMap(), *CollectionA, *CollectionB, WeightOfPoseB);
+			Transformers::FSanitize::Apply(VM.GetTransformerMap(), *CollectionA);
+
+			VM.PushValue(ATTRIBUTE_STACK_NAME, MoveTemp(CollectionA));
+		}
 	}
 	else
 	{
-		const float WeightOfPoseA = 1.0f - WeightOfPoseB;
+		// Pop our top two poses, we'll re-use the top keyframe for our result
 
-		if (EnumHasAnyFlags(VM.GetFlags(), EEvaluationFlags::Bones))
+		TUniquePtr<FKeyframeState> KeyframeB;
+		if (!VM.PopValue(KEYFRAME_STACK_NAME, KeyframeB))
 		{
-			check(KeyframeA->Pose.GetNumBones() == KeyframeB->Pose.GetNumBones());
+			// We have no inputs, nothing to do
+			return;
+		}
 
-			const FTransformArrayView KeyframeBTransformsView = KeyframeB->Pose.LocalTransforms.GetView();
+		TUniquePtr<FKeyframeState> KeyframeA;
+		if (!VM.PopValue(KEYFRAME_STACK_NAME, KeyframeA))
+		{
+			// We have a single input, leave it on top of the stack
+			VM.PushValue(KEYFRAME_STACK_NAME, MoveTemp(KeyframeB));
+			return;
+		}
 
-			BlendOverwriteWithScale(KeyframeBTransformsView, KeyframeB->Pose.LocalTransforms.GetConstView(), WeightOfPoseB);
-			BlendAddWithScale(KeyframeBTransformsView, KeyframeA->Pose.LocalTransforms.GetConstView(), WeightOfPoseA);
+		const float WeightOfPoseB = GetInterpolationAlpha(KeyframeA.Get(), KeyframeB.Get());
 
-			// Ensure that all of the resulting rotations are normalized
-			NormalizeRotations(KeyframeBTransformsView);
+		if (!FAnimWeight::IsRelevant(WeightOfPoseB))
+		{
+			VM.PushValue(KEYFRAME_STACK_NAME, MoveTemp(KeyframeA));
 		}
-
-		if (EnumHasAnyFlags(VM.GetFlags(), EEvaluationFlags::Curves))
+		else if (FAnimWeight::IsFullWeight(WeightOfPoseB))
 		{
-			// Curves cannot blend in place
-			FBlendedCurve Result;
-			Result.Lerp(KeyframeA->Curves, KeyframeB->Curves, WeightOfPoseB);
-
-			KeyframeB->Curves = MoveTemp(Result);
+			VM.PushValue(KEYFRAME_STACK_NAME, MoveTemp(KeyframeB));
 		}
+		else
+		{
+			const float WeightOfPoseA = 1.0f - WeightOfPoseB;
 
-		if (EnumHasAnyFlags(VM.GetFlags(), EEvaluationFlags::Attributes))
-		{			
-			UE::Anim::FStackAttributeContainer OutputAttributes;
-			UE::Anim::Attributes::BlendAttributes({ KeyframeA->Attributes, KeyframeB->Attributes }, { WeightOfPoseA, WeightOfPoseB }, { 0, 1 }, OutputAttributes);			
-			KeyframeB->Attributes.MoveFrom(OutputAttributes);
-		}
+			if (EnumHasAnyFlags(VM.GetFlags(), EEvaluationFlags::Bones))
+			{
+				check(KeyframeA->Pose.GetNumBones() == KeyframeB->Pose.GetNumBones());
 
-		VM.PushValue(KEYFRAME_STACK_NAME, MoveTemp(KeyframeB));
+				const FTransformArrayView KeyframeBTransformsView = KeyframeB->Pose.LocalTransforms.GetView();
+
+				BlendOverwriteWithScale(KeyframeBTransformsView, KeyframeB->Pose.LocalTransforms.GetConstView(), WeightOfPoseB);
+				BlendAddWithScale(KeyframeBTransformsView, KeyframeA->Pose.LocalTransforms.GetConstView(), WeightOfPoseA);
+
+				// Ensure that all of the resulting rotations are normalized
+				NormalizeRotations(KeyframeBTransformsView);
+			}
+
+			if (EnumHasAnyFlags(VM.GetFlags(), EEvaluationFlags::Curves))
+			{
+				// Curves cannot blend in place
+				FBlendedCurve Result;
+				Result.Lerp(KeyframeA->Curves, KeyframeB->Curves, WeightOfPoseB);
+
+				KeyframeB->Curves = MoveTemp(Result);
+			}
+
+			if (EnumHasAnyFlags(VM.GetFlags(), EEvaluationFlags::Attributes))
+			{
+				UE::Anim::FStackAttributeContainer OutputAttributes;
+				UE::Anim::Attributes::BlendAttributes({ KeyframeA->Attributes, KeyframeB->Attributes }, { WeightOfPoseA, WeightOfPoseB }, { 0, 1 }, OutputAttributes);
+				KeyframeB->Attributes.MoveFrom(OutputAttributes);
+			}
+
+			VM.PushValue(KEYFRAME_STACK_NAME, MoveTemp(KeyframeB));
+		}
 	}
 }
 
@@ -115,6 +164,44 @@ float FAnimNextBlendTwoKeyframesTask::GetInterpolationAlpha(const UE::UAF::FKeyf
 	return FMath::Clamp(Alpha, 0.0f, 1.0f);
 }
 
+float FAnimNextBlendTwoKeyframesTask::GetInterpolationAlpha(const UE::UAF::FAttributeCollection& CollectionA, const UE::UAF::FAttributeCollection& CollectionB) const
+{
+	using namespace UE::UAF;
+
+	float Alpha = InterpolationAlpha;
+
+	if (AlphaSourceCurveName != NAME_None && AlphaCurveInputIndex != INDEX_NONE)
+	{
+		const FAttributeMappingKey MappingKey = FAttributeMappingKey::MakeFromTo<FFloatAnimationAttribute>();
+		const FAttributeCollection& AlphaCurveCollection = AlphaCurveInputIndex == 0 ? CollectionA : CollectionB;
+		if (const TAttributeSetMap<FFloatAnimationAttribute>* FloatSetMap = AlphaCurveCollection.FindSetMap<FFloatAnimationAttribute>(MappingKey))
+		{
+			const FAttributeTypedSetPtr& TypedSet = FloatSetMap->GetTypedSet();
+			if (FAttributeSetIndex CurveIndex = TypedSet->FindIndex(AlphaSourceCurveName))
+			{
+				Alpha = (*FloatSetMap)[CurveIndex].Value;
+			}
+			else
+			{
+				// Curve name wasn't found
+				Alpha = 0.0f;
+			}
+		}
+		else
+		{
+			// No float curves found
+			Alpha = 0.0f;
+		}
+
+		if (InputScaleBiasClampFn.IsSet())
+		{
+			Alpha = InputScaleBiasClampFn(Alpha);
+		}
+	}
+
+	return FMath::Clamp(Alpha, 0.0f, 1.0f);
+}
+
 // --- FAnimNextBlendOverwriteKeyframeWithScaleTask ---
 
 FAnimNextBlendOverwriteKeyframeWithScaleTask FAnimNextBlendOverwriteKeyframeWithScaleTask::Make(float ScaleFactor)
@@ -129,33 +216,53 @@ void FAnimNextBlendOverwriteKeyframeWithScaleTask::Execute(UE::UAF::FEvaluationV
 {
 	using namespace UE::UAF;
 
-	TUniquePtr<FKeyframeState> Keyframe;
-	if (!VM.PopValue(KEYFRAME_STACK_NAME, Keyframe))
+	if (VM.GetContext().IsValid())
 	{
-		// We have no inputs, nothing to do
-		return;
-	}
+		// Pop our top two poses, we'll re-use the top keyframe for our result
 
-	if (EnumHasAnyFlags(VM.GetFlags(), EEvaluationFlags::Bones))
-	{
-		BlendOverwriteWithScale(Keyframe->Pose.LocalTransforms.GetView(), Keyframe->Pose.LocalTransforms.GetConstView(), ScaleFactor);
-	}
+		TUniquePtr<FAttributeCollection> Collection;
+		if (!VM.PopValue(ATTRIBUTE_STACK_NAME, Collection))
+		{
+			// We have no inputs, nothing to do
+			return;
+		}
 
-	if (EnumHasAnyFlags(VM.GetFlags(), EEvaluationFlags::Curves))
-	{
-		// Curves cannot override in place
-		FBlendedCurve Result;
-		Result.Override(Keyframe->Curves, ScaleFactor);
+		// Overwrite in place
 
-		Keyframe->Curves = MoveTemp(Result);
-	}
+		Transformers::FOverwrite::Apply(VM.GetTransformerMap(), *Collection, *Collection, ScaleFactor);
 
-	if (EnumHasAnyFlags(VM.GetFlags(), EEvaluationFlags::Attributes))
-	{
-		UE::Anim::Attributes::OverrideAttributes(Keyframe->Attributes, Keyframe->Attributes, ScaleFactor);
+		VM.PushValue(ATTRIBUTE_STACK_NAME, MoveTemp(Collection));
 	}
+	else
+	{
+		TUniquePtr<FKeyframeState> Keyframe;
+		if (!VM.PopValue(KEYFRAME_STACK_NAME, Keyframe))
+		{
+			// We have no inputs, nothing to do
+			return;
+		}
+
+		if (EnumHasAnyFlags(VM.GetFlags(), EEvaluationFlags::Bones))
+		{
+			BlendOverwriteWithScale(Keyframe->Pose.LocalTransforms.GetView(), Keyframe->Pose.LocalTransforms.GetConstView(), ScaleFactor);
+		}
 
-	VM.PushValue(KEYFRAME_STACK_NAME, MoveTemp(Keyframe));
+		if (EnumHasAnyFlags(VM.GetFlags(), EEvaluationFlags::Curves))
+		{
+			// Curves cannot override in place
+			FBlendedCurve Result;
+			Result.Override(Keyframe->Curves, ScaleFactor);
+
+			Keyframe->Curves = MoveTemp(Result);
+		}
+
+		if (EnumHasAnyFlags(VM.GetFlags(), EEvaluationFlags::Attributes))
+		{
+			UE::Anim::Attributes::OverrideAttributes(Keyframe->Attributes, Keyframe->Attributes, ScaleFactor);
+		}
+
+		VM.PushValue(KEYFRAME_STACK_NAME, MoveTemp(Keyframe));
+	}
 }
 
 // --- FAnimNextBlendAddKeyframeWithScaleTask  ---
@@ -172,39 +279,69 @@ void FAnimNextBlendAddKeyframeWithScaleTask::Execute(UE::UAF::FEvaluationVM& VM)
 {
 	using namespace UE::UAF;
 
-	// Pop our top two poses, we'll re-use the top keyframe for our result
-
-	TUniquePtr<FKeyframeState> KeyframeB;
-	if (!VM.PopValue(KEYFRAME_STACK_NAME, KeyframeB))
+	if (VM.GetContext().IsValid())
 	{
-		// We have no inputs, nothing to do
-		return;
-	}
+		// Pop our top two poses, we'll re-use the top keyframe for our result
 
-	TUniquePtr<FKeyframeState> KeyframeA;
-	if (!VM.PopValue(KEYFRAME_STACK_NAME, KeyframeA))
-	{
-		// We have a single input, leave it on top of the stack
-		VM.PushValue(KEYFRAME_STACK_NAME, MoveTemp(KeyframeB));
-		return;
-	}
+		TUniquePtr<FAttributeCollection> CollectionB;
+		if (!VM.PopValue(ATTRIBUTE_STACK_NAME, CollectionB))
+		{
+			// We have no inputs, nothing to do
+			return;
+		}
 
-	if (EnumHasAnyFlags(VM.GetFlags(), EEvaluationFlags::Bones))
-	{
-		check(KeyframeA->Pose.GetNumBones() == KeyframeB->Pose.GetNumBones());
+		TUniquePtr<FAttributeCollection> CollectionA;
+		if (!VM.PopValue(ATTRIBUTE_STACK_NAME, CollectionA))
+		{
+			// We have a single input, leave it on top of the stack
+			VM.PushValue(ATTRIBUTE_STACK_NAME, MoveTemp(CollectionB));
+			return;
+		}
 
-		BlendAddWithScale(KeyframeB->Pose.LocalTransforms.GetView(), KeyframeA->Pose.LocalTransforms.GetConstView(), ScaleFactor);
-	}
+		// Accumulate and place the result within B
+		// Named sets must match as it ensures that our inputs have the same sizes/shapes
+		check(CollectionA->GetNamedSet() == CollectionB->GetNamedSet());
 
-	if (EnumHasAnyFlags(VM.GetFlags(), EEvaluationFlags::Curves))
-	{
-		KeyframeB->Curves.Accumulate(KeyframeA->Curves, ScaleFactor);
-	}
+		Transformers::FAccumulate::Apply(VM.GetTransformerMap(), *CollectionB, *CollectionA, ScaleFactor);
 
-	if (EnumHasAnyFlags(VM.GetFlags(), EEvaluationFlags::Attributes))
-	{
-		UE::Anim::Attributes::AccumulateAttributes(KeyframeA->Attributes, KeyframeB->Attributes, ScaleFactor, AAT_None);
+		VM.PushValue(ATTRIBUTE_STACK_NAME, MoveTemp(CollectionB));
 	}
+	else
+	{
+		// Pop our top two poses, we'll re-use the top keyframe for our result
+
+		TUniquePtr<FKeyframeState> KeyframeB;
+		if (!VM.PopValue(KEYFRAME_STACK_NAME, KeyframeB))
+		{
+			// We have no inputs, nothing to do
+			return;
+		}
+
+		TUniquePtr<FKeyframeState> KeyframeA;
+		if (!VM.PopValue(KEYFRAME_STACK_NAME, KeyframeA))
+		{
+			// We have a single input, leave it on top of the stack
+			VM.PushValue(KEYFRAME_STACK_NAME, MoveTemp(KeyframeB));
+			return;
+		}
+
+		if (EnumHasAnyFlags(VM.GetFlags(), EEvaluationFlags::Bones))
+		{
+			check(KeyframeA->Pose.GetNumBones() == KeyframeB->Pose.GetNumBones());
 
-	VM.PushValue(KEYFRAME_STACK_NAME, MoveTemp(KeyframeB));
+			BlendAddWithScale(KeyframeB->Pose.LocalTransforms.GetView(), KeyframeA->Pose.LocalTransforms.GetConstView(), ScaleFactor);
+		}
+
+		if (EnumHasAnyFlags(VM.GetFlags(), EEvaluationFlags::Curves))
+		{
+			KeyframeB->Curves.Accumulate(KeyframeA->Curves, ScaleFactor);
+		}
+
+		if (EnumHasAnyFlags(VM.GetFlags(), EEvaluationFlags::Attributes))
+		{
+			UE::Anim::Attributes::AccumulateAttributes(KeyframeA->Attributes, KeyframeB->Attributes, ScaleFactor, AAT_None);
+		}
+
+		VM.PushValue(KEYFRAME_STACK_NAME, MoveTemp(KeyframeB));
+	}
 }
diff --git a/Engine/Plugins/Experimental/UAF/UAFAnimGraph/Source/UAFAnimGraph/Private/EvaluationVM/Tasks/BlendKeyframesPerBone.cpp b/Engine/Plugins/Experimental/UAF/UAFAnimGraph/Source/UAFAnimGraph/Private/EvaluationVM/Tasks/BlendKeyframesPerBone.cpp
--- a/Engine/Plugins/Experimental/UAF/UAFAnimGraph/Source/UAFAnimGraph/Private/EvaluationVM/Tasks/BlendKeyframesPerBone.cpp
+++ b/Engine/Plugins/Experimental/UAF/UAFAnimGraph/Source/UAFAnimGraph/Private/EvaluationVM/Tasks/BlendKeyframesPerBone.cpp
@@ -29,6 +29,13 @@ void FAnimNextBlendOverwriteKeyframePerBoneWithScaleTask::Execute(UE::UAF::FEval
 {
 	using namespace UE::UAF;
 
+	if (VM.GetContext().IsValid())
+	{
+		// TODO: Implement with new attribute runtime
+		Super::Execute(VM);
+		return;
+	}
+
 	if (BlendProfile == nullptr)
 	{
 		// If we don't have a blend profile, we blend the whole pose
@@ -115,6 +122,13 @@ void FAnimNextBlendAddKeyframePerBoneWithScaleTask::Execute(UE::UAF::FEvaluation
 {
 	using namespace UE::UAF;
 
+	if (VM.GetContext().IsValid())
+	{
+		// TODO: Implement with new attribute runtime
+		Super::Execute(VM);
+		return;
+	}
+
 	if (BlendProfile == nullptr)
 	{
 		// If we don't have a blend profile, we blend the whole pose
@@ -218,6 +232,13 @@ void FAnimNextBlendKeyframePerBoneWithScaleTask::Execute(UE::UAF::FEvaluationVM&
 {
 	using namespace UE::UAF;
 
+	if (VM.GetContext().IsValid())
+	{
+		// TODO: Implement with new attribute runtime
+		Super::Execute(VM);
+		return;
+	}
+
 	if (BlendProfile == nullptr)
 	{
 		// If we don't have a blend profile, we blend the whole pose
diff --git a/Engine/Plugins/Experimental/UAF/UAFAnimGraph/Source/UAFAnimGraph/Private/EvaluationVM/Tasks/ConvertRotationsToLocalSpace.cpp b/Engine/Plugins/Experimental/UAF/UAFAnimGraph/Source/UAFAnimGraph/Private/EvaluationVM/Tasks/ConvertRotationsToLocalSpace.cpp
--- a/Engine/Plugins/Experimental/UAF/UAFAnimGraph/Source/UAFAnimGraph/Private/EvaluationVM/Tasks/ConvertRotationsToLocalSpace.cpp
+++ b/Engine/Plugins/Experimental/UAF/UAFAnimGraph/Source/UAFAnimGraph/Private/EvaluationVM/Tasks/ConvertRotationsToLocalSpace.cpp
@@ -23,6 +23,12 @@ void FAnimNextConvertRotationsToLocalSpaceTask::Execute(UE::UAF::FEvaluationVM&
 {
 	using namespace UE::UAF;
 
+	if (VM.GetContext().IsValid())
+	{
+		// TODO: Implement with new attribute runtime
+		return;
+	}
+
 	if (!EnumHasAnyFlags(VM.GetFlags(), EEvaluationFlags::Bones))
 	{
 		return;
diff --git a/Engine/Plugins/Experimental/UAF/UAFAnimGraph/Source/UAFAnimGraph/Private/EvaluationVM/Tasks/DeadBlending.cpp b/Engine/Plugins/Experimental/UAF/UAFAnimGraph/Source/UAFAnimGraph/Private/EvaluationVM/Tasks/DeadBlending.cpp
--- a/Engine/Plugins/Experimental/UAF/UAFAnimGraph/Source/UAFAnimGraph/Private/EvaluationVM/Tasks/DeadBlending.cpp
+++ b/Engine/Plugins/Experimental/UAF/UAFAnimGraph/Source/UAFAnimGraph/Private/EvaluationVM/Tasks/DeadBlending.cpp
@@ -484,6 +484,12 @@ void FAnimNextDeadBlendingTransitionTask::Execute(UE::UAF::FEvaluationVM& VM) co
 {
 	using namespace UE::UAF;
 
+	if (VM.GetContext().IsValid())
+	{
+		// TODO: Implement with new attribute runtime
+		return;
+	}
+
 	if (EnumHasAnyFlags(VM.GetFlags(), EEvaluationFlags::Bones))
 	{
 		if (const TUniquePtr<FKeyframeState>* Keyframe = VM.PeekValue<TUniquePtr<FKeyframeState>>(KEYFRAME_STACK_NAME, 0))
@@ -553,6 +559,12 @@ void FAnimNextDeadBlendingApplyTask::Execute(UE::UAF::FEvaluationVM& VM) const
 {
 	using namespace UE::UAF;
 
+	if (VM.GetContext().IsValid())
+	{
+		// TODO: Implement with new attribute runtime
+		return;
+	}
+
 	if (EnumHasAnyFlags(VM.GetFlags(), EEvaluationFlags::Bones))
 	{
 		if (const TUniquePtr<FKeyframeState>* Keyframe = VM.PeekValue<TUniquePtr<FKeyframeState>>(KEYFRAME_STACK_NAME, 0))
diff --git a/Engine/Plugins/Experimental/UAF/UAFAnimGraph/Source/UAFAnimGraph/Private/EvaluationVM/Tasks/NormalizeRotations.cpp b/Engine/Plugins/Experimental/UAF/UAFAnimGraph/Source/UAFAnimGraph/Private/EvaluationVM/Tasks/NormalizeRotations.cpp
--- a/Engine/Plugins/Experimental/UAF/UAFAnimGraph/Source/UAFAnimGraph/Private/EvaluationVM/Tasks/NormalizeRotations.cpp
+++ b/Engine/Plugins/Experimental/UAF/UAFAnimGraph/Source/UAFAnimGraph/Private/EvaluationVM/Tasks/NormalizeRotations.cpp
@@ -5,6 +5,8 @@
 #include "EvaluationVM/EvaluationVM.h"
 #include "EvaluationVM/KeyframeState.h"
 #include "TransformArrayOperations.h"
+#include "UAF/Attributes/AttributeCollection.h"
+#include "UAF/Attributes/Transformers/Sanitize.h"
 
 #include UE_INLINE_GENERATED_CPP_BY_NAME(NormalizeRotations)
 
@@ -12,17 +14,33 @@ void FAnimNextNormalizeKeyframeRotationsTask::Execute(UE::UAF::FEvaluationVM& VM
 {
 	using namespace UE::UAF;
 
-	if (EnumHasAnyFlags(VM.GetFlags(), EEvaluationFlags::Bones))
+	if (VM.GetContext().IsValid())
 	{
-		TUniquePtr<FKeyframeState> Keyframe;
-		if (!VM.PopValue(KEYFRAME_STACK_NAME, Keyframe))
+		TUniquePtr<FAttributeCollection> Collection;
+		if (!VM.PopValue(ATTRIBUTE_STACK_NAME, Collection))
 		{
 			// We have no inputs, nothing to do
 			return;
 		}
 
-		NormalizeRotations(Keyframe->Pose.LocalTransforms.GetView());
+		Transformers::FSanitize::Apply(VM.GetTransformerMap(), *Collection);
 
-		VM.PushValue(KEYFRAME_STACK_NAME, MoveTemp(Keyframe));
+		VM.PushValue(ATTRIBUTE_STACK_NAME, MoveTemp(Collection));
+	}
+	else
+	{
+		if (EnumHasAnyFlags(VM.GetFlags(), EEvaluationFlags::Bones))
+		{
+			TUniquePtr<FKeyframeState> Keyframe;
+			if (!VM.PopValue(KEYFRAME_STACK_NAME, Keyframe))
+			{
+				// We have no inputs, nothing to do
+				return;
+			}
+
+			NormalizeRotations(Keyframe->Pose.LocalTransforms.GetView());
+
+			VM.PushValue(KEYFRAME_STACK_NAME, MoveTemp(Keyframe));
+		}
 	}
 }
diff --git a/Engine/Plugins/Experimental/UAF/UAFAnimGraph/Source/UAFAnimGraph/Private/EvaluationVM/Tasks/PushAnimSequenceKeyframe.cpp b/Engine/Plugins/Experimental/UAF/UAFAnimGraph/Source/UAFAnimGraph/Private/EvaluationVM/Tasks/PushAnimSequenceKeyframe.cpp
--- a/Engine/Plugins/Experimental/UAF/UAFAnimGraph/Source/UAFAnimGraph/Private/EvaluationVM/Tasks/PushAnimSequenceKeyframe.cpp
+++ b/Engine/Plugins/Experimental/UAF/UAFAnimGraph/Source/UAFAnimGraph/Private/EvaluationVM/Tasks/PushAnimSequenceKeyframe.cpp
@@ -36,51 +36,90 @@ void FAnimNextAnimSequenceKeyframeTask::Execute(UE::UAF::FEvaluationVM& VM) cons
 	
 	if(const UAnimSequence* AnimSequencePtr = AnimSequence.Get())
 	{
-		const bool bIsAdditive = AnimSequencePtr->IsValidAdditive();
-
 		const bool bExtractRootMotion = bExtractTrajectory;
 		const FAnimExtractContext ExtractionContext(SampleTime, bExtractRootMotion, DeltaTimeRecord, bLooping);
+		const FAttributeEvaluationContext& VMContext = VM.GetContext();
 
-		FKeyframeState Keyframe = VM.MakeUninitializedKeyframe(bIsAdditive);
-
-		const bool bUseRawData = FDecompressionTools::ShouldUseRawData(AnimSequencePtr, Keyframe.Pose);
+		const bool bIsAdditive = AnimSequencePtr->IsValidAdditive();
+		const bool bUseRawData = FDecompressionTools::ShouldUseRawData(AnimSequencePtr, ExtractionContext);
 
-		if (EnumHasAnyFlags(VM.GetFlags(), EEvaluationFlags::Bones))
+		if (VMContext.IsValid())
 		{
-			FDecompressionTools::GetAnimationPose(AnimSequencePtr, ExtractionContext, Keyframe.Pose, bUseRawData);
-		}
+			constexpr bool bInitializeWithDefaultValues = true;
+			FAttributeCollectionStack Collection = VM.MakeAttributeCollection(bIsAdditive, bInitializeWithDefaultValues);
 
-		if (EnumHasAnyFlags(VM.GetFlags(), EEvaluationFlags::Curves))
-		{
-			FDecompressionTools::GetAnimationCurves(AnimSequencePtr, ExtractionContext, Keyframe.Curves, bUseRawData);
-		}
+			FDecompressionTools::GetAnimationPose(AnimSequencePtr, ExtractionContext, VMContext, Collection, bUseRawData);
+			FDecompressionTools::GetAnimationCurves(AnimSequencePtr, ExtractionContext, VMContext, Collection, bUseRawData);
+			FDecompressionTools::GetAnimationAttributes(AnimSequencePtr, ExtractionContext, VMContext, Collection, bUseRawData);
 
-		if (EnumHasAnyFlags(VM.GetFlags(), EEvaluationFlags::Attributes))
-		{
-			FDecompressionTools::GetAnimationAttributes(AnimSequencePtr, ExtractionContext, Keyframe.Pose.GetRefPose(), Keyframe.Attributes, bUseRawData);
-		}
-
-		// Trajectory is currently held as an attribute
-		if (EnumHasAnyFlags(VM.GetFlags(), EEvaluationFlags::Attributes | EEvaluationFlags::Trajectory))
-		{
-			// If the sequence has root motion enabled, allow sampling of a root motion delta into the custom attribute container of the outgoing pose
+			// If the sequence has root motion enabled, extract it into its own attribute
 			if (AnimSequence->HasRootMotion())
 			{
-				// TODO: We should cache the provider in the VM
-				// We have to grab two locks to get it and it won't change during graph evaluation
-				if (const UE::Anim::IAnimRootMotionProvider* RootMotionProvider = UE::Anim::IAnimRootMotionProvider::Get())
+				FAttributeNamedSetPtr NamedSet = Collection.GetNamedSet();
+				if (FAttributeTypedSetPtr AttributeTypedSet = NamedSet->FindTypedSet<FTransformAnimationAttribute>())
 				{
-					RootMotionProvider->SampleRootMotion(ExtractionContext.DeltaTimeRecord, *AnimSequence, ExtractionContext.bLooping, Keyframe.Attributes);
+					// AttributeName: RootMotionDelta
+					if (const FAttributeSetIndex SetIndex = AttributeTypedSet->FindIndex(UE::Anim::IAnimRootMotionProvider::AttributeName))
+					{
+						// This attribute is present within the current named set
+
+						FAttributeMappingKey MappingKey = FAttributeMappingKey::MakeFromTo<FTransformAnimationAttribute>();
+						TAttributeSetMap<FTransformAnimationAttribute>* SetMap = Collection.FindSetMap<FTransformAnimationAttribute>(MappingKey);
+						checkf(SetMap != nullptr, TEXT("This attribute exists within our named set but isn't present in the output collection"));
+
+						const FTransform RootMotionTransform = AnimSequence->ExtractRootMotion(ExtractionContext);
+						(*SetMap)[SetIndex].Value = RootMotionTransform;
+					}
 				}
 			}
+
+			VM.PushValue(ATTRIBUTE_STACK_NAME, MakeUnique<FAttributeCollection>(MoveTemp(Collection)));
 		}
+		else
+		{
+			FKeyframeState Keyframe = VM.MakeUninitializedKeyframe(bIsAdditive);
 
-		VM.PushValue(KEYFRAME_STACK_NAME, MakeUnique<FKeyframeState>(MoveTemp(Keyframe)));
+			if (EnumHasAnyFlags(VM.GetFlags(), EEvaluationFlags::Bones))
+			{
+				FDecompressionTools::GetAnimationPose(AnimSequencePtr, ExtractionContext, Keyframe.Pose, bUseRawData);
+			}
+
+			if (EnumHasAnyFlags(VM.GetFlags(), EEvaluationFlags::Curves))
+			{
+				FDecompressionTools::GetAnimationCurves(AnimSequencePtr, ExtractionContext, Keyframe.Curves, bUseRawData);
+			}
+
+			if (EnumHasAnyFlags(VM.GetFlags(), EEvaluationFlags::Attributes))
+			{
+				FDecompressionTools::GetAnimationAttributes(AnimSequencePtr, ExtractionContext, Keyframe.Pose.GetRefPose(), Keyframe.Attributes, bUseRawData);
+			}
+
+			// Trajectory is currently held as an attribute
+			if (EnumHasAnyFlags(VM.GetFlags(), EEvaluationFlags::Attributes | EEvaluationFlags::Trajectory))
+			{
+				// If the sequence has root motion enabled, allow sampling of a root motion delta into the custom attribute container of the outgoing pose
+				if (AnimSequence->HasRootMotion())
+				{
+					// TODO: We should cache the provider in the VM
+					// We have to grab two locks to get it and it won't change during graph evaluation
+					if (const UE::Anim::IAnimRootMotionProvider* RootMotionProvider = UE::Anim::IAnimRootMotionProvider::Get())
+					{
+						RootMotionProvider->SampleRootMotion(ExtractionContext.DeltaTimeRecord, *AnimSequence, ExtractionContext.bLooping, Keyframe.Attributes);
+					}
+				}
+			}
+
+			VM.PushValue(KEYFRAME_STACK_NAME, MakeUnique<FKeyframeState>(MoveTemp(Keyframe)));
+		}
 	}
 	else
 	{
 		constexpr bool bIsAdditive = false;
 		FKeyframeState Keyframe = VM.MakeReferenceKeyframe(bIsAdditive);
 		VM.PushValue(KEYFRAME_STACK_NAME, MakeUnique<FKeyframeState>(MoveTemp(Keyframe)));
+
+		constexpr bool bInitializeWithDefaultValues = true;
+		FAttributeCollectionStack Collection = VM.MakeAttributeCollection(bIsAdditive, bInitializeWithDefaultValues);
+		VM.PushValue(ATTRIBUTE_STACK_NAME, MakeUnique<FAttributeCollection>(MoveTemp(Collection)));
 	}
 }
diff --git a/Engine/Plugins/Experimental/UAF/UAFAnimGraph/Source/UAFAnimGraph/Private/EvaluationVM/Tasks/PushPose.cpp b/Engine/Plugins/Experimental/UAF/UAFAnimGraph/Source/UAFAnimGraph/Private/EvaluationVM/Tasks/PushPose.cpp
--- a/Engine/Plugins/Experimental/UAF/UAFAnimGraph/Source/UAFAnimGraph/Private/EvaluationVM/Tasks/PushPose.cpp
+++ b/Engine/Plugins/Experimental/UAF/UAFAnimGraph/Source/UAFAnimGraph/Private/EvaluationVM/Tasks/PushPose.cpp
@@ -25,9 +25,18 @@ FAnimNextPushPoseTask& FAnimNextPushPoseTask::operator=(const FAnimNextPushPoseT
 
 void FAnimNextPushPoseTask::Execute(UE::UAF::FEvaluationVM& VM) const
 {
+	using namespace UE::UAF;
+
 	SCOPE_CYCLE_COUNTER(STAT_AnimNext_Task_PushPose);
 
-	using namespace UE::UAF;
+	if (VM.GetContext().IsValid())
+	{
+		// TODO: Implement with new attribute runtime
+		constexpr bool bIsAdditive = false;
+		constexpr bool bInitializeWithDefaultValues = true;
+		VM.PushValue(ATTRIBUTE_STACK_NAME, MakeUnique<FAttributeCollection>(VM.MakeAttributeCollection(bIsAdditive, bInitializeWithDefaultValues)));
+		return;
+	}
 
 	FKeyframeState Keyframe = VM.MakeUninitializedKeyframe(GraphPose->LODPose.IsAdditive());
 
diff --git a/Engine/Plugins/Experimental/UAF/UAFAnimGraph/Source/UAFAnimGraph/Private/EvaluationVM/Tasks/PushReferenceKeyframe.cpp b/Engine/Plugins/Experimental/UAF/UAFAnimGraph/Source/UAFAnimGraph/Private/EvaluationVM/Tasks/PushReferenceKeyframe.cpp
--- a/Engine/Plugins/Experimental/UAF/UAFAnimGraph/Source/UAFAnimGraph/Private/EvaluationVM/Tasks/PushReferenceKeyframe.cpp
+++ b/Engine/Plugins/Experimental/UAF/UAFAnimGraph/Source/UAFAnimGraph/Private/EvaluationVM/Tasks/PushReferenceKeyframe.cpp
@@ -23,5 +23,13 @@ void FAnimNextPushReferenceKeyframeTask::Execute(UE::UAF::FEvaluationVM& VM) con
 {
 	using namespace UE::UAF;
 
-	VM.PushValue(KEYFRAME_STACK_NAME, MakeUnique<FKeyframeState>(VM.MakeReferenceKeyframe(bIsAdditive)));
+	if (VM.GetContext().IsValid())
+	{
+		constexpr bool bInitializeWithDefaultValues = true;
+		VM.PushValue(ATTRIBUTE_STACK_NAME, MakeUnique<FAttributeCollection>(VM.MakeAttributeCollection(bIsAdditive, bInitializeWithDefaultValues)));
+	}
+	else
+	{
+		VM.PushValue(KEYFRAME_STACK_NAME, MakeUnique<FKeyframeState>(VM.MakeReferenceKeyframe(bIsAdditive)));
+	}
 }
diff --git a/Engine/Plugins/Experimental/UAF/UAFAnimGraph/Source/UAFAnimGraph/Private/EvaluationVM/Tasks/StoreKeyframe.cpp b/Engine/Plugins/Experimental/UAF/UAFAnimGraph/Source/UAFAnimGraph/Private/EvaluationVM/Tasks/StoreKeyframe.cpp
--- a/Engine/Plugins/Experimental/UAF/UAFAnimGraph/Source/UAFAnimGraph/Private/EvaluationVM/Tasks/StoreKeyframe.cpp
+++ b/Engine/Plugins/Experimental/UAF/UAFAnimGraph/Source/UAFAnimGraph/Private/EvaluationVM/Tasks/StoreKeyframe.cpp
@@ -32,6 +32,12 @@ void FAnimNextStoreKeyframeTransformsTask::Execute(UE::UAF::FEvaluationVM& VM) c
 {
 	using namespace UE::UAF;
 
+	if (VM.GetContext().IsValid())
+	{
+		// TODO: Implement with new attribute runtime
+		return;
+	}
+
 	if (EnumHasAnyFlags(VM.GetFlags(), EEvaluationFlags::Bones))
 	{
 		if (const TUniquePtr<FKeyframeState>* Keyframe = VM.PeekValue<TUniquePtr<FKeyframeState>>(KEYFRAME_STACK_NAME, 0))
@@ -55,6 +61,15 @@ void FAnimNextDuplicateTopKeyframeTask::Execute(UE::UAF::FEvaluationVM& VM) cons
 {
 	using namespace UE::UAF;
 
+	if (VM.GetContext().IsValid())
+	{
+		// TODO: Implement with new attribute runtime
+		constexpr bool bIsAdditive = false;
+		constexpr bool bInitializeWithDefaultValues = true;
+		VM.PushValue(ATTRIBUTE_STACK_NAME, MakeUnique<FAttributeCollection>(VM.MakeAttributeCollection(bIsAdditive, bInitializeWithDefaultValues)));
+		return;
+	}
+
 	TUniquePtr<FKeyframeState> KeyframeOut;
 
 	// Try to get a Keyframe from the stack without consuming it
diff --git a/Engine/Plugins/Experimental/UAF/UAFAnimGraph/Source/UAFAnimGraph/Private/Graph/AnimGraph/AnimNode_AnimNextGraph.cpp b/Engine/Plugins/Experimental/UAF/UAFAnimGraph/Source/UAFAnimGraph/Private/Graph/AnimGraph/AnimNode_AnimNextGraph.cpp
--- a/Engine/Plugins/Experimental/UAF/UAFAnimGraph/Source/UAFAnimGraph/Private/Graph/AnimGraph/AnimNode_AnimNextGraph.cpp
+++ b/Engine/Plugins/Experimental/UAF/UAFAnimGraph/Source/UAFAnimGraph/Private/Graph/AnimGraph/AnimNode_AnimNextGraph.cpp
@@ -15,6 +15,7 @@
 #include "TraitInterfaces/IUpdate.h"
 #include "EvaluationVM/EvaluationVM.h"
 #include "Graph/AnimNextGraphInstance.h"
+#include "UAF/UAFSkeletonUserData.h"
 
 #include UE_INLINE_GENERATED_CPP_BY_NAME(AnimNode_AnimNextGraph)
 
@@ -126,6 +127,19 @@ void FAnimNode_AnimNextGraph::Evaluate_AnyThread(FPoseContext& Output)
 		{
 			UE::UAF::FEvaluateGraphContext EvaluateGraphContext(*GraphInstance.Get(), RefPose, LODLevel);
 			EvaluateGraphContext.SetBindingObject(SkeletalMeshComponent);
+
+			if (UE::UAF::UseAbstractSkeletonRuntime())
+			{
+				// By default, use the everything set
+				USkeleton* Skeleton = const_cast<USkeleton*>(RefPose.GetSkeletonAsset());
+				UUAFSkeletonUserData* UserData = UUAFSkeletonUserData::FromSkeleton(*Skeleton);
+
+				UAbstractSkeletonSetBinding* Binding = UserData != nullptr ? UserData->GetSetBinding() : nullptr;
+				FName SetName;
+
+				EvaluateGraphContext.SetSkeletonSet(Binding, SetName);
+			}
+
 			const FEvaluationProgram EvaluationProgram = UE::UAF::EvaluateGraph(EvaluateGraphContext);
 
 			FEvaluationVM EvaluationVM(EEvaluationFlags::All, RefPose, LODLevel);
diff --git a/Engine/Plugins/Experimental/UAF/UAFAnimGraph/Source/UAFAnimGraph/Private/Graph/RigUnit_AnimNextRunAnimationGraph_v1.cpp b/Engine/Plugins/Experimental/UAF/UAFAnimGraph/Source/UAFAnimGraph/Private/Graph/RigUnit_AnimNextRunAnimationGraph_v1.cpp
--- a/Engine/Plugins/Experimental/UAF/UAFAnimGraph/Source/UAFAnimGraph/Private/Graph/RigUnit_AnimNextRunAnimationGraph_v1.cpp
+++ b/Engine/Plugins/Experimental/UAF/UAFAnimGraph/Source/UAFAnimGraph/Private/Graph/RigUnit_AnimNextRunAnimationGraph_v1.cpp
@@ -11,6 +11,7 @@
 #include "Graph/AnimNextGraphInstance.h"
 #include "Graph/RigVMTrait_AnimNextPublicVariables.h"
 #include "TraitCore/TraitEventList.h"
+#include "UAF/UAFSkeletonUserData.h"
 
 #include UE_INLINE_GENERATED_CPP_BY_NAME(RigUnit_AnimNextRunAnimationGraph_v1)
 
@@ -99,6 +100,19 @@ FRigUnit_AnimNextRunAnimationGraph_v1_Execute()
 	
 	FEvaluateGraphContext EvaluateGraphContext(InstanceRef.Get(), RefPose, DesiredLOD);
 	EvaluateGraphContext.SetBindingObject(RefPose.SkeletalMeshComponent);
+
+	if (UseAbstractSkeletonRuntime())
+	{
+		// By default, use the everything set
+		USkeleton* Skeleton = const_cast<USkeleton*>(RefPose.GetSkeletonAsset());
+		UUAFSkeletonUserData* UserData = UUAFSkeletonUserData::FromSkeleton(*Skeleton);
+
+		UAbstractSkeletonSetBinding* Binding = UserData != nullptr ? UserData->GetSetBinding() : nullptr;
+		FName SetName;
+
+		EvaluateGraphContext.SetSkeletonSet(Binding, SetName);
+	}
+
 	EvaluateGraph(EvaluateGraphContext, Result);
 
 	// We might have appended new input/output events, append them
diff --git a/Engine/Plugins/Experimental/UAF/UAFAnimGraph/Source/UAFAnimGraph/Private/Graph/RigUnit_AnimNextRunAnimationGraph_v2.cpp b/Engine/Plugins/Experimental/UAF/UAFAnimGraph/Source/UAFAnimGraph/Private/Graph/RigUnit_AnimNextRunAnimationGraph_v2.cpp
--- a/Engine/Plugins/Experimental/UAF/UAFAnimGraph/Source/UAFAnimGraph/Private/Graph/RigUnit_AnimNextRunAnimationGraph_v2.cpp
+++ b/Engine/Plugins/Experimental/UAF/UAFAnimGraph/Source/UAFAnimGraph/Private/Graph/RigUnit_AnimNextRunAnimationGraph_v2.cpp
@@ -16,6 +16,7 @@
 #include "Injection/AnimNextModuleInjectionComponent.h"
 #include "TraceAnimNextGraphInstances.h"
 #include "Factory/AnimGraphFactory.h"
+#include "UAF/UAFSkeletonUserData.h"
 
 #include UE_INLINE_GENERATED_CPP_BY_NAME(RigUnit_AnimNextRunAnimationGraph_v2)
 
@@ -123,6 +124,19 @@ FRigUnit_AnimNextRunAnimationGraph_v2_Execute()
 
 	FEvaluateGraphContext EvaluateGraphContext(HostInstance, RefPose, DesiredLOD);
 	EvaluateGraphContext.SetBindingObject(RefPose.SkeletalMeshComponent);
+
+	if (UseAbstractSkeletonRuntime())
+	{
+		// By default, use the everything set
+		USkeleton* Skeleton = const_cast<USkeleton*>(RefPose.GetSkeletonAsset());
+		UUAFSkeletonUserData* UserData = UUAFSkeletonUserData::FromSkeleton(*Skeleton);
+
+		UAbstractSkeletonSetBinding* Binding = UserData != nullptr ? UserData->GetSetBinding() : nullptr;
+		FName SetName;
+
+		EvaluateGraphContext.SetSkeletonSet(Binding, SetName);
+	}
+
 	EvaluateGraph(EvaluateGraphContext, Result);
 
 	TRACE_UAF_MODULE(ModuleInstance);
diff --git a/Engine/Plugins/Experimental/UAF/UAFAnimGraph/Source/UAFAnimGraph/Private/TraitInterfaces/IEvaluate.cpp b/Engine/Plugins/Experimental/UAF/UAFAnimGraph/Source/UAFAnimGraph/Private/TraitInterfaces/IEvaluate.cpp
--- a/Engine/Plugins/Experimental/UAF/UAFAnimGraph/Source/UAFAnimGraph/Private/TraitInterfaces/IEvaluate.cpp
+++ b/Engine/Plugins/Experimental/UAF/UAFAnimGraph/Source/UAFAnimGraph/Private/TraitInterfaces/IEvaluate.cpp
@@ -11,8 +11,17 @@
 #include "Graph/AnimNext_LODPose.h"
 #include "Graph/TraceAnimNextGraphInstances.h"
 
+#include "UAF/Attributes/AttributeCollection.h"
+#include "UAF/Attributes/EngineAttributes.h"
+
 namespace UE::UAF
 {
+	namespace Private
+	{
+		static bool GUseAbstractSkeletonRuntime = false;
+		FAutoConsoleVariableRef CVar_UseAbstractSkeletonRuntime(TEXT("UAF.UseAbstractSkeletonRuntime"), GUseAbstractSkeletonRuntime, TEXT("Whether or not to use the abstract skeleton runtime"));
+	}
+
 	FEvaluateTraversalContext::FEvaluateTraversalContext(FEvaluationProgram& InEvaluationProgram)
 		: EvaluationProgram(InEvaluationProgram)
 	{
@@ -235,15 +244,57 @@ namespace UE::UAF
 		{
 			SCOPED_NAMED_EVENT(UAF_EvaluatePose, FColor::Orange);
 
-			FEvaluationVM EvaluationVM(EEvaluationFlags::All, RefPose, GraphLODLevel);
+			UAbstractSkeletonSetBinding* Binding = EvaluateGraphContext.GetSkeletonSetBinding();
+			FName SetName = EvaluateGraphContext.GetSkeletonSetName();
+
+			FEvaluationVM EvaluationVM(EEvaluationFlags::All, RefPose, Binding, SetName, GraphLODLevel);
 			bool bHasValidOutput = false;
 
 			if (!EvaluationProgram.IsEmpty())
 			{
 				EvaluationProgram.Execute(EvaluationVM);
 
+				TUniquePtr<FAttributeCollection> Collection;
+				if (Private::GUseAbstractSkeletonRuntime && EvaluationVM.PopValue(ATTRIBUTE_STACK_NAME, Collection))
+				{
+					const FAttributeMappingKey MappingKey = FAttributeMappingKey::MakeFromTo<FBoneTransformAnimationAttribute>();
+					if (TAttributeSetMap<FBoneTransformAnimationAttribute>* BoneTransforms = Collection->FindSetMap<FBoneTransformAnimationAttribute>(MappingKey))
+					{
+						const FAttributeTypedSetPtr& BoneSet = BoneTransforms->GetTypedSet();
+						const int32 NumPoseLODBones = OutputPose.LODPose.GetNumBones();
+
+						for (int32 LODBoneIndex = 0; LODBoneIndex < NumPoseLODBones; ++LODBoneIndex)
+						{
+							const int32 SkeletonBoneIndex = RefPose.GetSkeletonBoneIndexFromLODBoneIndex(LODBoneIndex);
+							const FAttributeBindingIndex BoneBindingIndex = FAttributeBindingIndex(FSkeletonPoseBoneIndex(SkeletonBoneIndex));
+							check(BoneBindingIndex.IsValid());
+
+							const FAttributeSetIndex BoneSetIndex = BoneSet->GetIndex(BoneBindingIndex);
+							if (BoneSetIndex.IsValid())
+							{
+								OutputPose.LODPose.LocalTransformsView[LODBoneIndex] = (*BoneTransforms)[BoneSetIndex].Value;
+							}
+							else
+							{
+								// This bone isn't present with the current named set, use the bind pose value instead
+								OutputPose.LODPose.LocalTransformsView[LODBoneIndex] = RefPose.GetRefPoseTransform(LODBoneIndex);
+							}
+						}
+					}
+					else
+					{
+						// Nothing was produced, populate with bind pose
+						FKeyframeState ReferenceKeyframe = EvaluationVM.MakeReferenceKeyframe(false);
+						OutputPose.LODPose.CopyFrom(ReferenceKeyframe.Pose);
+						OutputPose.Curves.CopyFrom(ReferenceKeyframe.Curves);
+						OutputPose.Attributes.CopyFrom(ReferenceKeyframe.Attributes);
+					}
+
+					bHasValidOutput = true;
+				}
+
 				TUniquePtr<FKeyframeState> EvaluatedKeyframe;
-				if (EvaluationVM.PopValue(KEYFRAME_STACK_NAME, EvaluatedKeyframe))
+				if (!bHasValidOutput && EvaluationVM.PopValue(KEYFRAME_STACK_NAME, EvaluatedKeyframe))
 				{
 					OutputPose.LODPose.CopyFrom(EvaluatedKeyframe->Pose);
 					OutputPose.Curves.CopyFrom(EvaluatedKeyframe->Curves);
@@ -262,4 +313,9 @@ namespace UE::UAF
 			}
 		}
 	}
+
+	bool UseAbstractSkeletonRuntime()
+	{
+		return Private::GUseAbstractSkeletonRuntime;
+	}
 }
diff --git a/Engine/Plugins/Experimental/UAF/UAFAnimGraph/Source/UAFAnimGraph/Private/Traits/ModifyCurveTrait.cpp b/Engine/Plugins/Experimental/UAF/UAFAnimGraph/Source/UAFAnimGraph/Private/Traits/ModifyCurveTrait.cpp
--- a/Engine/Plugins/Experimental/UAF/UAFAnimGraph/Source/UAFAnimGraph/Private/Traits/ModifyCurveTrait.cpp
+++ b/Engine/Plugins/Experimental/UAF/UAFAnimGraph/Source/UAFAnimGraph/Private/Traits/ModifyCurveTrait.cpp
@@ -50,6 +50,12 @@ FModifyCurveTask FModifyCurveTask::Make(UE::UAF::FModifyCurveTrait::FInstanceDat
 void FModifyCurveTask::Execute(UE::UAF::FEvaluationVM& VM) const
 {
 	using namespace UE::UAF;
+
+	if (VM.GetContext().IsValid())
+	{
+		// TODO: Implement with new attribute runtime
+		return;
+	}
 	
 	if (const TUniquePtr<FKeyframeState>* Keyframe = VM.PeekValue<TUniquePtr<FKeyframeState>>(KEYFRAME_STACK_NAME, 0))
 	{
