Demote ensure to a log as OrientAndScale isn't supported in UAF yet

diff --git a/Engine/Plugins/Experimental/UAF/UAF/Source/UAF/Private/DecompressionTools.cpp b/Engine/Plugins/Experimental/UAF/UAF/Source/UAF/Private/DecompressionTools.cpp
--- a/Engine/Plugins/Experimental/UAF/UAF/Source/UAF/Private/DecompressionTools.cpp
+++ b/Engine/Plugins/Experimental/UAF/UAF/Source/UAF/Private/DecompressionTools.cpp
@@ -14,7 +14,7 @@
 #include "Animation/AnimRootMotionProvider.h"
 #include "Animation/AnimSequence.h"
 #include "HAL/ConsoleManager.h"
-
+#include "Logging/StructuredLog.h"
 #include "UAF/Attributes/AttributeCollection.h"
 #include "UAF/Attributes/AttributeEvaluationContext.h"
 #include "UAF/Attributes/EngineAttributes.h"
@@ -402,7 +402,10 @@ namespace Private
 		//		}
 		//	}
 		//}
-		ensureMsgf(OrientAndScaleRetargetingPairs.IsEmpty(), TEXT("OrientAndScaleRetargetingPairs found for %s (%s -> %s), Orient and Scale retargeting is not yet supported in AnimNext"), *DecompressionContext.AnimName.ToString(), SourceSkeleton ? *SourceSkeleton->GetName() : TEXT("None"), TargetSkeleton ? *TargetSkeleton->GetName() : TEXT("None"));
+		UE_CLOGFMT(OrientAndScaleRetargetingPairs.Num() == 0,
+			LogAnimation, Verbose,
+			"OrientAndScaleRetargetingPairs found for {AnimSequence} ({SourceSkeleton} -> {TargetSkeleton}), Orient and Scale retargeting is not yet supported in UAF",
+			*DecompressionContext.AnimName.ToString(), SourceSkeleton ? *SourceSkeleton->GetName() : TEXT("None"), TargetSkeleton ? *TargetSkeleton->GetName() : TEXT("None"));
 
 		if (bIsMeshSpaceAdditive)
 		{
@@ -1348,7 +1351,10 @@ void FDecompressionTools::DecompressPose(FLODPose& OutAnimationPoseData,
 	//		}
 	//	}
 	//}
-	ensureMsgf(OrientAndScaleRetargetingPairs.Num() == 0, TEXT("OrientAndScaleRetargetingPairs found for %s (%s -> %s), Orient and Scale retargeting is not yet supported in AnimNext"), *DecompressionContext.AnimName.ToString(), SourceSkeleton ? *SourceSkeleton->GetName() : TEXT("None"), TargetSkeleton ? *TargetSkeleton->GetName() : TEXT("None"));
+	UE_CLOGFMT(OrientAndScaleRetargetingPairs.Num() == 0,
+		LogAnimation, Verbose,
+		"OrientAndScaleRetargetingPairs found for {AnimSequence} ({SourceSkeleton} -> {TargetSkeleton}), Orient and Scale retargeting is not yet supported in UAF",
+		*DecompressionContext.AnimName.ToString(), SourceSkeleton ? *SourceSkeleton->GetName() : TEXT("None"), TargetSkeleton ? *TargetSkeleton->GetName() : TEXT("None"));
 
 	if (bIsMeshSpaceAdditive)
 	{
