Use TFunction instead of TFunctionRef as we own its lifetime

#FailedValidations Preflight

diff --git a/Engine/Plugins/Experimental/UAF/UAF/Source/UAF/Public/UAF/Attributes/AttributeRuntimeRegistry.h b/Engine/Plugins/Experimental/UAF/UAF/Source/UAF/Public/UAF/Attributes/AttributeRuntimeRegistry.h
--- a/Engine/Plugins/Experimental/UAF/UAF/Source/UAF/Public/UAF/Attributes/AttributeRuntimeRegistry.h
+++ b/Engine/Plugins/Experimental/UAF/UAF/Source/UAF/Public/UAF/Attributes/AttributeRuntimeRegistry.h
@@ -16,7 +16,7 @@ namespace UE::UAF
 	public:
 		UE_NONCOPYABLE(FAttributeRuntimeRegistry);	// We disallow copy/move semantics
 
-		using SetMapConstructorFun = TFunctionRef<FAttributeSetMap*(const FAttributeSetMap::FConstructArgs& Args)>;
+		using SetMapConstructorFun = TFunction<FAttributeSetMap*(const FAttributeSetMap::FConstructArgs& Args)>;
 
 		// Returns the set map registry instance
 		[[nodiscard]] static UE_API FAttributeRuntimeRegistry& Get();
