Introduce FGuid for UAF variable entries, allowing to lazy-reconciliation of SharedVariableNode and persistent StateTree bindings
- 'Automatic' refactor tools now only modify currently loaded assets (or force loads if specified)
- Added post-load behaviour for SharedVariableNode and various patchup paths
- RefreshSharedVariablesNode now sets up Orphan pins when requested, clarifies compiler errors for the user

#misc added missing python script logging

diff --git a/Engine/Plugins/Experimental/UAF/UAF/Source/UAFUncookedOnly/Internal/AnimNextControllerBase.h b/Engine/Plugins/Experimental/UAF/UAF/Source/UAFUncookedOnly/Internal/AnimNextControllerBase.h
--- a/Engine/Plugins/Experimental/UAF/UAF/Source/UAFUncookedOnly/Internal/AnimNextControllerBase.h
+++ b/Engine/Plugins/Experimental/UAF/UAF/Source/UAFUncookedOnly/Internal/AnimNextControllerBase.h
@@ -29,7 +29,7 @@ public:
 
 	// Refreshes a UAnimNextSharedVariableNode instance with provided data (similar to URigVMController::RefreshVariableNode)
 	UFUNCTION(BlueprintCallable, Category = RigVMController)
-	UAFUNCOOKEDONLY_API void RefreshSharedVariableNode(const FName& InNodeName, const FString& InSourceObjectPath, const FName& InVariableName, const FString& InCPPType, UObject* InCPPTypeObject, bool bSetupUndoRedo, bool bPrintPythonCommand);
+	UAFUNCOOKEDONLY_API void RefreshSharedVariableNode(const FName& InNodeName, const FString& InSourceObjectPath, const FName& InVariableName, const FString& InCPPType, UObject* InCPPTypeObject, bool bSetupUndoRedo, bool bSetupOrphanPins = true, bool bPrintPythonCommand = false);
 
 	// Add a shared variable from the specified source object
 	UFUNCTION(BlueprintCallable, Category = RigVMController)
diff --git a/Engine/Plugins/Experimental/UAF/UAF/Source/UAFUncookedOnly/Internal/AnimNextRigVMAssetEditorData.h b/Engine/Plugins/Experimental/UAF/UAF/Source/UAFUncookedOnly/Internal/AnimNextRigVMAssetEditorData.h
--- a/Engine/Plugins/Experimental/UAF/UAF/Source/UAFUncookedOnly/Internal/AnimNextRigVMAssetEditorData.h
+++ b/Engine/Plugins/Experimental/UAF/UAF/Source/UAFUncookedOnly/Internal/AnimNextRigVMAssetEditorData.h
@@ -23,6 +23,7 @@ class UAnimNextEdGraph;
 class UAnimNextSharedVariablesEntry;
 class UAnimNextWorkspaceEditorMode;
 class UUAFTestBlueprintLibrary;
+class UAnimNextSharedVariableNode;
 struct FAnimNextRigVMAssetCompileContext;
 struct FAnimNextGetFunctionHeaderCompileContext;
 struct FAnimNextGetVariableCompileContext;
@@ -275,6 +276,7 @@ protected:
 	friend UE::UAF::Editor::SFunctionsView;
 	friend UE::UAF::Editor::SCommonOutliner;
 	friend UE::UAF::Editor::SVariablesView;
+	friend UAnimNextSharedVariableNode;
 
 	// UObject interface
 	UE_API virtual void Serialize(FArchive& Ar) override;
@@ -565,10 +567,11 @@ protected:
 		EAnimNextExportAccessSpecifier Access = EAnimNextExportAccessSpecifier::Private;
 		const FProperty* Property = nullptr;
 		TConstArrayView<uint8> DefaultValue;
+		FGuid StableGuid;
 
 		friend uint32 GetTypeHash(const FVariableInfo& Info)
 		{
-			return HashCombine(HashCombine(GetTypeHash(Info.Name), GetTypeHash(Info.Type)), GetTypeHash(Info.SourceObject));
+			return HashCombine(HashCombine(GetTypeHash(Info.Name), GetTypeHash(Info.Type)), GetTypeHash(Info.SourceObject), GetTypeHash(Info.StableGuid));
 		}
 
 		bool operator==(const FVariableInfo& Other) const
@@ -576,7 +579,8 @@ protected:
 			return Name == Other.Name
 				&& Type == Other.Type
 				&& SourceObject == Other.SourceObject
-				&& Access == Other.Access;
+				&& Access == Other.Access
+				&& StableGuid == Other.StableGuid;
 		}
 	};
 
diff --git a/Engine/Plugins/Experimental/UAF/UAF/Source/UAFUncookedOnly/Internal/AnimNextSharedVariableNode.h b/Engine/Plugins/Experimental/UAF/UAF/Source/UAFUncookedOnly/Internal/AnimNextSharedVariableNode.h
--- a/Engine/Plugins/Experimental/UAF/UAF/Source/UAFUncookedOnly/Internal/AnimNextSharedVariableNode.h
+++ b/Engine/Plugins/Experimental/UAF/UAF/Source/UAFUncookedOnly/Internal/AnimNextSharedVariableNode.h
@@ -26,6 +26,11 @@ class UAnimNextSharedVariableNode : public URigVMVariableNode
 
 	// URigVMVariableNode interface
 	virtual FString GetNodeSubTitle() const override;
+	virtual void PostLoad() override;
+	virtual void Serialize(FArchive& Ar) override;
+
+private:
+	void UpdateCachedVariableGUID();
 
 private:
 	friend UAnimNextControllerBase;
@@ -40,4 +45,7 @@ private:
 	
 	UPROPERTY()
 	EAnimNextSharedVariablesType Type;
-};
\ No newline at end of file
+	
+	UPROPERTY()
+	FGuid CachedGuid;
+};
diff --git a/Engine/Plugins/Experimental/UAF/UAF/Source/UAFUncookedOnly/Internal/Entries/AnimNextVariableEntry.h b/Engine/Plugins/Experimental/UAF/UAF/Source/UAFUncookedOnly/Internal/Entries/AnimNextVariableEntry.h
--- a/Engine/Plugins/Experimental/UAF/UAF/Source/UAFUncookedOnly/Internal/Entries/AnimNextVariableEntry.h
+++ b/Engine/Plugins/Experimental/UAF/UAF/Source/UAFUncookedOnly/Internal/Entries/AnimNextVariableEntry.h
@@ -123,6 +123,7 @@ class UAnimNextVariableEntry : public UAnimNextRigVMAssetEntry, public IAnimNext
 	virtual TConstStructView<FAnimNextVariableBindingData> GetBinding() const override;
 	virtual FStringView GetVariableCategory() const override;
 	virtual void SetVariableCategory(FStringView InCategoryName, bool bSetupUndoRedo = true) override;
+	virtual const FGuid& GetGuid() const override { return Guid; }
 
 	// Set the default value
 	template<typename ValueType>
@@ -173,4 +174,8 @@ class UAnimNextVariableEntry : public UAnimNextRigVMAssetEntry, public IAnimNext
 	/** Property bag holding the default value of the variable */
 	UPROPERTY()
 	FInstancedPropertyBag DefaultValue;
+	
+	/** Stable (unique) Guid used to reference variable without requiring ParameterName */
+	UPROPERTY()
+	FGuid Guid;
 };
diff --git a/Engine/Plugins/Experimental/UAF/UAF/Source/UAFUncookedOnly/Internal/UncookedOnlyUtils.h b/Engine/Plugins/Experimental/UAF/UAF/Source/UAFUncookedOnly/Internal/UncookedOnlyUtils.h
--- a/Engine/Plugins/Experimental/UAF/UAF/Source/UAFUncookedOnly/Internal/UncookedOnlyUtils.h
+++ b/Engine/Plugins/Experimental/UAF/UAF/Source/UAFUncookedOnly/Internal/UncookedOnlyUtils.h
@@ -219,7 +219,10 @@ struct FUtils
 	// Replaces the provided variable reference with ReferenceToReplaceWith within the context of InEditorData
 	static UE_API void ReplaceVariableReferences(UAnimNextRigVMAssetEditorData* InEditorData, const FAnimNextSoftVariableReference& ReferenceToFind, const FAnimNextSoftVariableReference& ReferenceToReplaceWith, bool bSetupUndoRedo = true, bool bPrintPythonCommands = true);
 	// Replaces the provided variable reference with ReferenceToReplaceWith within the context of the project
-	static UE_API void ReplaceVariableReferencesAcrossProject(const FAnimNextSoftVariableReference& ReferenceToFind, const FAnimNextSoftVariableReference& ReferenceToReplaceWith, bool bSetupUndoRedo = true, bool bPrintPythonCommands = true);
+	static UE_API void ReplaceVariableReferencesAcrossProject(const FAnimNextSoftVariableReference& ReferenceToFind, const FAnimNextSoftVariableReference& ReferenceToReplaceWith, bool bForceLoadAssets = false, bool bSetupUndoRedo = true, bool bPrintPythonCommands = true);
+
+	// Generates a deterministic GUID provided a UStruct FProperty entry
+	static UE_API FGuid GenerateScriptStructPropertyGUID(const FProperty* Property);
 
 	// Returns whether or not it is valid to add InSharedVariables to InEditorData, and if not OutErrorMessage will contain the reasoning
 	static UE_API bool CanAddSharedVariablesReference(UAnimNextRigVMAssetEditorData* InEditorData, const UAnimNextSharedVariables* InSharedVariables, FString* OutErrorMessage = nullptr);
diff --git a/Engine/Plugins/Experimental/UAF/UAF/Source/UAFUncookedOnly/Private/AnimNextControllerBase.cpp b/Engine/Plugins/Experimental/UAF/UAF/Source/UAFUncookedOnly/Private/AnimNextControllerBase.cpp
--- a/Engine/Plugins/Experimental/UAF/UAF/Source/UAFUncookedOnly/Private/AnimNextControllerBase.cpp
+++ b/Engine/Plugins/Experimental/UAF/UAF/Source/UAFUncookedOnly/Private/AnimNextControllerBase.cpp
@@ -3,7 +3,9 @@
 #include "AnimNextControllerBase.h"
 #include "AnimNextSharedVariableNode.h"
 #include "RigVMPythonUtils.h"
+#include "UncookedOnlyUtils.h"
 #include "Variables/AnimNextSharedVariables.h"
+#include "Variables/AnimNextSharedVariables_EditorData.h"
 
 #include UE_INLINE_GENERATED_CPP_BY_NAME(AnimNextControllerBase)
 
@@ -63,11 +65,15 @@ UAnimNextSharedVariableNode* UAnimNextControllerBase::AddSharedVariableNode(cons
 	{
 		VariableNode->Type = EAnimNextSharedVariablesType::Asset;
 		VariableNode->Asset = Asset;
+		
+		VariableNode->UpdateCachedVariableGUID();
 	}
 	else if (const UScriptStruct* Struct = Cast<UScriptStruct>(SourceObject))
 	{
 		VariableNode->Type = EAnimNextSharedVariablesType::Struct;
 		VariableNode->Struct = Struct;
+		
+		VariableNode->UpdateCachedVariableGUID();
 	}
 	else
 	{
@@ -96,7 +102,7 @@ UAnimNextSharedVariableNode* UAnimNextControllerBase::AddSharedVariableNode(cons
 	return VariableNode;
 }
 
-void UAnimNextControllerBase::RefreshSharedVariableNode(const FName& InNodeName, const FString& InSourceObjectPath, const FName& InVariableName, const FString& InCPPType, UObject* InCPPTypeObject, bool bSetupUndoRedo, bool bPrintPythonCommand)
+void UAnimNextControllerBase::RefreshSharedVariableNode(const FName& InNodeName, const FString& InSourceObjectPath, const FName& InVariableName, const FString& InCPPType, UObject* InCPPTypeObject, bool bSetupUndoRedo, bool bSetupOrphanPins, bool bPrintPythonCommand)
 {
 	UObject* SourceObject = nullptr;
 	if (!InSourceObjectPath.IsEmpty())
@@ -113,7 +119,6 @@ void UAnimNextControllerBase::RefreshSharedVariableNode(const FName& InNodeName,
 			VariableNode->Modify();
 		}
 		
-		constexpr bool bSetupOrphanPins = false;
 		RefreshVariableNode(InNodeName, InVariableName, InCPPType, InCPPTypeObject, bSetupUndoRedo, bSetupOrphanPins);
 
 		if (SourceObject)
@@ -134,12 +139,27 @@ void UAnimNextControllerBase::RefreshSharedVariableNode(const FName& InNodeName,
 			VariableNode->Type = EAnimNextSharedVariablesType::Asset;
 			VariableNode->Asset = nullptr;
 		}
+		
+		VariableNode->UpdateCachedVariableGUID();
 
 		Notify(ERigVMGraphNotifType::NodeTitleChanged, VariableNode);
 
 		if (bPrintPythonCommand)
 		{
-			
+			const FString GraphName = GetSchema()->GetSanitizedGraphName(GetGraph()->GetGraphName());
+
+			RigVMPythonUtils::Print(GetSchema()->GetGraphOuterName(GetGraph()), 
+				FString::Printf(TEXT("asset.get_controller_by_name('%s').refresh_shared_variable_node('%s', '%s', '%s', '%s', '%s', '%s', '%s', 'True')"),
+					*GraphName,
+					*InNodeName.ToString(),
+					*InSourceObjectPath,
+					*InVariableName.ToString(),
+					*InCPPType,
+					*RigVMTypeUtils::CPPTypeFromObject(InCPPTypeObject),
+					bSetupUndoRedo ? TEXT("True") : TEXT("False"),
+					bSetupOrphanPins ? TEXT("True") : TEXT("False")
+				)
+			);
 		}
 	}
 }
diff --git a/Engine/Plugins/Experimental/UAF/UAF/Source/UAFUncookedOnly/Private/AnimNextRigVMAssetEditorData.cpp b/Engine/Plugins/Experimental/UAF/UAF/Source/UAFUncookedOnly/Private/AnimNextRigVMAssetEditorData.cpp
--- a/Engine/Plugins/Experimental/UAF/UAF/Source/UAFUncookedOnly/Private/AnimNextRigVMAssetEditorData.cpp
+++ b/Engine/Plugins/Experimental/UAF/UAF/Source/UAFUncookedOnly/Private/AnimNextRigVMAssetEditorData.cpp
@@ -2697,7 +2697,8 @@ void UAnimNextRigVMAssetEditorData::GetAllVariables(TArray<FVariableInfo>& OutVa
 						Asset,
 						VariableEntry->GetExportAccessSpecifier(),
 						Property,
-						Value
+						Value,
+						VariableEntry->GetGuid()
 					});
 			}
 		}
@@ -2732,7 +2733,8 @@ void UAnimNextRigVMAssetEditorData::GetAllVariables(TArray<FVariableInfo>& OutVa
 									SharedVariablesEntry->GetStruct(),
 									EAnimNextExportAccessSpecifier::Public,
 									Property,
-									TConstArrayView<uint8>()
+									TConstArrayView<uint8>(),
+									UE::UAF::UncookedOnly::FUtils::GenerateScriptStructPropertyGUID(Property)
 								});
 						}
 					}
diff --git a/Engine/Plugins/Experimental/UAF/UAF/Source/UAFUncookedOnly/Private/AnimNextSharedVariableNode.cpp b/Engine/Plugins/Experimental/UAF/UAF/Source/UAFUncookedOnly/Private/AnimNextSharedVariableNode.cpp
--- a/Engine/Plugins/Experimental/UAF/UAF/Source/UAFUncookedOnly/Private/AnimNextSharedVariableNode.cpp
+++ b/Engine/Plugins/Experimental/UAF/UAF/Source/UAFUncookedOnly/Private/AnimNextSharedVariableNode.cpp
@@ -1,10 +1,226 @@
 // Copyright Epic Games, Inc. All Rights Reserved.
 
 #include "AnimNextSharedVariableNode.h"
+#include "UncookedOnlyUtils.h"
 #include "Variables/AnimNextSharedVariables.h"
+#include "Variables/AnimNextSharedVariables_EditorData.h"
 
 #include UE_INLINE_GENERATED_CPP_BY_NAME(AnimNextSharedVariableNode)
 
+void UAnimNextSharedVariableNode::PostLoad()
+{
+	Super::PostLoad();
+
+	if (GetLinkerCustomVersion(FFortniteMainBranchObjectVersion::GUID) < FFortniteMainBranchObjectVersion::UAFVariablesGuid)
+	{
+		UpdateCachedVariableGUID();
+	}
+	else if (!CachedGuid.IsValid())
+	{
+		UE_LOG(LogAnimation, Warning, TEXT("%s had an invalid variable cached GUID value for %s in %s"),
+			*GetPathName(), 
+			*GetVariableName().ToString(),
+			Type == EAnimNextSharedVariablesType::Asset ? (Asset ? *Asset->GetPathName() : TEXT("None"))
+			: (Struct ? *Struct->GetPathName() : TEXT("None"))
+		);
+		
+		UpdateCachedVariableGUID();
+	}
+	else
+	{
+		const IRigVMClientHost* OuterHost = GetImplementingOuter<IRigVMClientHost>();
+		const URigVMGraph* OuterGraph = GetTypedOuter<URigVMGraph>();
+
+		if (OuterHost && OuterGraph)
+		{
+			const FName ExpectedVariableName = GetVariableName();
+
+			switch (Type)
+			{
+			case EAnimNextSharedVariablesType::Asset:
+				{
+					if (const UAnimNextSharedVariables* SharedVariablesAsset = Asset.Get())
+					{
+						const UAnimNextSharedVariables_EditorData* EditorData = UE::UAF::UncookedOnly::FUtils::GetEditorData<UAnimNextSharedVariables_EditorData>(SharedVariablesAsset);
+						const FRigVMTemplateArgumentType ExpectedVariableType(FName(*GetCPPType()), GetCPPTypeObject());
+
+						TArray<UAnimNextRigVMAssetEditorData::FVariableInfo> AssetVariables;
+						EditorData->GetAllVariables(AssetVariables, UAnimNextRigVMAssetEditorData::EVariableRecursion::SelfOnly, UAnimNextRigVMAssetEditorData::EVariableAccessFilter::PublicOnly);
+
+						UAnimNextRigVMAssetEditorData::FVariableInfo* MatchingVariable = nullptr;
+						MatchingVariable = AssetVariables.FindByPredicate([SearchGuid = CachedGuid](const UAnimNextRigVMAssetEditorData::FVariableInfo& VariableInfo) { return VariableInfo.StableGuid == SearchGuid; });
+						if (MatchingVariable)
+						{
+							// Variable has matching Guid but different name
+							if (UAnimNextControllerBase* AnimNextController = Cast<UAnimNextControllerBase>(OuterHost->GetController(OuterGraph)))
+							{
+								const FRigVMTemplateArgumentType VariableType = MatchingVariable->Type.ToRigVMTemplateArgument();
+								if (VariableType != ExpectedVariableType || MatchingVariable->Name != GetVariableName())
+								{
+									// Update with new variable name - and type if required, populating orphan pins of non-compatible data types
+									AnimNextController->RefreshSharedVariableNode(GetFName(), SharedVariablesAsset->GetPathName(), MatchingVariable->Name, VariableType.CPPType.ToString(), VariableType.CPPTypeObject.Get(), false, true, false );
+								}
+							}
+						}
+						else
+						{
+							// Look for variable with matching (node stored) variable name 
+							MatchingVariable = AssetVariables.FindByPredicate([ExpectedVariableName](const UAnimNextRigVMAssetEditorData::FVariableInfo& VariableInfo) { return VariableInfo.Name == ExpectedVariableName; });
+							if (MatchingVariable)
+							{
+								// Variable has matching name but different GUID 
+								checkf(CachedGuid != MatchingVariable->StableGuid, TEXT("Expected GUID to be different, otherwise previous search using GUID should have been successful"));
+								UpdateCachedVariableGUID();
+
+								if (UAnimNextControllerBase* AnimNextController = Cast<UAnimNextControllerBase>(OuterHost->GetController(OuterGraph)))
+								{
+									const FRigVMTemplateArgumentType VariableType = MatchingVariable->Type.ToRigVMTemplateArgument();
+									if (VariableType != ExpectedVariableType)
+									{
+										// Update data type, populating orphan pins of non-compatible data types
+										AnimNextController->RefreshSharedVariableNode(GetFName(), SharedVariablesAsset->GetPathName(), ExpectedVariableName, VariableType.CPPType.ToString(), VariableType.CPPTypeObject.Get(), false, true, false );
+									}
+								}
+							}
+							else
+							{
+								// node cannot be resolved
+								UE_LOG(LogAnimation, Warning, TEXT("Unable to find public variable with expected name: %s, type: %s and GUID %s in %s, SharedVariableNode will be invalidated"), *ExpectedVariableName.ToString(), *FAnimNextParamType::FromRigVMTemplateArgument(ExpectedVariableType).ToString(), *CachedGuid.ToString(EGuidFormats::Short), *SharedVariablesAsset->GetPathName());
+
+								if (UAnimNextControllerBase* AnimNextController = Cast<UAnimNextControllerBase>(OuterHost->GetController(OuterGraph)))
+								{
+									AnimNextController->RefreshSharedVariableNode(GetFName(), SharedVariablesAsset->GetPathName(), ExpectedVariableName, *RigVMTypeUtils::GetWildCardCPPTypeName().ToString(), RigVMTypeUtils::GetWildCardCPPTypeObject(), false, true, false );
+								}
+							}
+						}
+					}
+					else
+					{
+						UE_LOG(LogAnimation, Warning, TEXT("Unable to load Asset for %s when trying to resolve SharedVariable node source"), *GetPathName());
+
+						if (UAnimNextControllerBase* AnimNextController = Cast<UAnimNextControllerBase>(OuterHost->GetController(OuterGraph)))
+						{
+							AnimNextController->RefreshSharedVariableNode(GetFName(), TEXT(""), ExpectedVariableName, *RigVMTypeUtils::GetWildCardCPPTypeName().ToString(), RigVMTypeUtils::GetWildCardCPPTypeObject(), false, true, false );
+						}
+					}
+					break;
+				}
+
+			case EAnimNextSharedVariablesType::Struct:
+				{
+					if (const UScriptStruct* SharedVariablesStruct = Struct.Get())
+					{
+						const FAnimNextParamType ExpectedVariableType = FAnimNextParamType::FromRigVMTemplateArgument(FRigVMTemplateArgumentType(FName(*GetCPPType()), GetCPPTypeObject()));	
+						FName ExpectedStructVariableName = ExpectedVariableName;
+						const FProperty* Property = SharedVariablesStruct->FindPropertyByName(ExpectedStructVariableName);
+
+						// Try property redirector
+						if (!Property)
+						{
+							const FName RedirectedName = FProperty::FindRedirectedPropertyName(SharedVariablesStruct, ExpectedStructVariableName);
+							if (RedirectedName != NAME_None && RedirectedName != ExpectedStructVariableName)
+							{
+								Property = SharedVariablesStruct->FindPropertyByName(RedirectedName);
+								ExpectedStructVariableName = RedirectedName;
+							}
+						}
+
+						if (Property)
+						{
+							const FAnimNextParamType PropertyParamType = FAnimNextParamType::FromProperty(Property);
+							if (ExpectedVariableType != PropertyParamType || ExpectedStructVariableName != ExpectedVariableName)
+							{
+								if (UAnimNextControllerBase* AnimNextController = Cast<UAnimNextControllerBase>(OuterHost->GetController(OuterGraph)))
+								{
+									// Update with new variable name - and type if required, populating orphan pins of non-compatible data types
+									AnimNextController->RefreshSharedVariableNode(GetFName(), SharedVariablesStruct->GetPathName(), ExpectedStructVariableName, *Property->GetCPPType(), const_cast<UObject*>(PropertyParamType.GetValueTypeObject()), false, true, false );
+								}
+							}
+						}
+						else
+						{
+							// node cannot be resolved
+							UE_LOG(LogAnimation, Warning, TEXT("Unable to find struct FProperty with expected name: %s, type: %s in %s, SharedVariableNode will be invalidated"), *ExpectedVariableName.ToString(), *ExpectedVariableType.ToString(), *SharedVariablesStruct->GetPathName());
+
+							if (UAnimNextControllerBase* AnimNextController = Cast<UAnimNextControllerBase>(OuterHost->GetController(OuterGraph)))
+							{
+								AnimNextController->RefreshSharedVariableNode(GetFName(), SharedVariablesStruct->GetPathName(), ExpectedStructVariableName, *RigVMTypeUtils::GetWildCardCPPTypeName().ToString(), RigVMTypeUtils::GetWildCardCPPTypeObject(), false, true, false );
+							}
+						}
+					}
+					else
+					{
+						UE_LOG(LogAnimation, Warning, TEXT("Unable to find Struct for %s when trying to resolve SharedVariable node source"), *GetPathName());
+						
+						if (UAnimNextControllerBase* AnimNextController = Cast<UAnimNextControllerBase>(OuterHost->GetController(OuterGraph)))
+						{
+							AnimNextController->RefreshSharedVariableNode(GetFName(), TEXT(""), ExpectedVariableName, *RigVMTypeUtils::GetWildCardCPPTypeName().ToString(), RigVMTypeUtils::GetWildCardCPPTypeObject(), false, true, false );
+						}
+					}
+					break;
+				}
+			}
+		}
+	}
+}
+
+void UAnimNextSharedVariableNode::Serialize(FArchive& Ar)
+{
+	Super::Serialize(Ar);
+	Ar.UsingCustomVersion(FFortniteMainBranchObjectVersion::GUID);
+}
+
+void UAnimNextSharedVariableNode::UpdateCachedVariableGUID()
+{
+	CachedGuid.Invalidate();
+
+	switch (Type)
+	{
+	case EAnimNextSharedVariablesType::Asset:
+		{
+			if (Asset)
+			{
+				const UAnimNextSharedVariables_EditorData* EditorData = UE::UAF::UncookedOnly::FUtils::GetEditorData<const UAnimNextSharedVariables_EditorData>(Asset.Get());
+
+				// FindEntry depends on UAnimNextSharedVariables_EditorData::PostLoadExternalPackages having been called (copying internal_entries into entries array)
+				const_cast<UAnimNextSharedVariables_EditorData*>(EditorData)->ConditionalPostLoad();
+				if (const IAnimNextRigVMVariableInterface* VariableInterface = EditorData != nullptr ? Cast<IAnimNextRigVMVariableInterface>(EditorData->FindEntry(GetVariableName())) : nullptr)
+				{
+					CachedGuid = VariableInterface->GetGuid();
+				}
+				else
+				{
+					UE_LOG(LogAnimation, Warning, TEXT("Unable to generate GUID for %s as Asset %s does not contain variable %s"), *GetPathName(), *Asset->GetPathName(), *GetVariableName().ToString());
+				}
+			}
+			else
+			{
+				UE_LOG(LogAnimation, Warning, TEXT("Unable to generate GUID for %s as Asset is invalid"), *GetPathName());
+			}
+		}
+		break;
+	case EAnimNextSharedVariablesType::Struct:
+		{
+			if (Struct)
+			{
+				if (const FProperty* Property = Struct->FindPropertyByName(GetVariableName()))
+				{
+					CachedGuid = UE::UAF::UncookedOnly::FUtils::GenerateScriptStructPropertyGUID(Property);
+				}
+				else
+				{
+					UE_LOG(LogAnimation, Warning, TEXT("Unable to generate GUID for %s as Struct %s does not contain variable %s"), *GetPathName(), *Struct->GetPathName(), *GetVariableName().ToString());
+				}
+			}
+			else
+			{
+				UE_LOG(LogAnimation, Warning, TEXT("Unable to generate GUID for %s as Struct is invalid"), *GetPathName());
+			}
+		}
+		break;
+	}
+}
+
 FString UAnimNextSharedVariableNode::GetNodeSubTitle() const
 {
 	switch (Type)
@@ -16,3 +232,4 @@ FString UAnimNextSharedVariableNode::GetNodeSubTitle() const
 	}
 	return FString();
 }
+
diff --git a/Engine/Plugins/Experimental/UAF/UAF/Source/UAFUncookedOnly/Private/Entries/AnimNextVariableEntry.cpp b/Engine/Plugins/Experimental/UAF/UAF/Source/UAFUncookedOnly/Private/Entries/AnimNextVariableEntry.cpp
--- a/Engine/Plugins/Experimental/UAF/UAF/Source/UAFUncookedOnly/Private/Entries/AnimNextVariableEntry.cpp
+++ b/Engine/Plugins/Experimental/UAF/UAF/Source/UAFUncookedOnly/Private/Entries/AnimNextVariableEntry.cpp
@@ -14,6 +14,7 @@ const FLazyName IAnimNextRigVMVariableInterface::ValueName("Value");
 
 UAnimNextVariableEntry::UAnimNextVariableEntry()
 {
+	Guid = FGuid::NewGuid();
 }
 
 FAnimNextParamType UAnimNextVariableEntry::GetExportType() const
@@ -269,6 +270,11 @@ void UAnimNextVariableEntry::PostLoad()
 			NewDesc->CachedProperty->CopyCompleteValue(NewValue, OldValue);
 		}
 	}
+
+	if (GetLinkerCustomVersion(FFortniteMainBranchObjectVersion::GUID) < FFortniteMainBranchObjectVersion::UAFVariablesGuid)
+	{
+		Guid = FGuid::NewGuid();
+	}
 }
 
 #if WITH_EDITOR
diff --git a/Engine/Plugins/Experimental/UAF/UAF/Source/UAFUncookedOnly/Private/UncookedOnlyUtils.cpp b/Engine/Plugins/Experimental/UAF/UAF/Source/UAFUncookedOnly/Private/UncookedOnlyUtils.cpp
--- a/Engine/Plugins/Experimental/UAF/UAF/Source/UAFUncookedOnly/Private/UncookedOnlyUtils.cpp
+++ b/Engine/Plugins/Experimental/UAF/UAF/Source/UAFUncookedOnly/Private/UncookedOnlyUtils.cpp
@@ -75,6 +75,7 @@ FInstancedPropertyBag FUtils::MakePropertyBagForEditorData(const UAnimNextRigVMA
 		EAnimNextExportAccessSpecifier AccessSpecifier = EAnimNextExportAccessSpecifier::Private;
 		TConstArrayView<uint8> Value;
 		EPropertyFlags PropertyFlags = CPF_NativeAccessSpecifierPrivate | CPF_Edit;
+		FGuid Guid;
 	};
 
 	// Gather all variables in this asset.
@@ -85,7 +86,7 @@ FInstancedPropertyBag FUtils::MakePropertyBagForEditorData(const UAnimNextRigVMA
 	int32 NumPublicVariables = 0;
 
 	UAnimNextRigVMAsset* Asset = GetAsset<UAnimNextRigVMAsset>(InEditorData);
-	auto AddVariable = [Asset, &NumPublicVariables, &StructEntryInfos, &EntryInfoIndexMap, &InCompileContext](FName InName, const FAnimNextParamType& InType, EAnimNextExportAccessSpecifier InAccess, TConstArrayView<uint8> InValue, EPropertyFlags InFlags)
+	auto AddVariable = [Asset, &NumPublicVariables, &StructEntryInfos, &EntryInfoIndexMap, &InCompileContext](FName InName, const FAnimNextParamType& InType, EAnimNextExportAccessSpecifier InAccess, TConstArrayView<uint8> InValue, EPropertyFlags InFlags, FGuid InStableGUID)
 	{
 		if(!InType.IsValid())
 		{
@@ -111,7 +112,8 @@ FInstancedPropertyBag FUtils::MakePropertyBagForEditorData(const UAnimNextRigVMA
 				FAnimNextParamType(InType.GetValueType(), InType.GetContainerType(), InType.GetValueTypeObject()),
 				InAccess,
 				InValue,
-				InFlags
+				InFlags,
+				InStableGUID
 			});
 
 		EntryInfoIndexMap.Add(InName, Index);
@@ -130,13 +132,14 @@ FInstancedPropertyBag FUtils::MakePropertyBagForEditorData(const UAnimNextRigVMA
 			{
 				Flags |= CPF_NativeAccessSpecifierPublic;
 			}
-			
+
 			AddVariable(
 				VariableEntry->GetExportName(),
 				VariableEntry->GetType(),
 				VariableEntry->GetExportAccessSpecifier(),
 				TConstArrayView<uint8>(VariableEntry->GetValuePtr(), VariableEntry->GetType().GetSize()),
-				Flags);
+				Flags,
+				VariableEntry->GetGuid());
 		}
 	}
 
@@ -147,7 +150,10 @@ FInstancedPropertyBag FUtils::MakePropertyBagForEditorData(const UAnimNextRigVMA
 			Variable.Type,
 			EAnimNextExportAccessSpecifier::Private,
 			TConstArrayView<uint8>(Variable.GetValuePtr(), Variable.Type.GetSize()),
-			CPF_AdvancedDisplay | CPF_NativeAccessSpecifierPrivate);
+			CPF_AdvancedDisplay | CPF_NativeAccessSpecifierPrivate,
+			// [TODO] this assumes Variable.Name being unique yet consistent across compilations
+			FGuid::NewDeterministicGuid(Variable.Name.ToString(), FCrc::StrCrc32<TCHAR>(TEXT("ProgrammaticVariableGUID")))
+		);
 	}
 
 	// Sort by size, largest first, for better packing
@@ -175,7 +181,8 @@ FInstancedPropertyBag FUtils::MakePropertyBagForEditorData(const UAnimNextRigVMA
 
 		for (const FStructEntryInfo& StructEntryInfo : StructEntryInfos)
 		{
-			PropertyDescs.Emplace(StructEntryInfo.Name, StructEntryInfo.Type.ContainerType, StructEntryInfo.Type.ValueType, StructEntryInfo.Type.ValueTypeObject, StructEntryInfo.PropertyFlags);
+			FPropertyBagPropertyDesc& Desc = PropertyDescs.Emplace_GetRef(StructEntryInfo.Name, StructEntryInfo.Type.ContainerType, StructEntryInfo.Type.ValueType, StructEntryInfo.Type.ValueTypeObject, StructEntryInfo.PropertyFlags);
+			Desc.ID = StructEntryInfo.Guid;
 			Values.Add(StructEntryInfo.Value);
 		}
 
@@ -206,7 +213,6 @@ void FUtils::CompileVariables(const FRigVMCompileSettings& InSettings, UAnimNext
 	{
 		if(UAnimNextSharedVariablesEntry* SharedVariablesEntry = Cast<UAnimNextSharedVariablesEntry>(Entry))
 		{
-
 			auto SharedAssetEntryError = [InAsset, EditorData, SharedVariablesEntry, &OutCompileContext](const FText& InMessage)
 			{
 				TWeakObjectPtr<UAnimNextRigVMAssetEditorData> WeakEditorData = EditorData;
@@ -1617,7 +1623,8 @@ void FUtils::DeleteVariable(UAnimNextVariableEntry* VariableEntry, bool bSetupUn
 		const bool bIsPublicVariable = VariableEntry->GetExportAccessSpecifier() == EAnimNextExportAccessSpecifier::Public;
 		if (bIsPublicVariable)
 		{
-			ReplaceVariableReferencesAcrossProject(FindReference, FAnimNextSoftVariableReference(), bSetupUndoRedo, bPrintPythonCommands);
+			constexpr bool bForceLoadReferencingAssets = false;
+			ReplaceVariableReferencesAcrossProject(FindReference, FAnimNextSoftVariableReference(), bForceLoadReferencingAssets, bSetupUndoRedo, bPrintPythonCommands);
 		}
 		else
 		{
@@ -1649,7 +1656,8 @@ void FUtils::RenameVariable(UAnimNextVariableEntry* VariableEntry, const FName N
 		const bool bIsPublicVariable = VariableEntry->GetExportAccessSpecifier() == EAnimNextExportAccessSpecifier::Public;
 		if (bIsPublicVariable)
 		{
-			ReplaceVariableReferencesAcrossProject(FindReference, ReplaceReference, bSetupUndoRedo, bPrintPythonCommands);
+			constexpr bool bForceLoadReferencingAssets = false;
+			ReplaceVariableReferencesAcrossProject(FindReference, ReplaceReference, bForceLoadReferencingAssets, bSetupUndoRedo, bPrintPythonCommands);
 		}
 		else
 		{
@@ -1683,7 +1691,8 @@ UAnimNextVariableEntry* FUtils::MoveVariableToAsset(UAnimNextVariableEntry* Vari
 			const bool bIsPublicVariable = VariableEntry->GetExportAccessSpecifier() == EAnimNextExportAccessSpecifier::Public;
 			if (bIsPublicVariable)
 			{
-				ReplaceVariableReferencesAcrossProject(FindReference, ReplaceReference, bSetupUndoRedo, bPrintPythonCommands);
+				constexpr bool bForceLoadReferencingAssets = false;
+				ReplaceVariableReferencesAcrossProject(FindReference, ReplaceReference, bForceLoadReferencingAssets, bSetupUndoRedo, bPrintPythonCommands);
 			}
 			else
 			{
@@ -1718,12 +1727,13 @@ void FUtils::SetVariableType(UAnimNextVariableEntry* VariableEntry, const FAnimN
 		const bool bIsAPublicVariable = VariableEntry->GetExportAccessSpecifier() == EAnimNextExportAccessSpecifier::Public;
 		if (bIsAPublicVariable)
 		{
-			ReplaceVariableReferencesAcrossProject(FindReference, FindReference, bSetupUndoRedo, bPrintPythonCommands);
+			constexpr bool bForceLoadReferencingAssets = false;
+			ReplaceVariableReferencesAcrossProject(FindReference, FindReference, bForceLoadReferencingAssets, bSetupUndoRedo, bPrintPythonCommands);
 		}
 		else
 		{
 			// Otherwise replace, potential, local references only 
-			ReplaceVariableReferences(GetEditorData(OuterRigVMAsset), FindReference, FindReference, bSetupUndoRedo, bPrintPythonCommands);		
+			ReplaceVariableReferences(GetEditorData(OuterRigVMAsset), FindReference, FindReference, bSetupUndoRedo, bPrintPythonCommands);
 		}
 	}	
 }
@@ -1873,7 +1883,9 @@ void FUtils::ReplaceVariableReferences(UAnimNextRigVMAssetEditorData* InEditorDa
 										else
 										{
 											const FRigVMTemplateArgumentType NewRigVMType = ReplaceReferenceType.ToRigVMTemplateArgument();
-											Controller->RefreshSharedVariableNode(SharedVariableNode->GetFName(), ReferenceToReplaceWith.GetObject().GetPath(), ReferenceToReplaceWith.GetName(), NewRigVMType.CPPType.ToString(), NewRigVMType.CPPTypeObject.Get(), bSetupUndoRedo, bPrintPythonCommands);											
+
+											constexpr bool bSetupOrphanPins = true;
+											Controller->RefreshSharedVariableNode(SharedVariableNode->GetFName(), ReferenceToReplaceWith.GetObject().GetPath(), ReferenceToReplaceWith.GetName(), NewRigVMType.CPPType.ToString(), NewRigVMType.CPPTypeObject.Get(), bSetupUndoRedo, bSetupOrphanPins, bPrintPythonCommands);											
 										}
 									}
 								}
@@ -1911,12 +1923,12 @@ void FUtils::ReplaceVariableReferences(UAnimNextRigVMAssetEditorData* InEditorDa
 										{
 											Controller->OnExternalVariableRenamed(ReferenceToFind.GetName(), ReferenceToReplaceWith.GetName(), bSetupUndoRedo);
 										}
-									}									
+									}
 								}
 							}
 						}
 					}
-				}			
+				}
 			}
 		}
 
@@ -1924,7 +1936,7 @@ void FUtils::ReplaceVariableReferences(UAnimNextRigVMAssetEditorData* InEditorDa
 	}	
 }
 
-void FUtils::ReplaceVariableReferencesAcrossProject(const FAnimNextSoftVariableReference& ReferenceToFind, const FAnimNextSoftVariableReference& ReferenceToReplaceWith, bool bSetupUndoRedo, bool bPrintPythonCommands)
+void FUtils::ReplaceVariableReferencesAcrossProject(const FAnimNextSoftVariableReference& ReferenceToFind, const FAnimNextSoftVariableReference& ReferenceToReplaceWith,  bool bForceLoadAssets, bool bSetupUndoRedo, bool bPrintPythonCommands)
 {
 	if (ReferenceToFind.IsNone())
 	{
@@ -1954,11 +1966,14 @@ void FUtils::ReplaceVariableReferencesAcrossProject(const FAnimNextSoftVariableR
 		{
 			if (ReferenceToFind.GetName() == Export.Identifier)
 			{
-				if (UAnimNextRigVMAsset* RigVMAsset = Cast<UAnimNextRigVMAsset>(AssetData.GetAsset()))
+				if (AssetData.IsAssetLoaded() || (bForceLoadAssets && AssetData.GetAsset()))
 				{
-					if (UAnimNextRigVMAssetEditorData* EditorData = GetEditorData(RigVMAsset))
+					if (UAnimNextRigVMAsset* RigVMAsset = Cast<UAnimNextRigVMAsset>(AssetData.GetAsset()))
 					{
-						EditorDataRequiringReplacing.Add(EditorData);
+						if (UAnimNextRigVMAssetEditorData* EditorData = GetEditorData(RigVMAsset))
+						{
+							EditorDataRequiringReplacing.Add(EditorData);
+						}
 					}
 				}
 			}
@@ -1970,6 +1985,13 @@ void FUtils::ReplaceVariableReferencesAcrossProject(const FAnimNextSoftVariableR
 		ReplaceVariableReferences(EditorData, ReferenceToFind, ReferenceToReplaceWith, bSetupUndoRedo, bPrintPythonCommands);
 	}
 }
+
+FGuid FUtils::GenerateScriptStructPropertyGUID(const FProperty* Property)
+{
+	checkf(Property, TEXT("Expected valid property for GUID generation"));
+	checkf(Property->GetOwnerChecked<UStruct>(), TEXT("Expected valid owning UStruct for GUID generation"));
+	return FGuid::NewDeterministicGuid(Property->GetPathName(), FCrc::StrCrc32<TCHAR>(TEXT("StructBasedSharedVariableEntry")));
+}
 }
 
 #undef LOCTEXT_NAMESPACE
diff --git a/Engine/Plugins/Experimental/UAF/UAF/Source/UAFUncookedOnly/Public/Variables/IAnimNextRigVMVariableInterface.h b/Engine/Plugins/Experimental/UAF/UAF/Source/UAFUncookedOnly/Public/Variables/IAnimNextRigVMVariableInterface.h
--- a/Engine/Plugins/Experimental/UAF/UAF/Source/UAFUncookedOnly/Public/Variables/IAnimNextRigVMVariableInterface.h
+++ b/Engine/Plugins/Experimental/UAF/UAF/Source/UAFUncookedOnly/Public/Variables/IAnimNextRigVMVariableInterface.h
@@ -67,6 +67,9 @@ public:
 	
 	// Set the category for this variable
 	virtual void SetVariableCategory(FStringView InCategoryName, bool bSetupUndoRedo = true) = 0;
+	
+	// Returns (stable and unique) GUID representing this variable
+	virtual const FGuid& GetGuid() const = 0;
 
 	// Access the memory for the internal value
 	const uint8* GetValuePtr() const
diff --git a/Engine/Source/Runtime/Core/Public/UObject/FortniteMainBranchObjectVersions.inl b/Engine/Source/Runtime/Core/Public/UObject/FortniteMainBranchObjectVersions.inl
--- a/Engine/Source/Runtime/Core/Public/UObject/FortniteMainBranchObjectVersions.inl
+++ b/Engine/Source/Runtime/Core/Public/UObject/FortniteMainBranchObjectVersions.inl
@@ -700,5 +700,8 @@ FORTNITE_MAIN_VERSION(Firefight_228, 228)
 // Add default runtime state to FDataLayerInstanceDesc
 FORTNITE_MAIN_VERSION(WorldPartitionDataLayerDefaultRuntimeState, 229)
 
+// Introducing FGuid to UAF variable entries
+FORTNITE_MAIN_VERSION(UAFVariablesGuid, 230)
+
 // -----<new versions can be added above this line>-------------------------------------------------
 #undef FORTNITE_MAIN_VERSION
\ No newline at end of file
