Small safety/python output fixes

diff --git a/Engine/Plugins/Experimental/UAF/UAF/Source/UAFUncookedOnly/Private/AnimNextControllerBase.cpp b/Engine/Plugins/Experimental/UAF/UAF/Source/UAFUncookedOnly/Private/AnimNextControllerBase.cpp
--- a/Engine/Plugins/Experimental/UAF/UAF/Source/UAFUncookedOnly/Private/AnimNextControllerBase.cpp
+++ b/Engine/Plugins/Experimental/UAF/UAF/Source/UAFUncookedOnly/Private/AnimNextControllerBase.cpp
@@ -17,7 +17,7 @@ UAnimNextSharedVariableNode* UAnimNextControllerBase::AddAssetSharedVariableNode
 		return nullptr;
 	}
 
-	return AddSharedVariableNode(InAsset->GetPathName(), InVariableName, InCPPType, InCPPTypeObject->GetPathName(), bIsGetter, InDefaultValue, InPosition, InNodeName, bSetupUndoRedo, bPrintPythonCommand);
+	return AddSharedVariableNode(InAsset->GetPathName(), InVariableName, InCPPType, InCPPTypeObject ? InCPPTypeObject->GetPathName() : TEXT(""), bIsGetter, InDefaultValue, InPosition, InNodeName, bSetupUndoRedo, bPrintPythonCommand);
 }
 
 UAnimNextSharedVariableNode* UAnimNextControllerBase::AddStructSharedVariableNode(const UScriptStruct* InStruct, const FName& InVariableName, const FString& InCPPType, UObject* InCPPTypeObject, bool bIsGetter, const FString& InDefaultValue, const FVector2D& InPosition, const FString& InNodeName, bool bSetupUndoRedo, bool bPrintPythonCommand)
@@ -28,7 +28,7 @@ UAnimNextSharedVariableNode* UAnimNextControllerBase::AddStructSharedVariableNod
 		return nullptr;
 	}
 
-	return AddSharedVariableNode(InStruct->GetPathName(), InVariableName, InCPPType, InCPPTypeObject->GetPathName(), bIsGetter, InDefaultValue, InPosition, InNodeName, bSetupUndoRedo, bPrintPythonCommand);
+	return AddSharedVariableNode(InStruct->GetPathName(), InVariableName, InCPPType, InCPPTypeObject ? InCPPTypeObject->GetPathName() : TEXT(""), bIsGetter, InDefaultValue, InPosition, InNodeName, bSetupUndoRedo, bPrintPythonCommand);
 }
 
 UAnimNextSharedVariableNode* UAnimNextControllerBase::AddSharedVariableNode(const FString& InSourceObjectPath, const FName& InVariableName, const FString& InCPPType, const FString& InCPPTypeObjectPath, bool bIsGetter, const FString& InDefaultValue, const FVector2D& InPosition, const FString& InNodeName, bool bSetupUndoRedo, bool bPrintPythonCommand)
@@ -85,7 +85,7 @@ UAnimNextSharedVariableNode* UAnimNextControllerBase::AddSharedVariableNode(cons
 		const FString GraphName = GetSchema()->GetSanitizedGraphName(GetGraph()->GetGraphName());
 
 		RigVMPythonUtils::Print(GetSchema()->GetGraphOuterName(GetGraph()), 
-			FString::Printf(TEXT("asset.get_controller_by_name('%s').add_shared_variable_node('%s', '%s', '%s', '%s', '%s', '%s', %s, %s)"),
+			FString::Printf(TEXT("asset.get_controller_by_name('%s').add_shared_variable_node('%s', '%s', '%s', '%s', '%s', '%s', %s, '%s')"),
 				*GraphName,
 				*InSourceObjectPath,
 				*InVariableName.ToString(),
