Fix for GetHostTransform not reporting the correct value until the actor has moved after spawn

diff --git a/Engine/Plugins/Experimental/UAF/UAF/Source/UAF/Private/Graph/RigUnit_GetActorTransform.cpp b/Engine/Plugins/Experimental/UAF/UAF/Source/UAF/Private/Graph/RigUnit_GetActorTransform.cpp
--- a/Engine/Plugins/Experimental/UAF/UAF/Source/UAF/Private/Graph/RigUnit_GetActorTransform.cpp
+++ b/Engine/Plugins/Experimental/UAF/UAF/Source/UAF/Private/Graph/RigUnit_GetActorTransform.cpp
@@ -13,7 +13,7 @@ void FAnimNextActorTransformComponent::OnBindToInstance()
 	UAnimNextComponent* AnimNextComponent = Cast<UAnimNextComponent>(ModuleInstance.GetObject());
 
 	FAnimNextModuleInstance::RunTaskOnGameThread(
-		[WeakAnimNextComponent = TWeakObjectPtr<UAnimNextComponent>(AnimNextComponent)]()
+		[this, WeakAnimNextComponent = TWeakObjectPtr<UAnimNextComponent>(AnimNextComponent)]()
 		{
 			UAnimNextComponent* Component = WeakAnimNextComponent.Get();
 			if(Component == nullptr)
@@ -33,6 +33,11 @@ void FAnimNextActorTransformComponent::OnBindToInstance()
 				return;
 			}
 
+			// Sync to current value
+			// Note here we assume its safe to capture 'this' since the lifetime of this lambda should be tied to the existence of the module component
+			// Components live until the system dies, so that implies that tasks queued on a system can access components on that system by raw ptr
+			ActorTransform = Owner->GetActorTransform();
+
 			SceneComponent->TransformUpdated.AddWeakLambda(Component, [Component](USceneComponent* UpdatedComponent, EUpdateTransformFlags UpdateTransformFlags, ETeleportType Teleport)
 			{
 				AActor* UpdatedComponentOwner = UpdatedComponent->GetOwner();
