Fix loctext key collision UAF Tagged asset browser filter

#tests nonunity

diff --git a/Engine/Plugins/Experimental/UAF/UAF/Source/UAFEditor/Private/Common/TaggedAssetBrowser_UAFBrowserFilters.cpp b/Engine/Plugins/Experimental/UAF/UAF/Source/UAFEditor/Private/Common/TaggedAssetBrowser_UAFBrowserFilters.cpp
--- a/Engine/Plugins/Experimental/UAF/UAF/Source/UAFEditor/Private/Common/TaggedAssetBrowser_UAFBrowserFilters.cpp
+++ b/Engine/Plugins/Experimental/UAF/UAF/Source/UAFEditor/Private/Common/TaggedAssetBrowser_UAFBrowserFilters.cpp
@@ -12,7 +12,7 @@
 
 #include UE_INLINE_GENERATED_CPP_BY_NAME(TaggedAssetBrowser_UAFBrowserFilters)
 
-#define LOCTEXT_NAMESPACE "UserAssetTags"
+#define LOCTEXT_NAMESPACE "UAFBrowserFilters"
 
 
 //////////////////////////////////////////////////////////////////////////
@@ -60,7 +60,7 @@ FText UTaggedAssetBrowserFilter_Skeleton::GetTooltip() const
 {
 	if (Skeleton)
 	{
-		FText BaseTooltip = LOCTEXT("UTaggedAssetBrowserFilter_ClassTooltip", "Skeleton Filter: {0}.");
+		FText BaseTooltip = LOCTEXT("UTaggedAssetBrowserFilter_SkeletonTooltip", "Skeleton Filter: {0}.");
 		return FText::FormatOrdered(BaseTooltip, FText::FromString(Skeleton->GetFName().ToString()));
 	}
 
