Fix duplicate internal variable issues where variables would be created at the leaf graph level and at the parent module / outer asset owning the shared variable + public function.
- Only add internal variables for a programmatic function if we own the function.

#tests editor, PIE

diff --git a/Engine/Plugins/Experimental/UAF/UAF/Source/UAFUncookedOnly/Private/AnimNextRigVMAssetEditorData.cpp b/Engine/Plugins/Experimental/UAF/UAF/Source/UAFUncookedOnly/Private/AnimNextRigVMAssetEditorData.cpp
--- a/Engine/Plugins/Experimental/UAF/UAF/Source/UAFUncookedOnly/Private/AnimNextRigVMAssetEditorData.cpp
+++ b/Engine/Plugins/Experimental/UAF/UAF/Source/UAFUncookedOnly/Private/AnimNextRigVMAssetEditorData.cpp
@@ -830,6 +830,13 @@ void UAnimNextRigVMAssetEditorData::BuildFunctionWrapperEventVariables(FAnimNext
 	for (const FAnimNextProgrammaticFunctionHeader& ProgrammaticFunctionHeader : InContext.FunctionHeaders)
 	{
 		const FRigVMGraphFunctionHeader& FunctionHeader = ProgrammaticFunctionHeader.FunctionHeader;
+
+		// Don't generate internal variables for functions we don't own. we will get those via shared variables
+		if (FunctionHeader.GetFunctionHost() != this)
+		{
+			continue;
+		}
+
 		for (const FRigVMGraphFunctionArgument& Argument : FunctionHeader.Arguments)
 		{
 			// Don't create internal variables for execution IO, visible only, or hidden pins. Those types do not need variable nodes in the graph.
