UAF: Owning/weak system references

Prerequisite for scene graph-like system reference semantics.
This allows us to better reason about the validity of weak system references without having access to implementation details.
The most significant changes are in reworking the world subsystem code into FSystemReference.
- Pooling is still used for allocation (via a shared ptr custom deleter), but 'handles' are now references that wrap a shared/weak ptr.
- 'Owning references' (FSystemReference) are specifically designed to be single-owner (move-only), despite using TSharedPtr under the hood.
- World subsystems are replaced with an engine subsystem and a single allocation pool, rather than one per world.
- Deferred operations like enable/disable/release now use a SPSC queue of operations rather than locking.

Also switches UAnimNextComponent to use Activate/Deactivate pattern inherited from actor components

diff --git a/Engine/Plugins/Experimental/MovieSceneAnimMixer/Source/MovieSceneAnimMixer/Private/Systems/MovieSceneAnimNextTargetSystem.cpp b/Engine/Plugins/Experimental/MovieSceneAnimMixer/Source/MovieSceneAnimMixer/Private/Systems/MovieSceneAnimNextTargetSystem.cpp
--- a/Engine/Plugins/Experimental/MovieSceneAnimMixer/Source/MovieSceneAnimMixer/Private/Systems/MovieSceneAnimNextTargetSystem.cpp
+++ b/Engine/Plugins/Experimental/MovieSceneAnimMixer/Source/MovieSceneAnimMixer/Private/Systems/MovieSceneAnimNextTargetSystem.cpp
@@ -41,7 +41,6 @@
 
 #include "PrimitiveSceneProxy.h"
 #include "Module/AnimNextModule.h"
-#include "UAF/Private/Component/AnimNextComponentWorldSubsystem.h"
 #include "UObject/FortniteMainBranchObjectVersion.h"
 
 #include UE_INLINE_GENERATED_CPP_BY_NAME(MovieSceneAnimNextTargetSystem)
@@ -92,9 +91,9 @@ namespace UE::MovieScene
 					// TODO the animation needs to be updated, or the post needs to be reset to ref pose when restoring to a previous module.
 					// Currently there is no support for this in UAF.
 					// This is most pronounced in the editor, where the mesh stays in the same pose when the sequence is closed.
-					FAnimMixerUtils::UnregisterWithUAFSubsystem(AnimNextComponent);
+					FAnimMixerUtils::UnregisterUAFSystem(AnimNextComponent);
 					FAnimMixerUtils::SetUAFComponentModule(AnimNextComponent, InOutCachedValue.PreAnimatedModule.Get());
-					FAnimMixerUtils::RegisterWithUAFSubsystem(AnimNextComponent);
+					FAnimMixerUtils::RegisterUAFSystem(AnimNextComponent);
 				}
 			}
 			
@@ -103,7 +102,7 @@ namespace UE::MovieScene
 			{
 				if (UAnimNextComponent* AnimNextComponent = Cast<UAnimNextComponent>(ObjectAndInjectionSite.Key.ResolveObjectPtr()))
 				{
-					AnimNextComponent->SetEnabled(false);
+					AnimNextComponent->SetActive(false);
 				}
 			}
 		}
@@ -268,9 +267,9 @@ namespace UE::MovieScene
 				{
 					const UMovieSceneAnimMixerSettings* AnimMixerSettings = GetDefault<UMovieSceneAnimMixerSettings>();
 
-					FAnimMixerUtils::UnregisterWithUAFSubsystem(AnimNextComponent);
+					FAnimMixerUtils::UnregisterUAFSystem(AnimNextComponent);
 					FAnimMixerUtils::SetUAFComponentModule(AnimNextComponent, AnimMixerSettings->DefaultUAFModule.LoadSynchronous());
-					FAnimMixerUtils::RegisterWithUAFSubsystem(AnimNextComponent);
+					FAnimMixerUtils::RegisterUAFSystem(AnimNextComponent);
 
 					bModuleReplaced = true;
 				}
@@ -300,7 +299,7 @@ namespace UE::MovieScene
 					PreAnimatedStorage->BeginTrackingEntity(EntityID, true, FRootInstanceHandle(), AnimNextComponent, InjectionSite.DesiredSite);
 					PreAnimatedSkelMeshCompStorage->BeginTrackingEntity(EntityID, true, FRootInstanceHandle(), SkeletalMeshComponent);
 
-					bool bAnimNextModuleEnabled = AnimNextComponent->IsEnabled();
+					bool bAnimNextModuleEnabled = AnimNextComponent->IsActive();
 
 					auto OnCacheAnimatedState = [this, TargetData, bAnimNextModuleEnabled, bModuleReplaced, PreAnimatedModule](UAnimNextComponent*)
 					{
@@ -319,7 +318,7 @@ namespace UE::MovieScene
 					// Force enable the anim next module if not currently enabled
 					if (!bAnimNextModuleEnabled)
 					{
-						AnimNextComponent->SetEnabled(true);
+						AnimNextComponent->SetActive(true);
 					}
 				}
 				else if(TargetData->Modifier.IsValid())
diff --git a/Engine/Plugins/Experimental/MovieSceneAnimMixer/Source/MovieSceneAnimMixer/Public/AnimMixerUtils.h b/Engine/Plugins/Experimental/MovieSceneAnimMixer/Source/MovieSceneAnimMixer/Public/AnimMixerUtils.h
--- a/Engine/Plugins/Experimental/MovieSceneAnimMixer/Source/MovieSceneAnimMixer/Public/AnimMixerUtils.h
+++ b/Engine/Plugins/Experimental/MovieSceneAnimMixer/Source/MovieSceneAnimMixer/Public/AnimMixerUtils.h
@@ -7,19 +7,19 @@
 struct FAnimMixerUtils
 {
 
-	static void RegisterWithUAFSubsystem(UAnimNextComponent* InComponent)
+	static void RegisterUAFSystem(UAnimNextComponent* InComponent)
 	{
 		if (InComponent)
 		{
-			InComponent->RegisterWithSubsystem();
+			InComponent->RegisterSystem();
 		}
 	}
 
-	static void UnregisterWithUAFSubsystem(UAnimNextComponent* InComponent)
+	static void UnregisterUAFSystem(UAnimNextComponent* InComponent)
 	{
 		if (InComponent)
 		{
-			InComponent->UnregisterWithSubsystem();
+			InComponent->UnregisterSystem();
 		}
 	}
 
diff --git a/Engine/Plugins/Experimental/UAF/UAF/Config/DefaultUAF.ini b/Engine/Plugins/Experimental/UAF/UAF/Config/DefaultUAF.ini
--- a/Engine/Plugins/Experimental/UAF/UAF/Config/DefaultUAF.ini
+++ b/Engine/Plugins/Experimental/UAF/UAF/Config/DefaultUAF.ini
@@ -296,3 +296,9 @@
 +StructRedirects=(OldName="/Script/SkeletonTemplateFrameworkRuntime.AbstractSkeleton_BoneBinding", NewName="/Script/UAF.AbstractSkeleton_BoneBinding")
 +StructRedirects=(OldName="/Script/SkeletonTemplateFrameworkRuntime.AbstractSkeleton_AttributeBinding", NewName="/Script/UAF.AbstractSkeleton_AttributeBinding")
 +StructRedirects=(OldName="/Script/SkeletonTemplateFrameworkRuntime.AbstractSkeletonSet", NewName="/Script/UAF.AbstractSkeletonSet")
+
+; System I/O
++StructRedirects=(OldName="/Script/UAF.AnimNextModuleHandle", NewName="/Script/UAF.UAFWeakSystemReference")
++FunctionRedirects=(OldName="AnimNextComponent.SetEnabled",NewName="ActorComponent.SetActive")
++FunctionRedirects=(OldName="AnimNextComponent.IsEnabled",NewName="ActorComponent.IsActive")
++FunctionRedirects=(OldName="AnimNextComponent.BlueprintGetModuleHandle",NewName="AnimNextComponent.GetSystemReference")
\ No newline at end of file
diff --git a/Engine/Plugins/Experimental/UAF/UAF/Source/UAF/Internal/AnimNextDebugDraw.h b/Engine/Plugins/Experimental/UAF/UAF/Source/UAF/Internal/AnimNextDebugDraw.h
--- a/Engine/Plugins/Experimental/UAF/UAF/Source/UAF/Internal/AnimNextDebugDraw.h
+++ b/Engine/Plugins/Experimental/UAF/UAF/Source/UAF/Internal/AnimNextDebugDraw.h
@@ -11,11 +11,11 @@
 class UObject;
 struct FRigVMDrawInterface;
 struct FAnimNextModuleInstance;
-class UAnimNextWorldSubsystem;
 
 namespace UE::UAF
 {
 	struct FModuleEventTickFunction;
+	struct FSystemReference;
 }
 
 namespace UE::UAF::Debug
@@ -90,7 +90,7 @@ private:
 
 	friend struct ::FAnimNextModuleInstance;
 	friend struct UE::UAF::FModuleEventTickFunction;
-	friend class ::UAnimNextWorldSubsystem;
+	friend UE::UAF::FSystemReference;
 };
 
 #endif
diff --git a/Engine/Plugins/Experimental/UAF/UAF/Source/UAF/Internal/Component/AnimNextWorldSubsystem.h b/Engine/Plugins/Experimental/UAF/UAF/Source/UAF/Internal/Component/AnimNextWorldSubsystem.h
deleted file mode 100644
diff --git a/Engine/Plugins/Experimental/UAF/UAF/Source/UAF/Internal/Module/AnimNextModuleInstance.h b/Engine/Plugins/Experimental/UAF/UAF/Source/UAF/Internal/Module/AnimNextModuleInstance.h
--- a/Engine/Plugins/Experimental/UAF/UAF/Source/UAF/Internal/Module/AnimNextModuleInstance.h
+++ b/Engine/Plugins/Experimental/UAF/UAF/Source/UAF/Internal/Module/AnimNextModuleInstance.h
@@ -4,7 +4,7 @@
 
 #include "CoreMinimal.h"
 #include "AnimNextModule.h"
-#include "Module/ModuleHandle.h"
+#include "AnimNextPoolHandle.h"
 #include "Variables/UAFInstanceVariableDataProxy.h"
 #include "UAFAssetInstance.h"
 #include "Module/ModuleTickFunction.h"
@@ -16,20 +16,25 @@
 #include "Module/TaskRunLocation.h"
 #include "AnimNextModuleInstance.generated.h"
 
+#define UE_API UAF_API
+
 struct FAnimNextModuleInstance;
 struct FRigUnit_AnimNextRunAnimationGraph_v1;
 struct FRigUnit_AnimNextRunAnimationGraph_v2;
-class UAnimNextWorldSubsystem;
+class UUAFEngineSubsystem;
 struct FRigUnit_AnimNextInitializeEvent;
 struct FRigUnit_CopyModuleProxyVariables;
 struct FAnimNextNotifyDispatcherComponent;
 struct FAnimNextGraphInstance;
 struct FUAFSystemFactoryParams;
+struct FUAFWeakSystemReference;
 
 namespace UE::UAF
 {
 	struct FModuleEventTickFunction;
 	struct FModuleEndTickFunction;
+	struct FSystemReference;
+	struct FCustomSystemDeleter;
 }
 
 namespace UE::UAF::Debug
@@ -48,96 +53,93 @@ namespace UE::UAF::Tests
 	class FVariables_UOLBindings;
 }
 
-// Root memory owner of a parameterized schedule 
+// Represents an instance of a UAF system.
 USTRUCT()
 struct FAnimNextModuleInstance : public FUAFAssetInstance
+#if CPP
+	, TSharedFromThis<FAnimNextModuleInstance>
+#endif
 {
 	GENERATED_BODY()
 
 	FAnimNextModuleInstance();
 
 	// Standalone constructor
-	UAF_API FAnimNextModuleInstance(
-	const UAnimNextModule* InModule,
+	UE_API FAnimNextModuleInstance(
+		const UAnimNextModule* InModule,
 		EAnimNextModuleInitMethod InInitMethod);
 
 	// Standalone constructor, UObject host
-	UAF_API FAnimNextModuleInstance(
-	const UAnimNextModule* InModule,
+	UE_API FAnimNextModuleInstance(
+		const UAnimNextModule* InModule,
 		UObject* InObject,
 		EAnimNextModuleInitMethod InInitMethod);
 
-	UAF_API ~FAnimNextModuleInstance();
+	UE_API ~FAnimNextModuleInstance();
 
 	// Checks to see if this entry is ticking
 	bool IsEnabled() const;
 
 	// Enables/disables the ticking of this entry
-	void Enable(bool bInEnabled);
+	void SetEnabled(bool bInEnabled);
 
 	// Queues an input trait event
 	// Input events will be processed in the next graph update after they are queued
-	UAF_API void QueueInputTraitEvent(FAnimNextTraitEventPtr Event);
+	UE_API void QueueInputTraitEvent(FAnimNextTraitEventPtr Event);
 
 	// Queues an output trait event
 	// Output events will be processed at the end of the schedule tick
-	UAF_API void QueueOutputTraitEvent(FAnimNextTraitEventPtr Event);
+	UE_API void QueueOutputTraitEvent(FAnimNextTraitEventPtr Event);
 
 	// Get the object that this module is bound to
 	UObject* GetObject() const { return Object; }
 
-	// Get the module that this instance represents
-	UAF_API const UAnimNextModule* GetModule() const;
+	// Get the system that this instance represents
+	UAF_API const UAnimNextModule* GetSystem() const;
 
 #if UE_ENABLE_DEBUG_DRAWING
 	// Get the debug draw interface
-	UAF_API FRigVMDrawInterface* GetDebugDrawInterface();
+	UE_API FRigVMDrawInterface* GetDebugDrawInterface();
 
 	// Set whether to show the module's debug drawing instructions in the viewport
-	UAF_API void ShowDebugDrawing(bool bInShowDebugDrawing);
+	UE_API void ShowDebugDrawing(bool bInShowDebugDrawing);
 #endif
 
 	// Run a simple task on the GT via FFunctionGraphTask::CreateAndDispatchWhenReady
-	UAF_API static void RunTaskOnGameThread(TUniqueFunction<void(void)>&& InFunction);
+	UE_API static void RunTaskOnGameThread(TUniqueFunction<void(void)>&& InFunction);
 
 	// Find the tick function for the specified event
 	UE::UAF::FModuleEventTickFunction* FindTickFunctionByName(FName InEventName);
 	const UE::UAF::FModuleEventTickFunction* FindTickFunctionByName(FName InEventName) const;
 
 	// Find the first 'user' tick function
-	UAF_API UE::UAF::FModuleEventTickFunction* FindFirstUserTickFunction();
+	UE_API UE::UAF::FModuleEventTickFunction* FindFirstUserTickFunction();
 
 	// Run the specified script event using a module context
-	UAF_API void RunScriptEvent(const UE::UAF::FScriptEvent& InEvent, float InDeltaTime);
+	UE_API void RunScriptEvent(const UE::UAF::FScriptEvent& InEvent, float InDeltaTime);
 
 	// Run the specified script event using a module context, using name lookup alone (slower)
-	UAF_API void RunScriptEventByName(FName InEventName, float InDeltaTime);
+	UE_API void RunScriptEventByName(FName InEventName, float InDeltaTime);
 
 	// Get the world type that this module was instantiated within
 	EWorldType::Type GetWorldType() const { return WorldType; }
 
 	// Get tick functions for this module instance
-	UAF_API TArrayView<UE::UAF::FModuleEventTickFunction> GetTickFunctions();
+	UE_API TArrayView<UE::UAF::FModuleEventTickFunction> GetTickFunctions();
 
-	UE::UAF::FModuleHandle GetHandle() const { return Handle; }
-
-	// Queue a task to run at a particular module event
-	// @param	InModuleEvent		The event in the module to run the supplied task relative to. If this is NAME_None or not found then the first 'user' event will be used.
-	// @param	InTaskFunction		The function to run
-	// @param	InLocation			Where to run the task, before or after
-	UAF_API void QueueTask(FName InEventName, TUniqueFunction<void(const UE::UAF::FModuleTaskContext&)>&& InTaskFunction, UE::UAF::ETaskRunLocation InLocation = UE::UAF::ETaskRunLocation::Before);
+	// Get a weak, non-owning reference to this instance
+	UE_API FUAFWeakSystemReference GetReference() const;
 
-	// Queue a task to run at a particular module event on some other module
-	// @param	InModuleHandle		The handle for the module to queue the task on
-	// @param	InModuleEventName	The event in the module to run the supplied task relative to. If this is NAME_None or not found then the first user 'event' will be used.
+	// Queue a task to run at a particular system event
+	// @param	InEventName			The event in the system to run the supplied task relative to. If this is NAME_None or not found then the first 'user' event will be used.
 	// @param	InTaskFunction		The function to run
 	// @param	InLocation			Where to run the task, before or after
-	UAF_API void QueueTaskOnOtherModule(const UE::UAF::FModuleHandle InOtherModuleHandle, FName InEventName, TUniqueFunction<void(const UE::UAF::FModuleTaskContext&)>&& InTaskFunction, UE::UAF::ETaskRunLocation InLocation = UE::UAF::ETaskRunLocation::Before);
+	UE_API void QueueTask(FName InEventName, TUniqueFunction<void(const UE::UAF::FModuleTaskContext&)>&& InTaskFunction, UE::UAF::ETaskRunLocation InLocation = UE::UAF::ETaskRunLocation::Before);
 
 private:
 	// Setup the entry
 	// Note that we perform initialization in two phases because we may need to access weak references to 'this' during initialization
-	UAF_API void Initialize(const FUAFSystemFactoryParams& InFactoryParams, UE::UAF::TPool<FAnimNextModuleInstance>* InPool = nullptr);
+	UE_API void Initialize(const FUAFSystemFactoryParams& InFactoryParams);
 
 	// Clear data that binds the schedule to a runtime (e.g. tick functions) and any instance data
 	void ResetBindingsAndInstanceData();
@@ -150,7 +152,7 @@ private:
 
 #if UAF_TRACE_ENABLED
 	// trace this module if it hasn't already been traced this frame
-	UAF_API void Trace();
+	UE_API void Trace();
 #endif
 
 #if WITH_EDITOR
@@ -172,28 +174,25 @@ private:
 	EPropertyBagResult WriteProxyVariable(const FAnimNextVariableReference& InVariable, const FAnimNextParamType& InType, TFunctionRef<void(TArrayView<uint8>)> InFunction);
 
 private:
+	friend UE::UAF::FSystemReference;
+	friend FUAFWeakSystemReference;
 	friend UE::UAF::FModuleEventTickFunction;
 	friend UE::UAF::FModuleEndTickFunction;
 	friend FRigUnit_AnimNextRunAnimationGraph_v1;
 	friend FRigUnit_AnimNextRunAnimationGraph_v2;
-	friend UAnimNextWorldSubsystem;
+	friend UUAFEngineSubsystem;
 	friend FRigUnit_AnimNextInitializeEvent;
 	friend FRigUnit_CopyModuleProxyVariables;
 	friend FAnimNextNotifyDispatcherComponent;
 	friend FAnimNextGraphInstance;
 	friend UE::UAF::Tests::FVariables_Shared;
 	friend UE::UAF::Tests::FVariables_UOLBindings;
+	friend UE::UAF::FCustomSystemDeleter;
 
 	// Object this entry is bound to
 	UPROPERTY(Transient)
 	TObjectPtr<UObject> Object = nullptr;
 
-	// The pool that this module instance exists in
-	UE::UAF::TPool<FAnimNextModuleInstance>* Pool = nullptr;
-
-	// Copy of the handle that represents this entry to client systems
-	UE::UAF::FModuleHandle Handle;
-
 	// Pre-allocated graph of tick functions
 	TArray<UE::UAF::FModuleEventTickFunction> TickFunctions;
 
@@ -292,3 +291,5 @@ struct TStructOpsTypeTraits<FAnimNextModuleInstance> : public TStructOpsTypeTrai
 #else
 #define TRACE_UAF_MODULE(ModuleInstance)
 #endif
+
+#undef UE_API
\ No newline at end of file
diff --git a/Engine/Plugins/Experimental/UAF/UAF/Source/UAF/Private/AnimNextModuleImpl.cpp b/Engine/Plugins/Experimental/UAF/UAF/Source/UAF/Private/AnimNextModuleImpl.cpp
--- a/Engine/Plugins/Experimental/UAF/UAF/Source/UAF/Private/AnimNextModuleImpl.cpp
+++ b/Engine/Plugins/Experimental/UAF/UAF/Source/UAF/Private/AnimNextModuleImpl.cpp
@@ -14,7 +14,7 @@
 #include "Curves/CurveFloat.h"
 #include "Graph/AnimNext_LODPose.h"
 #include "Module/AnimNextModuleInstance.h"
-#include "Module/ModuleHandle.h"
+#include "Module/UAFWeakSystemReference.h"
 #include "Modules/ModuleManager.h"
 #include "Param/AnimNextActorLocatorFragment.h"
 #include "Param/AnimNextComponentLocatorFragment.h"
@@ -78,7 +78,7 @@ namespace UE::UAF
 			FRigVMGraphFunctionHeader::StaticStruct(),
 			TBaseStructure<FGuid>::Get(),
 			FRigVMVariant::StaticStruct(),
-			FAnimNextModuleHandle::StaticStruct(),
+			FUAFWeakSystemReference::StaticStruct(),
 			FAnimNextGraphLODPose::StaticStruct(),
 			FAnimNextGraphReferencePose::StaticStruct(),
 		};
diff --git a/Engine/Plugins/Experimental/UAF/UAF/Source/UAF/Private/Component/AnimNextComponent.cpp b/Engine/Plugins/Experimental/UAF/UAF/Source/UAF/Private/Component/AnimNextComponent.cpp
--- a/Engine/Plugins/Experimental/UAF/UAF/Source/UAF/Private/Component/AnimNextComponent.cpp
+++ b/Engine/Plugins/Experimental/UAF/UAF/Source/UAF/Private/Component/AnimNextComponent.cpp
@@ -2,13 +2,14 @@
 
 #include "Component/AnimNextComponent.h"
 
-#include "AnimNextComponentWorldSubsystem.h"
 #include "Blueprint/BlueprintExceptionInfo.h"
 #include "Engine/World.h"
 #include "Module/ModuleTaskContext.h"
 #include "UObject/FortniteMainBranchObjectVersion.h"
 #include "Variables/AnimNextVariableReference.h"
 #include "Module/AnimNextModule.h"
+#include "Module/SystemDependency.h"
+#include "Module/UAFWeakSystemReference.h"
 
 #include UE_INLINE_GENERATED_CPP_BY_NAME(AnimNextComponent)
 
@@ -18,15 +19,13 @@ void UAnimNextComponent::OnRegister()
 
 	Super::OnRegister();
 
-	Subsystem = GetWorld()->GetSubsystem<UAnimNextComponentWorldSubsystem>();
-
-	if (Subsystem && Asset)
+	if (Asset)
 	{
-		check(!ModuleHandle.IsValid());
+		check(!SystemReference.IsValid());
 
-		Subsystem->Register(this);
+		SystemReference = FSystemReference(Asset, this, InitMethod);
 #if UE_ENABLE_DEBUG_DRAWING
-		Subsystem->ShowDebugDrawing(this, bShowDebugDrawing);
+		SystemReference.ShowDebugDrawing(bShowDebugDrawing);
 #endif
 	}
 }
@@ -35,11 +34,7 @@ void UAnimNextComponent::OnUnregister()
 {
 	Super::OnUnregister();
 
-	if(Subsystem)
-	{
-		Subsystem->Unregister(this);
-		Subsystem = nullptr;
-	}
+	SystemReference.Reset();
 }
 
 void UAnimNextComponent::BeginPlay()
@@ -52,7 +47,7 @@ void UAnimNextComponent::BeginPlay()
 	if (InitMethod == EAnimNextModuleInitMethod::InitializeAndRun
 	|| (InitMethod == EAnimNextModuleInitMethod::InitializeAndPauseInEditor && World->WorldType != EWorldType::Editor && World->WorldType != EWorldType::EditorPreview))
 	{
-		SetEnabled(true);
+		Activate();
 	}
 }
 
@@ -60,7 +55,7 @@ void UAnimNextComponent::EndPlay(const EEndPlayReason::Type EndPlayReason)
 {
 	Super::EndPlay(EndPlayReason);
 
-	SetEnabled(false);
+	Deactivate();
 }
 
 void UAnimNextComponent::Serialize(FArchive& Ar)
@@ -275,36 +270,31 @@ DEFINE_FUNCTION(UAnimNextComponent::execBlueprintSetVariableReference)
 
 EPropertyBagResult UAnimNextComponent::SetVariableInternal(const FAnimNextVariableReference& InVariable, const FAnimNextParamType& InType, TConstArrayView<uint8> InNewValue)
 {
-	if (Subsystem)
-	{
-		return Subsystem->SetVariable(this, InVariable, InType, InNewValue);
-	}
-	return EPropertyBagResult::PropertyNotFound;
+	return SystemReference.SetVariable(InVariable, InType, InNewValue);
 }
 
 EPropertyBagResult UAnimNextComponent::WriteVariableInternal(const FAnimNextVariableReference& InVariable, const FAnimNextParamType& InType, TFunctionRef<void(TArrayView<uint8>)> InFunction)
 {
-	if (Subsystem)
-	{
-		return Subsystem->WriteVariable(this, InVariable, InType, InFunction);
-	}
-	return EPropertyBagResult::PropertyNotFound;
+	return SystemReference.WriteVariable(InVariable, InType, InFunction);
 }
 
-bool UAnimNextComponent::IsEnabled() const
+void UAnimNextComponent::Activate(bool bReset)
 {
-	if (Subsystem)
+	if (bReset || ShouldActivate())
 	{
-		return Subsystem->IsEnabled(this);
+		SystemReference.SetEnabled(true);
+		SetActiveFlag(true);
+		OnComponentActivated.Broadcast(this, bReset);
 	}
-	return false;
 }
 
-void UAnimNextComponent::SetEnabled(bool bEnabled)
+void UAnimNextComponent::Deactivate()
 {
-	if(Subsystem)
+	if (!ShouldActivate())
 	{
-		Subsystem->SetEnabled(this, bEnabled);
+		SystemReference.SetEnabled(false);
+		SetActiveFlag(false);
+		OnComponentDeactivated.Broadcast(this);
 	}
 }
 
@@ -312,10 +302,7 @@ void UAnimNextComponent::ShowDebugDrawing(bool bInShowDebugDrawing)
 {
 #if UE_ENABLE_DEBUG_DRAWING
 	bShowDebugDrawing = bInShowDebugDrawing;
-	if(Subsystem)
-	{
-		Subsystem->ShowDebugDrawing(this, bInShowDebugDrawing);
-	}
+	SystemReference.ShowDebugDrawing(bInShowDebugDrawing);
 #endif
 }
 
@@ -323,37 +310,24 @@ void UAnimNextComponent::QueueTask(FName InModuleEventName, TUniqueFunction<void
 {
 	using namespace UE::UAF;
 
-	if(Subsystem)
-	{
-		Subsystem->QueueTask(this, InModuleEventName, MoveTemp(InTaskFunction), InLocation);
-	}
+	SystemReference.QueueTask(InModuleEventName, MoveTemp(InTaskFunction), InLocation);
 }
 
 void UAnimNextComponent::QueueInputTraitEvent(FAnimNextTraitEventPtr Event)
 {
 	using namespace UE::UAF;
 
-	if(Subsystem)
-	{
-		Subsystem->QueueInputTraitEvent(this, Event);
-	}
+	SystemReference.QueueInputTraitEvent(Event);
 }
 
 const FTickFunction* UAnimNextComponent::FindTickFunction(FName InEventName) const
 {
-	if(Subsystem)
-	{
-		return Subsystem->FindTickFunction(this, InEventName);
-	}
-	return nullptr;
+	return SystemReference.FindTickFunction(InEventName);
 }
 
 void UAnimNextComponent::AddPrerequisite(UObject* InObject, FTickFunction& InTickFunction, FName InEventName)
 {
-	if(Subsystem)
-	{
-		Subsystem->AddDependency(this, InObject, InTickFunction, InEventName, UAnimNextWorldSubsystem::EDependency::Prerequisite);
-	}
+	SystemReference.AddDependency(InObject, InTickFunction, InEventName, UE::UAF::ESystemDependency::Prerequisite);
 }
 
 void UAnimNextComponent::AddComponentPrerequisite(UActorComponent* InComponent, FName InEventName)
@@ -366,10 +340,7 @@ void UAnimNextComponent::AddComponentPrerequisite(UActorComponent* InComponent,
 
 void UAnimNextComponent::AddSubsequent(UObject* InObject, FTickFunction& InTickFunction, FName InEventName)
 {
-	if(Subsystem)
-	{
-		Subsystem->AddDependency(this, InObject, InTickFunction, InEventName, UAnimNextWorldSubsystem::EDependency::Subsequent);
-	}
+	SystemReference.AddDependency(InObject, InTickFunction, InEventName, UE::UAF::ESystemDependency::Subsequent);
 }
 
 void UAnimNextComponent::AddComponentSubsequent(UActorComponent* InComponent, FName InEventName)
@@ -382,10 +353,7 @@ void UAnimNextComponent::AddComponentSubsequent(UActorComponent* InComponent, FN
 
 void UAnimNextComponent::RemovePrerequisite(UObject* InObject, FTickFunction& InTickFunction, FName InEventName)
 {
-	if(Subsystem)
-	{
-		Subsystem->RemoveDependency(this, InObject, InTickFunction, InEventName, UAnimNextWorldSubsystem::EDependency::Prerequisite);
-	}
+	SystemReference.RemoveDependency(InObject, InTickFunction, InEventName, UE::UAF::ESystemDependency::Prerequisite);
 }
 
 void UAnimNextComponent::RemoveComponentPrerequisite(UActorComponent* InComponent, FName InEventName)
@@ -398,10 +366,7 @@ void UAnimNextComponent::RemoveComponentPrerequisite(UActorComponent* InComponen
 
 void UAnimNextComponent::RemoveSubsequent(UObject* InObject, FTickFunction& InTickFunction, FName InEventName)
 {
-	if(Subsystem)
-	{
-		Subsystem->RemoveDependency(this, InObject, InTickFunction, InEventName, UAnimNextWorldSubsystem::EDependency::Subsequent);
-	}
+	SystemReference.RemoveDependency(InObject, InTickFunction, InEventName, UE::UAF::ESystemDependency::Subsequent);
 }
 
 void UAnimNextComponent::RemoveComponentSubsequent(UActorComponent* InComponent, FName InEventName)
@@ -422,9 +387,9 @@ void UAnimNextComponent::AddModuleEventPrerequisite(FName InEventName, UAnimNext
 	{
 		UE_LOG(LogAnimation, Warning, TEXT("UAFComponent::AddModuleEventPrerequisite called using the same component"));
 	}
-	else if (Subsystem)
+	else
 	{
-		Subsystem->AddModuleEventDependency(this, InEventName, OtherAnimNextComponent, OtherEventName, UAnimNextWorldSubsystem::EDependency::Prerequisite);
+		SystemReference.AddSystemEventDependency(InEventName, OtherAnimNextComponent->GetSystemReference(), OtherEventName, UE::UAF::ESystemDependency::Prerequisite);
 	}
 }
 
@@ -438,9 +403,9 @@ void UAnimNextComponent::AddModuleEventSubsequent(FName InEventName, UAnimNextCo
 	{
 		UE_LOG(LogAnimation, Warning, TEXT("UAFComponent::AddModuleEventPrerequisite called using the same component"));
 	}
-	else if (Subsystem)
+	else
 	{
-		Subsystem->AddModuleEventDependency(this, InEventName, OtherAnimNextComponent, OtherEventName, UAnimNextWorldSubsystem::EDependency::Subsequent);
+		SystemReference.AddSystemEventDependency(InEventName, OtherAnimNextComponent->GetSystemReference(), OtherEventName, UE::UAF::ESystemDependency::Subsequent);
 	}
 }
 
@@ -454,9 +419,9 @@ void UAnimNextComponent::RemoveModuleEventPrerequisite(FName InEventName, UAnimN
 	{
 		UE_LOG(LogAnimation, Warning, TEXT("UAnimNextComponent::AddModuleEventPrerequisite called using the same component"));
 	}
-	else if (Subsystem)
+	else
 	{
-		Subsystem->RemoveModuleEventDependency(this, InEventName, OtherAnimNextComponent, OtherEventName, UAnimNextWorldSubsystem::EDependency::Prerequisite);
+		SystemReference.RemoveSystemEventDependency(InEventName, OtherAnimNextComponent->GetSystemReference(), OtherEventName, UE::UAF::ESystemDependency::Prerequisite);
 	}
 }
 
@@ -470,15 +435,15 @@ void UAnimNextComponent::RemoveModuleEventSubsequent(FName InEventName, UAnimNex
 	{
 		UE_LOG(LogAnimation, Warning, TEXT("UAnimNextComponent::AddModuleEventPrerequisite called using the same component"));
 	}
-	else if (Subsystem)
+	else
 	{
-		Subsystem->RemoveModuleEventDependency(this, InEventName, OtherAnimNextComponent, OtherEventName, UAnimNextWorldSubsystem::EDependency::Subsequent);
+		SystemReference.RemoveSystemEventDependency(InEventName, OtherAnimNextComponent->GetSystemReference(), OtherEventName, UE::UAF::ESystemDependency::Subsequent);
 	}
 }
 
-FAnimNextModuleHandle UAnimNextComponent::BlueprintGetModuleHandle() const
+FUAFWeakSystemReference UAnimNextComponent::GetSystemReference() const
 {
-	return FAnimNextModuleHandle(ModuleHandle);
+	return FUAFWeakSystemReference(SystemReference);
 }
 
 void UAnimNextComponent::SetAsset(TObjectPtr<UObject> InAsset)
@@ -486,25 +451,25 @@ void UAnimNextComponent::SetAsset(TObjectPtr<UObject> InAsset)
 	Asset = InAsset;
 }
 
-void UAnimNextComponent::RegisterWithSubsystem()
+void UAnimNextComponent::RegisterSystem()
 {
-	if (Subsystem && Asset)
+	using namespace UE::UAF;
+
+	if (Asset)
 	{
-		Subsystem->Register(this);
+		check(!SystemReference.IsValid());
+		SystemReference = FSystemReference(Asset, this, InitMethod);
 	}
 }
 
-void UAnimNextComponent::UnregisterWithSubsystem()
+void UAnimNextComponent::UnregisterSystem()
 {
-	if (Subsystem)
-	{
-		Subsystem->Unregister(this);
-	}
+	SystemReference.Reset();
 }
 
 bool UAnimNextComponent::IsModuleValid()
 {
-	return ModuleHandle.IsValid();
+	return SystemReference.IsValid();
 }
 
 TObjectPtr<UObject> UAnimNextComponent::GetAsset() const
diff --git a/Engine/Plugins/Experimental/UAF/UAF/Source/UAF/Private/Component/AnimNextComponentWorldSubsystem.cpp b/Engine/Plugins/Experimental/UAF/UAF/Source/UAF/Private/Component/AnimNextComponentWorldSubsystem.cpp
deleted file mode 100644
diff --git a/Engine/Plugins/Experimental/UAF/UAF/Source/UAF/Private/Component/AnimNextComponentWorldSubsystem.h b/Engine/Plugins/Experimental/UAF/UAF/Source/UAF/Private/Component/AnimNextComponentWorldSubsystem.h
deleted file mode 100644
diff --git a/Engine/Plugins/Experimental/UAF/UAF/Source/UAF/Private/Component/AnimNextWorldSubsystem.cpp b/Engine/Plugins/Experimental/UAF/UAF/Source/UAF/Private/Component/AnimNextWorldSubsystem.cpp
deleted file mode 100644
diff --git a/Engine/Plugins/Experimental/UAF/UAF/Source/UAF/Private/Module/AnimNextActorComponentReferenceComponent.cpp b/Engine/Plugins/Experimental/UAF/UAF/Source/UAF/Private/Module/AnimNextActorComponentReferenceComponent.cpp
--- a/Engine/Plugins/Experimental/UAF/UAF/Source/UAF/Private/Module/AnimNextActorComponentReferenceComponent.cpp
+++ b/Engine/Plugins/Experimental/UAF/UAF/Source/UAF/Private/Module/AnimNextActorComponentReferenceComponent.cpp
@@ -4,6 +4,7 @@
 
 #include "Component/AnimNextComponent.h"
 #include "Module/AnimNextModuleInstance.h"
+#include "Module/UAFWeakSystemReference.h"
 
 #include UE_INLINE_GENERATED_CPP_BY_NAME(AnimNextActorComponentReferenceComponent)
 
@@ -20,7 +21,7 @@ void FAnimNextActorComponentReferenceComponent::OnInitializeHelper(UScriptStruct
 	UAnimNextComponent* AnimNextComponent = Cast<UAnimNextComponent>(ModuleInstance->GetObject());
 
 	FAnimNextModuleInstance::RunTaskOnGameThread(
-		[WeakAnimNextComponent = TWeakObjectPtr<UAnimNextComponent>(AnimNextComponent), ComponentType = ComponentType, InScriptStruct]()
+		[Reference = ModuleInstance->GetReference(), WeakAnimNextComponent = TWeakObjectPtr<UAnimNextComponent>(AnimNextComponent), ComponentType = ComponentType, InScriptStruct]()
 		{
 			UAnimNextComponent* PinnedAnimNextComponent = WeakAnimNextComponent.Get();
 			if(PinnedAnimNextComponent == nullptr)
@@ -35,7 +36,7 @@ void FAnimNextActorComponentReferenceComponent::OnInitializeHelper(UScriptStruct
 			}
 
 			UActorComponent* FoundComponent = Owner->FindComponentByClass(ComponentType);
-			PinnedAnimNextComponent->QueueTask(NAME_None, [FoundComponent, InScriptStruct](const UE::UAF::FModuleTaskContext& InContext)
+			Reference.QueueTask(NAME_None, [FoundComponent, InScriptStruct](const UE::UAF::FModuleTaskContext& InContext)
 			{
 				InContext.TryAccessComponent(InScriptStruct, [FoundComponent](FUAFModuleInstanceComponent& InComponent)
 				{
diff --git a/Engine/Plugins/Experimental/UAF/UAF/Source/UAF/Private/Module/AnimNextModule.cpp b/Engine/Plugins/Experimental/UAF/UAF/Source/UAF/Private/Module/AnimNextModule.cpp
--- a/Engine/Plugins/Experimental/UAF/UAF/Source/UAF/Private/Module/AnimNextModule.cpp
+++ b/Engine/Plugins/Experimental/UAF/UAF/Source/UAF/Private/Module/AnimNextModule.cpp
@@ -1,6 +1,7 @@
 // Copyright Epic Games, Inc. All Rights Reserved.
 
 #include "Module/AnimNextModule.h"
+#include "AnimNextPool.h"
 #include "RigVMCore/RigVMExecuteContext.h"
 #include "UObject/AssetRegistryTagsContext.h"
 #include "UObject/ObjectSaveContext.h"
diff --git a/Engine/Plugins/Experimental/UAF/UAF/Source/UAF/Private/Module/AnimNextModuleInstance.cpp b/Engine/Plugins/Experimental/UAF/UAF/Source/UAF/Private/Module/AnimNextModuleInstance.cpp
--- a/Engine/Plugins/Experimental/UAF/UAF/Source/UAF/Private/Module/AnimNextModuleInstance.cpp
+++ b/Engine/Plugins/Experimental/UAF/UAF/Source/UAF/Private/Module/AnimNextModuleInstance.cpp
@@ -18,6 +18,8 @@
 #include "Misc/EnumerateRange.h"
 #include "Module/RigUnit_AnimNextModuleEvents.h"
 #include "Module/ModuleTickFunction.h"
+#include "Module/SystemReference.h"
+#include "Module/UAFWeakSystemReference.h"
 #include "UObject/UObjectIterator.h"
 #include "RewindDebugger/UAFTrace.h"
 #include "Script/UAFScriptComponent.h"
@@ -70,10 +72,9 @@ FAnimNextModuleInstance::~FAnimNextModuleInstance()
 
 	Object = nullptr;
 	Asset = nullptr;
-	Handle.Reset();
 }
 
-void FAnimNextModuleInstance::Initialize(const FUAFSystemFactoryParams& InFactoryParams, UE::UAF::TPool<FAnimNextModuleInstance>* InPool)
+void FAnimNextModuleInstance::Initialize(const FUAFSystemFactoryParams& InFactoryParams)
 {
 	SCOPE_CYCLE_COUNTER(STAT_AnimNext_InitializeInstance);
 	
@@ -81,13 +82,11 @@ void FAnimNextModuleInstance::Initialize(const FUAFSystemFactoryParams& InFactor
 
 	check(IsInGameThread());
 
-	Pool = InPool;
-
-	const UAnimNextModule* Module = GetModule();
-	if (Module)
+	const UAnimNextModule* System = GetSystem();
+	if (System)
 	{
 #if WITH_EDITOR
-		if (Module->CompilationState == EAnimNextRigVMAssetState::CompiledWithErrors)
+		if (System->CompilationState == EAnimNextRigVMAssetState::CompiledWithErrors)
 		{
 			return;
 		}
@@ -158,7 +157,7 @@ void FAnimNextModuleInstance::Initialize(const FUAFSystemFactoryParams& InFactor
 				PrevTickFunction = &TickFunction;
 
 				// Set up dependencies, if any
-				for (const TInstancedStruct<FRigVMTrait_ModuleEventDependency>& DependencyInstance : Module->Dependencies)
+				for (const TInstancedStruct<FRigVMTrait_ModuleEventDependency>& DependencyInstance : System->Dependencies)
 				{
 					const FRigVMTrait_ModuleEventDependency* Dependency = DependencyInstance.GetPtr<FRigVMTrait_ModuleEventDependency>();
 					if (Dependency != nullptr && Dependency->EventName == ModuleEvent.Event.GetEventName())
@@ -261,7 +260,7 @@ void FAnimNextModuleInstance::Initialize(const FUAFSystemFactoryParams& InFactor
 		// Just pause now if we arent needing an initial update
 		if(InitMethod == EAnimNextModuleInitMethod::None)
 		{
-			Enable(false);
+			SetEnabled(false);
 		}
 #if WITH_EDITOR
 		else if(World)
@@ -336,12 +335,18 @@ bool FAnimNextModuleInstance::IsEnabled() const
 	return RunState == ERunState::Running;
 }
 
-void FAnimNextModuleInstance::Enable(bool bInEnabled)
+void FAnimNextModuleInstance::SetEnabled(bool bInEnabled)
 {
 	using namespace UE::UAF;
 
 	check(IsInGameThread());
 
+	// Ensure that we synchronize activation state with actor component if we are hosted there
+	if (UActorComponent* Component = Cast<UActorComponent>(Object))
+	{
+		Component->SetActiveFlag(bInEnabled);
+	}
+
 	if(RunState == ERunState::Paused || RunState == ERunState::Running)
 	{
 		for (FModuleEventTickFunction& TickFunction : TickFunctions)
@@ -431,7 +436,7 @@ void FAnimNextModuleInstance::Trace()
 }
 #endif
 
-const UAnimNextModule* FAnimNextModuleInstance::GetModule() const
+const UAnimNextModule* FAnimNextModuleInstance::GetSystem() const
 {
 	return CastChecked<UAnimNextModule>(Asset, ECastCheckedType::NullAllowed);
 }
@@ -578,6 +583,13 @@ TArrayView<UE::UAF::FModuleEventTickFunction> FAnimNextModuleInstance::GetTickFu
 	return TickFunctions;
 }
 
+FUAFWeakSystemReference FAnimNextModuleInstance::GetReference() const
+{
+	FUAFWeakSystemReference Handle;
+	Handle.WeakPtr = SharedThis(const_cast<FAnimNextModuleInstance*>(this));
+	return Handle;
+}
+
 UE::UAF::FModuleEventTickFunction* FAnimNextModuleInstance::FindFirstUserTickFunction()
 {
 	UE::UAF::FModuleEventTickFunction* FoundTickFunction = TickFunctions.FindByPredicate([](const UE::UAF::FModuleEventTickFunction& InTickFunction)
@@ -637,19 +649,3 @@ void FAnimNextModuleInstance::QueueTask(FName InEventName, TUniqueFunction<void(
 		UE_LOGFMT(LogAnimation, Warning, "QueueTask: Could not find event '{EventName}' in module '{ModuleName}'", InEventName, GetAssetName());
 	}
 }
-
-void FAnimNextModuleInstance::QueueTaskOnOtherModule(const UE::UAF::FModuleHandle InOtherModuleHandle, const FName InEventName, TUniqueFunction<void(const UE::UAF::FModuleTaskContext&)>&& InTaskFunction, UE::UAF::ETaskRunLocation InLocation)
-{
-	if (Pool == nullptr)
-	{
-		return;
-	}
-
-	FAnimNextModuleInstance* OtherModuleInstance = Pool->TryGet(InOtherModuleHandle);
-	if (OtherModuleInstance == nullptr)
-	{
-		return;
-	}
-
-	OtherModuleInstance->QueueTask(InEventName, MoveTemp(InTaskFunction), InLocation);
-}
diff --git a/Engine/Plugins/Experimental/UAF/UAF/Source/UAF/Private/Module/ModuleInstanceUtils.cpp b/Engine/Plugins/Experimental/UAF/UAF/Source/UAF/Private/Module/ModuleInstanceUtils.cpp
deleted file mode 100644
diff --git a/Engine/Plugins/Experimental/UAF/UAF/Source/UAF/Private/Module/ModuleTickFunction.cpp b/Engine/Plugins/Experimental/UAF/UAF/Source/UAF/Private/Module/ModuleTickFunction.cpp
--- a/Engine/Plugins/Experimental/UAF/UAF/Source/UAF/Private/Module/ModuleTickFunction.cpp
+++ b/Engine/Plugins/Experimental/UAF/UAF/Source/UAF/Private/Module/ModuleTickFunction.cpp
@@ -6,12 +6,13 @@
 #include "AnimNextStats.h"
 #include "RigVMRuntimeDataRegistry.h"
 #include "Algo/TopologicalSort.h"
-#include "Component/AnimNextWorldSubsystem.h"
 #include "Module/AnimNextModule.h"
 #include "Module/AnimNextModuleInstance.h"
 #include "Module/ModuleTaskContext.h"
 #include "Module/RigUnit_AnimNextModuleEvents.h"
 #include "Module/ModuleEvents.h"
+#include "Module/SystemReference.h"
+#include "Module/UAFWeakSystemReference.h"
 #include "TraitCore/TraitEventList.h"
 
 namespace UE::UAF
@@ -92,16 +93,8 @@ void FModuleEventTickFunction::EndTick(float DeltaTime)
 #endif
 		)
 		{
-			// Queue task to disable ourselves
-			RunTaskOnGameThread([ModuleHandle = ModuleInstance->Handle, WeakObject = TWeakObjectPtr<UObject>(ModuleInstance->Object)]()
-			{
-				check(IsInGameThread());
-				if (UObject* Object = WeakObject.Get())
-				{
-					UAnimNextWorldSubsystem* WorldSubsystem = UAnimNextWorldSubsystem::Get(Object);
-					WorldSubsystem->EnableHandle(ModuleHandle, false);
-				}
-			});
+			// Disable ourselves (deferred)
+			FSystemReference(*ModuleInstance).SetEnabled(false);
 		}
 	}
 
diff --git a/Engine/Plugins/Experimental/UAF/UAF/Source/UAF/Private/Module/RigUnit_AnimNextModuleEvents.cpp b/Engine/Plugins/Experimental/UAF/UAF/Source/UAF/Private/Module/RigUnit_AnimNextModuleEvents.cpp
--- a/Engine/Plugins/Experimental/UAF/UAF/Source/UAF/Private/Module/RigUnit_AnimNextModuleEvents.cpp
+++ b/Engine/Plugins/Experimental/UAF/UAF/Source/UAF/Private/Module/RigUnit_AnimNextModuleEvents.cpp
@@ -2,7 +2,6 @@
 
 #include "Module/RigUnit_AnimNextModuleEvents.h"
 
-#include "Component/AnimNextWorldSubsystem.h"
 #include "Engine/World.h"
 #include "Module/AnimNextModuleInstance.h"
 #include "Script/ScriptEventTickFunctionBindings.h"
diff --git a/Engine/Plugins/Experimental/UAF/UAF/Source/UAF/Private/Module/SystemReference.cpp b/Engine/Plugins/Experimental/UAF/UAF/Source/UAF/Private/Module/SystemReference.cpp
new file mode 100644
--- /dev/null
+++ b/Engine/Plugins/Experimental/UAF/UAF/Source/UAF/Private/Module/SystemReference.cpp
@@ -0,0 +1,439 @@
+﻿// Copyright Epic Games, Inc. All Rights Reserved.
+
+#include "Module/SystemReference.h"
+
+#include "AnimNextDebugDraw.h"
+#include "AnimNextPool.h"
+#include "AnimNextPoolHandle.h"
+#include "AnimNextRigVMAsset.h"
+#include "Engine/World.h"
+#include "Factory/SystemFactory.h"
+#include "Factory/UAFSystemFactoryParams.h"
+#include "Logging/StructuredLog.h"
+#include "Module/AnimNextModuleInstance.h"
+#include "Module/ModuleTaskContext.h"
+#include "Module/SystemDependency.h"
+#include "Module/UAFWeakSystemReference.h"
+
+namespace UE::UAF::Private
+{
+	struct FSystemPool : public FGCObject
+	{
+		// Global pool of system instances
+		TPool<FAnimNextModuleInstance> SystemInstances;
+
+		// Lock for concurrent allocations/resolves
+		FRWLock SystemInstancesLock;
+
+		// Actions to perform once world tick is completed each frame
+		TSpscQueue<TUniqueFunction<void()>> PendingActions;
+
+		// Deferred cleanup of handles that are released this frame
+		void FlushPendingActions()
+		{
+			using namespace UE::UAF;
+
+			while (!PendingActions.IsEmpty())
+			{
+				TUniqueFunction<void()> PendingAction;
+				if (PendingActions.Dequeue(PendingAction))
+				{
+					PendingAction();
+				}
+			}
+		}
+
+		// FGCObject interface
+		virtual void AddReferencedObjects(FReferenceCollector& Collector) override
+		{
+			UE::TReadScopeLock InstancesLockScope(SystemInstancesLock);
+
+			for (FAnimNextModuleInstance& Instance : SystemInstances)
+			{
+				Collector.AddPropertyReferencesWithStructARO(FAnimNextModuleInstance::StaticStruct(), &Instance);
+			}
+		}
+
+		virtual FString GetReferencerName() const override
+		{
+			return TEXT("UAFSystemPool");
+		}
+	};
+
+	static FSystemPool GPool;
+}
+
+namespace UE::UAF
+{
+
+struct FCustomSystemDeleter
+{
+	explicit FCustomSystemDeleter(TPoolHandle<FAnimNextModuleInstance> InHandle)
+		: Handle(InHandle)
+	{
+	}
+
+	void operator()(FAnimNextModuleInstance* InInstance) const
+	{
+		using namespace UE::UAF;
+		check(IsInGameThread());
+		if (InInstance)
+		{
+			check(InInstance == &Private::GPool.SystemInstances.Get(Handle));
+
+			Private::GPool.PendingActions.Enqueue([Handle = Handle]()
+			{
+				UE::TWriteScopeLock InstancesLockScope(Private::GPool.SystemInstancesLock);
+				Private::GPool.SystemInstances.Release(Handle);
+			});
+		}
+	}
+
+	// Handle used for pooled allocation
+	TPoolHandle<FAnimNextModuleInstance> Handle;
+};
+
+FSystemReference::FSystemReference(FAnimNextModuleInstance& InInstance)
+	: Ptr(InInstance.AsShared())
+{
+}
+
+FSystemReference::FSystemReference(const UObject* InAsset, UObject* InObject, EAnimNextModuleInitMethod InitMethod)
+{
+	check(IsInGameThread());
+	{
+		const UAnimNextModule* Module = Cast<UAnimNextModule>(InAsset);
+		FUAFSystemFactoryParams FactoryParams;
+		if (Module == nullptr)
+		{
+			FactoryParams = FSystemFactory::GetDefaultParamsForObject(InAsset);
+			Module = FSystemFactory::BuildSystem(FactoryParams.GetBuilder());
+		}
+
+		FAnimNextModuleInstance* Instance = nullptr;
+		TPoolHandle<FAnimNextModuleInstance> Handle;
+		{
+			UE::TWriteScopeLock InstancesLockScope(Private::GPool.SystemInstancesLock);
+			Handle = Private::GPool.SystemInstances.Emplace(Module, InObject, InitMethod);
+			Instance = &Private::GPool.SystemInstances.Get(Handle);
+		}
+		Ptr = TSharedPtr<FAnimNextModuleInstance>(Instance, FCustomSystemDeleter(Handle));
+		Instance->Initialize(FactoryParams);
+	}
+}
+
+FSystemReference::FSystemReference(FSystemReference&& InOther)
+{
+	Ptr = MoveTemp(InOther.Ptr);
+}
+
+FSystemReference& FSystemReference::operator=(FSystemReference&& InOther)
+{
+	Ptr = MoveTemp(InOther.Ptr);
+	return *this;
+}
+
+void FSystemReference::Reset()
+{
+	if (Ptr)
+	{
+#if UE_ENABLE_DEBUG_DRAWING
+		// Remove debug drawing immediately as the renderer will need to know about this before EOF
+		Ptr->DebugDraw->RemovePrimitive();
+#endif
+
+		// Remove all tick dependencies immediately, as once the handle has been invalidated there is no way for external systems to remove their dependencies
+		Ptr->RemoveAllTickDependencies();
+		Ptr.Reset();
+	}
+}
+
+bool FSystemReference::IsValid() const
+{
+	return Ptr.IsValid();
+}
+
+const UAnimNextModule* FSystemReference::GetSystem() const
+{
+	if (IsValid())
+	{
+		return Ptr->GetSystem();
+	}
+	return nullptr;
+}
+
+void FSystemReference::SetEnabled(bool bInEnabled) const
+{
+	using namespace UE::UAF;
+	if (IsValid())
+	{
+		Private::GPool.PendingActions.Enqueue([WeakPtr = TWeakPtr<FAnimNextModuleInstance>(Ptr), bInEnabled]()
+		{
+			if (TSharedPtr<FAnimNextModuleInstance> PinnedPtr = WeakPtr.Pin())
+			{
+				PinnedPtr->SetEnabled(bInEnabled);
+			}
+		});
+	}
+}
+
+#if UE_ENABLE_DEBUG_DRAWING
+void FSystemReference::ShowDebugDrawing(bool bInShowDebugDrawing) const
+{
+	using namespace UE::UAF;
+	if (IsValid())
+	{
+		Private::GPool.PendingActions.Enqueue([WeakPtr = TWeakPtr<FAnimNextModuleInstance>(Ptr), bInShowDebugDrawing]()
+		{
+			if (TSharedPtr<FAnimNextModuleInstance> PinnedPtr = WeakPtr.Pin())
+			{
+				PinnedPtr->ShowDebugDrawing(bInShowDebugDrawing);
+			}
+		});
+	}
+}
+#endif
+
+void FSystemReference::QueueTask(FName InSystemEventName, TUniqueFunction<void(const UE::UAF::FModuleTaskContext&)>&& InTaskFunction, UE::UAF::ETaskRunLocation InLocation) const
+{
+	using namespace UE::UAF;
+	if (IsValid())
+	{
+		Ptr->QueueTask(InSystemEventName, MoveTemp(InTaskFunction), InLocation);
+	}
+}
+
+void FSystemReference::QueueInputTraitEvent(FAnimNextTraitEventPtr Event) const
+{
+	using namespace UE::UAF;
+	QueueTask(NAME_None, [Event = MoveTemp(Event)](const FModuleTaskContext& InContext)
+	{
+		InContext.QueueInputTraitEvent(Event);
+	},
+	ETaskRunLocation::Before);
+}
+
+const FTickFunction* FSystemReference::FindTickFunction(FName InEventName) const
+{
+	using namespace UE::UAF;
+	check(IsInGameThread());
+	if (!IsValid())
+	{
+		UE_LOGFMT(LogAnimation, Warning, "FindTickFunction: Invalid system handle");
+		return nullptr;
+	}
+
+	const FModuleEventTickFunction* TickFunction = Ptr->FindTickFunctionByName(InEventName);
+	if (TickFunction == nullptr)
+	{
+		UE_LOGFMT(LogAnimation, Warning, "FindTickFunction: Could not find event '{EventName}' in module '{ModuleName}'", InEventName, Ptr->GetAssetName());
+		return nullptr;
+	}
+
+	if (TickFunction->bUserEvent)
+	{
+		return TickFunction;
+	}
+
+	UE_LOGFMT(LogAnimation, Warning, "FindTickFunction: Event '{EventName}' in module '{ModuleName}' is not a bUserEvent, therefore cannot be exposed", InEventName, Ptr->GetAssetName());
+	return nullptr;
+}
+
+void FSystemReference::AddDependency(UObject* InObject, FTickFunction& InTickFunction, FName InEventName, UE::UAF::ESystemDependency InDependency) const
+{
+	using namespace UE::UAF;
+	check(IsInGameThread());
+	if (!IsValid())
+	{
+		UE_LOGFMT(LogAnimation, Warning, "AddDependency: Invalid system handle");
+		return;
+	}
+
+	FModuleEventTickFunction* TickFunction = Ptr->FindTickFunctionByName(InEventName);
+	if (TickFunction == nullptr)
+	{
+		UE_LOGFMT(LogAnimation, Warning, "AddDependency: Could not find event '{EventName}' in module '{ModuleName}'", InEventName, Ptr->GetAssetName());
+		return;
+	}
+
+	if (InDependency == ESystemDependency::Prerequisite)
+	{
+		TickFunction->AddPrerequisite(InObject, InTickFunction);
+	}
+	else
+	{
+		TickFunction->AddSubsequent(InObject, InTickFunction);
+	}
+}
+
+void FSystemReference::RemoveDependency(UObject* InObject, FTickFunction& InTickFunction, FName InEventName, UE::UAF::ESystemDependency InDependency) const
+{
+	using namespace UE::UAF;
+	check(IsInGameThread());
+	if (!IsValid())
+	{
+		UE_LOGFMT(LogAnimation, Warning, "RemoveDependency: Invalid system handle");
+		return;
+	}
+
+	FModuleEventTickFunction* TickFunction = Ptr->FindTickFunctionByName(InEventName);
+	if (TickFunction == nullptr)
+	{
+		UE_LOGFMT(LogAnimation, Warning, "RemoveDependency: Could not find event '{EventName}' in module '{ModuleName}'", InEventName, Ptr->GetAssetName());
+		return;
+	}
+
+	if (InDependency == ESystemDependency::Prerequisite)
+	{
+		TickFunction->RemovePrerequisite(InObject, InTickFunction);
+	}
+	else
+	{
+		TickFunction->RemoveSubsequent(InObject, InTickFunction);
+	}
+}
+
+void FSystemReference::AddSystemEventDependency(FName InSystemEventName, FUAFWeakSystemReference InOtherSystem, FName InOtherSystemEventName, UE::UAF::ESystemDependency InDependency) const
+{
+	using namespace UE::UAF;
+	check(IsInGameThread());
+	if (!IsValid())
+	{
+		UE_LOGFMT(LogAnimation, Warning, "AddDependency: Invalid system handle");
+		return;
+	}
+
+	TSharedPtr<FAnimNextModuleInstance> OtherInstance = InOtherSystem.WeakPtr.Pin();
+	if (!OtherInstance.IsValid())
+	{
+		UE_LOGFMT(LogAnimation, Warning, "AddDependency: Invalid other system reference");
+		return;
+	}
+
+	if(Ptr == OtherInstance)
+	{
+		UE_LOGFMT(LogAnimation, Warning, "AddDependency: Cannot add a dependency on ourselves");
+		return;
+	}
+
+	FModuleEventTickFunction* TickFunction = OtherInstance->FindTickFunctionByName(InOtherSystemEventName);
+	if (TickFunction == nullptr)
+	{
+		UE_LOGFMT(LogAnimation, Warning, "AddDependency: Could not find event '{EventName}' in module '{ModuleName}'", InOtherSystemEventName, OtherInstance->GetAssetName());
+		return;
+	}
+
+	AddDependency(OtherInstance->GetObject(), *TickFunction, InSystemEventName, InDependency);
+}
+
+void FSystemReference::RemoveSystemEventDependency(FName InSystemEventName, FUAFWeakSystemReference InOtherSystem, FName InOtherSystemEventName, UE::UAF::ESystemDependency InDependency) const
+{
+	using namespace UE::UAF;
+	check(IsInGameThread());
+	if (!IsValid())
+	{
+		return;
+	}
+
+	TSharedPtr<FAnimNextModuleInstance> OtherInstance = InOtherSystem.WeakPtr.Pin();
+	if (!OtherInstance.IsValid())
+	{
+		UE_LOGFMT(LogAnimation, Warning, "RemoveDependency: Invalid other system reference");
+		return;
+	}
+
+	if(Ptr == OtherInstance)
+	{
+		UE_LOGFMT(LogAnimation, Warning, "RemoveDependency: Cannot remove a dependency on ourselves");
+		return;
+	}
+
+	FModuleEventTickFunction* TickFunction = OtherInstance->FindTickFunctionByName(InOtherSystemEventName);
+	if (TickFunction == nullptr)
+	{
+		UE_LOGFMT(LogAnimation, Warning, "RemoveDependency: Could not find event '{EventName}' in module '{ModuleName}'", InOtherSystemEventName, OtherInstance->GetAssetName());
+		return;
+	}
+
+	RemoveDependency(OtherInstance->GetObject(), *TickFunction, InSystemEventName, InDependency);
+}
+
+EPropertyBagResult FSystemReference::SetVariable(const FAnimNextVariableReference& InVariable, const FAnimNextParamType& InType, TConstArrayView<uint8> InData) const
+{
+	using namespace UE::UAF;
+	if (IsValid())
+	{
+		return Ptr->SetProxyVariable(InVariable, InType, InData);
+	}
+	return EPropertyBagResult::PropertyNotFound;
+}
+
+EPropertyBagResult FSystemReference::WriteVariable(const FAnimNextVariableReference& InVariable, const FAnimNextParamType& InType, TFunctionRef<void(TArrayView<uint8>)> InFunction) const
+{
+	using namespace UE::UAF;
+	if (IsValid())
+	{
+		return Ptr->WriteProxyVariable(InVariable, InType, InFunction);
+	}
+	return EPropertyBagResult::PropertyNotFound;
+}
+
+}
+
+UUAFEngineSubsystem::UUAFEngineSubsystem()
+{
+	if (!HasAnyFlags(RF_ClassDefaultObject))
+	{
+#if WITH_EDITOR
+		OnCompileJobFinishedHandle = UAnimNextRigVMAsset::OnCompileJobFinished().AddUObject(this, &UUAFEngineSubsystem::OnCompileJobFinished);
+#endif
+
+		OnPreGarbageCollectHandle = FCoreUObjectDelegates::GetPreGarbageCollectDelegate().AddLambda([]()
+		{
+			// Ensure we flush any pending releases before we perform GC, as we could
+			// end up with spurious references to already-released systems
+			UE::UAF::Private::GPool.FlushPendingActions();
+		});
+
+		OnWorldPreActorTickHandle = FWorldDelegates::OnWorldPreActorTick.AddLambda([this](UWorld* InWorld, ELevelTick InTickType, float InDeltaSeconds)
+		{
+			if (InTickType == LEVELTICK_All || InTickType == LEVELTICK_ViewportsOnly)
+			{
+				UE::UAF::Private::GPool.FlushPendingActions();
+			}
+		});
+	}
+}
+
+void UUAFEngineSubsystem::BeginDestroy()
+{
+	Super::BeginDestroy();
+
+	if (!HasAnyFlags(RF_ClassDefaultObject))
+	{
+#if WITH_EDITOR
+		UAnimNextRigVMAsset::OnCompileJobFinished().Remove(OnCompileJobFinishedHandle);
+#endif
+
+		FCoreUObjectDelegates::GetPreGarbageCollectDelegate().Remove(OnPreGarbageCollectHandle);
+		FWorldDelegates::OnWorldPreActorTick.Remove(OnWorldPreActorTickHandle);
+	}
+}
+
+#if WITH_EDITOR
+void UUAFEngineSubsystem::OnCompileJobFinished(UAnimNextRigVMAsset* InAsset)
+{
+	using namespace UE::UAF;
+
+	check(IsInGameThread());
+	UE::TReadScopeLock InstancesLockScope(Private::GPool.SystemInstancesLock);
+
+	for (FAnimNextModuleInstance& Instance : Private::GPool.SystemInstances)
+	{
+		if (Instance.GetSystem() == InAsset)
+		{
+			Instance.OnCompileJobFinished();
+		}
+	}
+}
+#endif
\ No newline at end of file
diff --git a/Engine/Plugins/Experimental/UAF/UAF/Source/UAF/Private/Module/UAFWeakSystemReference.cpp b/Engine/Plugins/Experimental/UAF/UAF/Source/UAF/Private/Module/UAFWeakSystemReference.cpp
new file mode 100644
--- /dev/null
+++ b/Engine/Plugins/Experimental/UAF/UAF/Source/UAF/Private/Module/UAFWeakSystemReference.cpp
@@ -0,0 +1,48 @@
+﻿// Copyright Epic Games, Inc. All Rights Reserved.
+
+#include "Module/UAFWeakSystemReference.h"
+#include "Module/AnimNextModuleInstance.h"
+#include "Module/SystemReference.h"
+
+FUAFWeakSystemReference::FUAFWeakSystemReference(const UE::UAF::FSystemReference& InHandle)
+	: WeakPtr(InHandle.Ptr)
+{
+}
+
+void FUAFWeakSystemReference::QueueTask(FName InSystemEventName, TUniqueFunction<void(const UE::UAF::FModuleTaskContext&)>&& InTaskFunction, UE::UAF::ETaskRunLocation InLocation) const
+{
+	TSharedPtr<FAnimNextModuleInstance> Instance = WeakPtr.Pin();
+	if (ensure(Instance.IsValid()))
+	{
+		Instance->QueueTask(InSystemEventName, MoveTemp(InTaskFunction), InLocation);
+	}
+}
+
+void FUAFWeakSystemReference::QueueInputTraitEvent(TSharedPtr<FAnimNextTraitEvent> InEvent) const
+{
+	TSharedPtr<FAnimNextModuleInstance> Instance = WeakPtr.Pin();
+	if (ensure(Instance.IsValid()))
+	{
+		Instance->QueueInputTraitEvent(InEvent);
+	}
+}
+
+EPropertyBagResult FUAFWeakSystemReference::SetVariable(const FAnimNextVariableReference& InVariable, const FAnimNextParamType& InType, TConstArrayView<uint8> InData) const
+{
+	TSharedPtr<FAnimNextModuleInstance> Instance = WeakPtr.Pin();
+	if (ensure(Instance.IsValid()))
+	{
+		return Instance->SetProxyVariable(InVariable, InType, InData);
+	}
+	return EPropertyBagResult::PropertyNotFound;
+}
+
+EPropertyBagResult FUAFWeakSystemReference::WriteVariable(const FAnimNextVariableReference& InVariable, const FAnimNextParamType& InType, TFunctionRef<void(TArrayView<uint8>)> InFunction) const
+{
+	TSharedPtr<FAnimNextModuleInstance> Instance = WeakPtr.Pin();
+	if (ensure(Instance.IsValid()))
+	{
+		return Instance->WriteProxyVariable(InVariable, InType, InFunction);
+	}
+	return EPropertyBagResult::PropertyNotFound;
+}
diff --git a/Engine/Plugins/Experimental/UAF/UAF/Source/UAF/Public/Component/AnimNextComponent.h b/Engine/Plugins/Experimental/UAF/UAF/Source/UAF/Public/Component/AnimNextComponent.h
--- a/Engine/Plugins/Experimental/UAF/UAF/Source/UAF/Public/Component/AnimNextComponent.h
+++ b/Engine/Plugins/Experimental/UAF/UAF/Source/UAF/Public/Component/AnimNextComponent.h
@@ -3,7 +3,7 @@
 #pragma once
 
 #include "Components/ActorComponent.h"
-#include "Module/ModuleHandle.h"
+#include "Module/SystemReference.h"
 #include "Param/ParamType.h"
 #include "TraitCore/TraitEvent.h"
 #include "Module/AnimNextModuleInitMethod.h"
@@ -18,11 +18,11 @@ namespace UE::Workspace
 }
 
 struct FAnimNextComponentInstanceData;
-class UAnimNextComponentWorldSubsystem;
 class UAnimNextModule;
 struct AnimNextVariableReference;
 // TODO: if a public API is implemented, the forward declaration of this struct would no longer be needed.
 struct FAnimMixerUtils;
+struct FUAFWeakSystemReference;
 
 namespace UE::UAF::UncookedOnly
 {
@@ -47,6 +47,8 @@ class UAnimNextComponent : public UActorComponent
 	virtual void OnUnregister() override;
 	virtual void BeginPlay() override;
 	virtual void EndPlay(const EEndPlayReason::Type EndPlayReason) override;
+	virtual void Activate(bool bReset = false) override;
+	virtual void Deactivate() override;
 
 	// UObject interface
 	virtual void Serialize(FArchive& Ar) override;
@@ -90,14 +92,6 @@ public:
 		}) == EPropertyBagResult::Success;
 	}
 
-	// Whether this component is currently updating
-	UFUNCTION(BlueprintCallable, Category = "UAF")
-	UAF_API bool IsEnabled() const;
-
-	// Enable or disable this component's update
-	UFUNCTION(BlueprintCallable, Category = "UAF")
-	UAF_API void SetEnabled(bool bEnabled);
-
 	// Enable or disable debug drawing. Note only works in builds with UE_ENABLE_DEBUG_DRAWING enabled
 	UFUNCTION(BlueprintCallable, Category = "UAF")
 	UAF_API void ShowDebugDrawing(bool bShowDebugDrawing);
@@ -109,8 +103,9 @@ public:
 	// Input events will be processed in the next graph update after they are queued
 	UAF_API void QueueInputTraitEvent(FAnimNextTraitEventPtr Event);
 
-	// Get the handle to the registered module
-	UE::UAF::FModuleHandle GetModuleHandle() const { return ModuleHandle; }
+	// Get a reference to the registered system
+	UFUNCTION(BlueprintCallable, Category = "UAF")
+	UAF_API FUAFWeakSystemReference GetSystemReference() const;
 
 	// Find the tick function for the specified event
 	// @param	InEventName			The event associated to the wanted tick function
@@ -194,16 +189,13 @@ public:
 	UFUNCTION(BlueprintCallable, Category = "UAF")
 	UAF_API void RemoveModuleEventSubsequent(FName InEventName, UAnimNextComponent* OtherAnimNextComponent, FName OtherEventName);
 
-	UFUNCTION(BlueprintCallable, Category = "UAF", meta=(DisplayName = "Get Module Handle"))
-	UAF_API FAnimNextModuleHandle BlueprintGetModuleHandle() const;
-
 	UAF_API TObjectPtr<UObject> GetAsset() const;
 
 private:
 	UAF_API void SetAsset(TObjectPtr<UObject> InAsset);
 
-	UAF_API void RegisterWithSubsystem();
-	UAF_API void UnregisterWithSubsystem();
+	UAF_API void RegisterSystem();
+	UAF_API void UnregisterSystem();
 
 	UAF_API bool IsModuleValid();
 
@@ -219,7 +211,6 @@ private:
 private:
 	friend struct UE::UAF::UncookedOnly::FUtils;
 	friend struct FAssetBlueprintTests;
-	friend class UAnimNextComponentWorldSubsystem;
 	friend class UE::UAF::Editor::FAssetWizard;
 	// TODO ideally there would be a public API for the things FAnimMixerUtils does, so it wouldn't be necessary.
 	friend struct FAnimMixerUtils;
@@ -236,11 +227,8 @@ private:
 	TObjectPtr<UAnimNextModule> Module_DEPRECATED = nullptr;
 #endif
 
-	UPROPERTY(Transient)
-	TObjectPtr<UAnimNextComponentWorldSubsystem> Subsystem = nullptr;
-
-	// Handle to the registered module
-	UE::UAF::FModuleHandle ModuleHandle;
+	// Owning reference to the registered system
+	UE::UAF::FSystemReference SystemReference;
 
 	// How to initialize the system
 	UPROPERTY(EditAnywhere, Category="System")
diff --git a/Engine/Plugins/Experimental/UAF/UAF/Source/UAF/Public/Module/ModuleHandle.h b/Engine/Plugins/Experimental/UAF/UAF/Source/UAF/Public/Module/ModuleHandle.h
deleted file mode 100644
diff --git a/Engine/Plugins/Experimental/UAF/UAF/Source/UAF/Public/Module/ModuleInstanceUtils.h b/Engine/Plugins/Experimental/UAF/UAF/Source/UAF/Public/Module/ModuleInstanceUtils.h
deleted file mode 100644
diff --git a/Engine/Plugins/Experimental/UAF/UAF/Source/UAF/Public/Module/SystemDependency.h b/Engine/Plugins/Experimental/UAF/UAF/Source/UAF/Public/Module/SystemDependency.h
new file mode 100644
--- /dev/null
+++ b/Engine/Plugins/Experimental/UAF/UAF/Source/UAF/Public/Module/SystemDependency.h
@@ -0,0 +1,19 @@
+﻿// Copyright Epic Games, Inc. All Rights Reserved.
+
+#pragma once
+
+#include "CoreTypes.h"
+
+namespace UE::UAF
+{
+
+// Dependency type passed to AddDependency in FSystemReference/FUAFWeakSystemReference etc.
+enum class ESystemDependency : uint8
+{
+	// Dependency runs before the specified event
+	Prerequisite,
+	// Dependency runs after the specified event
+	Subsequent
+};
+
+}
\ No newline at end of file
diff --git a/Engine/Plugins/Experimental/UAF/UAF/Source/UAF/Public/Module/SystemReference.h b/Engine/Plugins/Experimental/UAF/UAF/Source/UAF/Public/Module/SystemReference.h
new file mode 100644
--- /dev/null
+++ b/Engine/Plugins/Experimental/UAF/UAF/Source/UAF/Public/Module/SystemReference.h
@@ -0,0 +1,188 @@
+﻿// Copyright Epic Games, Inc. All Rights Reserved.
+
+#pragma once
+
+#include "EngineDefines.h"
+#include "Module/TaskRunLocation.h"
+#include "Subsystems/EngineSubsystem.h"
+#include "SystemReference.generated.h"
+
+#define UE_API UAF_API
+
+class UAnimNextModule;
+struct FAnimNextParamType;
+struct FAnimNextVariableReference;
+struct FAnimNextModuleInstance;
+enum class EAnimNextModuleInitMethod : uint8;
+enum class EPropertyBagResult : uint8;
+struct FAnimNextTraitEvent;
+struct FTickFunction;
+class UAnimNextComponent;
+struct FUAFWeakSystemReference;
+class UAnimNextRigVMAsset;
+
+namespace UE::UAF
+{
+enum class ESystemDependency : uint8;
+struct FModuleTaskContext;
+struct FInjectionRequest;
+struct FModuleEventTickFunction;
+}
+
+namespace UE::UAF
+{
+
+// An owning reference to a UAF system.
+// Provides a partially thread-safe API for interacting with UAF systems. Check individual methods for thread-safety. Methods that can only be called
+// on the game thread will assert when called from non-game threads.
+struct FSystemReference
+{
+	// Default-constructable only because it needs to be embedded in UObjects/UStructs
+	FSystemReference() = default;
+
+	// Move only, non-copyable
+	FSystemReference(const FSystemReference&) = delete;
+	FSystemReference& operator=(const FSystemReference&) = delete;
+	FSystemReference(FSystemReference&&);
+	FSystemReference& operator=(FSystemReference&&);
+	
+private:
+	friend UAnimNextComponent;
+	friend FUAFWeakSystemReference;
+	friend UE::UAF::FModuleEventTickFunction;
+
+	// Construct an owning handle to a system instance.
+	// Constructs a FAnimNextModuleInstance for use by worker threads.
+	// Should be called only from the game thread
+	UE_API FSystemReference(const UObject* InAsset, UObject* InObject, EAnimNextModuleInitMethod InitMethod);
+
+	// Reset this handle to invalid - gives up ownership
+	// Should be called only from the game thread
+	UE_API void Reset();
+
+	// Check if this handle is valid
+	// Can be called on any thread.
+	UE_API bool IsValid() const;
+
+	// Get the system that this handle references
+	// Can be called on any thread.
+	UE_API const UAnimNextModule* GetSystem() const;
+
+	// Enables or disables the system
+	// This operation is deferred until just before the next world tick (FWorldDelegates::OnWorldPreActorTick)
+	// Can be called on any thread.
+	UE_API void SetEnabled(bool bInEnabled) const;
+
+#if UE_ENABLE_DEBUG_DRAWING
+	// Enables or disables the system's debug drawing
+	// This operation is deferred until just before the next world tick (FWorldDelegates::OnWorldPreActorTick)
+	// Can be called on any thread.
+	UE_API void ShowDebugDrawing(bool bInShowDebugDrawing) const;
+#endif
+
+	// Queue a task to run at a particular point in a system's execution
+	// Can be called on any thread.
+	// @param	InSystemEventName	The name of the event in the system to run the supplied task relative to. If this is NAME_None, then the first user event will be used.
+	// @param	InTaskFunction		The function to run
+	// @param	InLocation			Where to run the task, before or after
+	UE_API void QueueTask(FName InSystemEventName, TUniqueFunction<void(const UE::UAF::FModuleTaskContext&)>&& InTaskFunction, UE::UAF::ETaskRunLocation InLocation = UE::UAF::ETaskRunLocation::Before) const;
+
+	// Queue an input trait event to be fired before the first user event
+	// Can be called on any thread.
+	// @param	InEvent				The event to queue
+	UE_API void QueueInputTraitEvent(TSharedPtr<FAnimNextTraitEvent> Event) const;
+
+	// Find the const tick function for the specified event. Useful to analyze/log prerequisites. Use AddDependency, RemoveDependency to actually modify dependencies
+	// Should be called only from the game thread
+	// @param	InEventName			The event associated to the wanted tick function
+	UE_API const FTickFunction* FindTickFunction(FName InEventName) const;
+
+	// Add a dependency on a tick function to the specified event
+	// Should be called only from the game thread
+	// @param	InObject			The object that owns the tick function
+	// @param	InTickFunction		The tick function to depend on
+	// @param	InSystemEventName	The event to add the dependency to
+	// @param	InDependency		The kind of dependency to add
+	UE_API void AddDependency(UObject* InObject, FTickFunction& InTickFunction, FName InSystemEventName, UE::UAF::ESystemDependency InDependency) const;
+
+	// Remove a dependency on a tick function from the specified event
+	// Should be called only from the game thread
+	// @param	InObject			The object that owns the tick function
+	// @param	InTickFunction		The tick function that was depended on
+	// @param	InSystemEventName	The event to remove the dependency from
+	// @param	InDependency		The kind of dependency to remove
+	UE_API void RemoveDependency(UObject* InObject, FTickFunction& InTickFunction, FName InSystemEventName, UE::UAF::ESystemDependency InDependency) const;
+
+	// Add a dependency on a system event to the specified system event
+	// Should be called only from the game thread
+	// @param	InSystemEventName		The event to add the dependency to/from
+	// @param	InOtherSystem			The reference to the other system whose dependencies are being modified
+	// @param	InOtherSystemEventName	The other system's event to add the dependency to/from
+	// @param	InDependency			The kind of dependency to add
+	UE_API void AddSystemEventDependency(FName InSystemEventName, FUAFWeakSystemReference InOtherSystem, FName InOtherSystemEventName, UE::UAF::ESystemDependency InDependency) const;
+
+	// Remove a dependency on a system event from the specified system event
+	// Should be called only from the game thread
+	// @param	InSystemEventName		The event to remove the dependency to/from
+	// @param	InOtherSystem			The reference to the other system whose dependencies are being modified
+	// @param	InOtherSystemEventName	The other system's event to remove the dependency to/from
+	// @param	InDependency			The kind of dependency to remove
+	UE_API void RemoveSystemEventDependency(FName InSystemEventName, FUAFWeakSystemReference InOtherSystem, FName InOtherSystemEventName, UE::UAF::ESystemDependency InDependency) const;
+
+	// Set the value of the specified variable
+	// Can be called on any thread.
+	// @param	InVariable			The variable we want to set
+	// @param	InType				The type of the variable we want to set
+	// @param	InData				The data to set the variable with
+	// @return true if the variable was set successfully
+	UE_API EPropertyBagResult SetVariable(const FAnimNextVariableReference& InVariable, const FAnimNextParamType& InType, TConstArrayView<uint8> InData) const;
+
+	// Accesses the variable of the specified name for writing.
+	// General read access is not provided via this API due to the current double-buffering strategy used to communicate variable writes to worker threads.
+	// This is intended to allow for copy-free writing of larger data structures & arrays, rather than read access.
+	// Can be called on any thread.
+	// @param	InVariable			The variable we want to access
+	// @param	InType				The type of the variable we want to access
+	// @param	InFunction			Function that will be called to allow modification of the variables
+	// @return true if the variable was accessed successfully
+	UE_API EPropertyBagResult WriteVariable(const FAnimNextVariableReference& InVariable, const FAnimNextParamType& InType, TFunctionRef<void(TArrayView<uint8>)> InFunction) const;
+
+private:
+	// Internal only constructor for selective access to owned-ref APIs
+	explicit FSystemReference(FAnimNextModuleInstance& InInstance);
+
+	// Pointer to the instance
+	TSharedPtr<FAnimNextModuleInstance> Ptr;
+};
+
+}
+
+// Engine subsystem used to manage system compilation routing, ticking and GC
+UCLASS(MinimalAPI)
+class UUAFEngineSubsystem : public UEngineSubsystem
+{
+	GENERATED_BODY()
+
+private:
+	UE_API UUAFEngineSubsystem();
+
+	UE_API virtual void BeginDestroy() override;
+
+#if WITH_EDITOR
+	// Refresh any entries that use the provided asset as it has been recompiled.
+	UE_API void OnCompileJobFinished(UAnimNextRigVMAsset* InAsset);
+#endif
+
+#if WITH_EDITOR
+	// Handle used to hook asset compilation
+	FDelegateHandle OnCompileJobFinishedHandle;
+#endif
+
+	// Handle used to hook world tick
+	FDelegateHandle OnWorldPreActorTickHandle;
+
+	// Handle used to hook into GC
+	FDelegateHandle OnPreGarbageCollectHandle;
+};
+
+#undef UE_API
\ No newline at end of file
diff --git a/Engine/Plugins/Experimental/UAF/UAF/Source/UAF/Public/Module/UAFWeakSystemReference.h b/Engine/Plugins/Experimental/UAF/UAF/Source/UAF/Public/Module/UAFWeakSystemReference.h
new file mode 100644
--- /dev/null
+++ b/Engine/Plugins/Experimental/UAF/UAF/Source/UAF/Public/Module/UAFWeakSystemReference.h
@@ -0,0 +1,98 @@
+// Copyright Epic Games, Inc. All Rights Reserved.
+
+#pragma once
+
+#include "AnimNextPoolHandle.h"
+#include "ModuleTaskContext.h"
+#include "TaskRunLocation.h"
+#include "UAFWeakSystemReference.generated.h"
+
+#define UE_API UAF_API
+
+class UAnimNextModule;
+struct FAnimNextModuleInstance;
+
+namespace UE::UAF
+{
+	enum class ESystemDependency : uint8;
+	struct FInjectionRequest;
+	struct FModuleEventTickFunction;
+	struct FSystemReference;
+}
+
+// A non-owning reference to a UAF system.
+// ~Provides a thread-safe API for interacting with UAF systems.
+USTRUCT(BlueprintType, DisplayName="UAF System Reference")
+struct FUAFWeakSystemReference
+{
+	GENERATED_BODY()
+
+	FUAFWeakSystemReference() = default;
+
+	UAF_API explicit FUAFWeakSystemReference(const UE::UAF::FSystemReference& InHandle);
+
+	bool operator==(const FUAFWeakSystemReference& InOther) const
+	{
+		return WeakPtr == InOther.WeakPtr;
+	}
+
+	// Check to see if this handle is for a valid instance
+	// Can be called on any thread.
+	bool IsValid() const
+	{
+		return WeakPtr.IsValid();
+	}
+
+	// Queue a task to run at a particular point in a system's execution
+	// Can be called on any thread.
+	// @param	InSystemEventName	The name of the event in the system to run the supplied task relative to. If this is NAME_None, then the first user event will be used.
+	// @param	InTaskFunction		The function to run
+	// @param	InLocation			Where to run the task, before or after
+	UE_API void QueueTask(FName InSystemEventName, TUniqueFunction<void(const UE::UAF::FModuleTaskContext&)>&& InTaskFunction, UE::UAF::ETaskRunLocation InLocation = UE::UAF::ETaskRunLocation::Before) const;
+
+	// Queue an input trait event to be fired before the first user event
+	// Can be called on any thread.
+	// @param	InEvent				The event to queue
+	UE_API void QueueInputTraitEvent(TSharedPtr<FAnimNextTraitEvent> InEvent) const;
+
+	// Set the value of the specified variable.
+	// Note: Sets the proxy variables that will be deferred-transferred to the system on its next ExecuteBindings_WT call.
+	// Can be called on any thread.
+	// @param	InVariable			The variable we want to set
+	// @param	InType				The type of the variable we want to set
+	// @param	InData				The data to set the variable with
+	// @return true if the variable was set successfully
+	UE_API EPropertyBagResult SetVariable(const FAnimNextVariableReference& InVariable, const FAnimNextParamType& InType, TConstArrayView<uint8> InData) const;
+
+	// Accesses the variable of the specified name for writing.
+	// Note: Sets the proxy variables that will be deferred-transferred to the system on its next ExecuteBindings_WT call.
+	// General read access is not provided via this API due to the current double-buffering strategy used to communicate variable writes to worker threads.
+	// This is intended to allow for copy-free writing of larger data structures & arrays, rather than read access.
+	// Can be called on any thread.
+	// @param	InVariable			The variable we want to access
+	// @param	InType				The type of the variable we want to access
+	// @param	InFunction			Function that will be called to allow modification of the variables
+	// @return true if the variable was accessed successfully
+	UE_API EPropertyBagResult WriteVariable(const FAnimNextVariableReference& InVariable, const FAnimNextParamType& InType, TFunctionRef<void(TArrayView<uint8>)> InFunction) const;
+
+	// Access a variable's value for read only.
+	// This performs non-recursive locks on the system instance to ensure that concurrent access is respected, so care must be taken calling this from
+	// animation worker tasks as deadlocks can occur.
+	// Type must match strictly, no conversions are performed.
+	// @param	InVariable			The variable to get the value of
+	// @param	InType				The type of the variable
+	// @param	InFunction			Function that will be called if no errors occur
+	// @return see EPropertyBagResult
+	UE_API EPropertyBagResult ReadVariable(const FAnimNextVariableReference& InVariable, const FAnimNextParamType& InType, TFunctionRef<void(TConstArrayView<uint8>)> InFunction) const;
+
+private:
+	friend UE::UAF::FSystemReference;
+	friend FAnimNextModuleInstance;
+	friend UE::UAF::FModuleEventTickFunction;
+	friend UE::UAF::FInjectionRequest;
+
+	// Ptr to the instance
+	TWeakPtr<FAnimNextModuleInstance> WeakPtr;
+};
+
+#undef UE_API
\ No newline at end of file
diff --git a/Engine/Plugins/Experimental/UAF/UAF/Source/UAFEditor/Private/AnimNextEditorModule.cpp b/Engine/Plugins/Experimental/UAF/UAF/Source/UAFEditor/Private/AnimNextEditorModule.cpp
--- a/Engine/Plugins/Experimental/UAF/UAF/Source/UAFEditor/Private/AnimNextEditorModule.cpp
+++ b/Engine/Plugins/Experimental/UAF/UAF/Source/UAFEditor/Private/AnimNextEditorModule.cpp
@@ -614,6 +614,7 @@ void FAnimNextEditorModule::ShutdownModule()
 		PropertyModule.UnregisterCustomPropertyTypeLayout("AnimNextVariableReference");
 		PropertyModule.UnregisterCustomPropertyTypeLayout("AnimNextSoftVariableReference");
 		PropertyModule.UnregisterCustomPropertyTypeLayout("NameProperty", ModuleEventPropertyTypeIdentifier);
+		PropertyModule.UnregisterCustomPropertyTypeLayout("UAFSystemReference");
 		PropertyModule.UnregisterCustomClassLayout("AnimNextVariableEntry");
 		PropertyModule.UnregisterCustomClassLayout("AnimNextRigVMAssetEditorData");
 	}
diff --git a/Engine/Plugins/Experimental/UAF/UAFAnimGraph/Source/UAFAnimGraph/Private/Injection/InjectionCallbackProxy.cpp b/Engine/Plugins/Experimental/UAF/UAFAnimGraph/Source/UAFAnimGraph/Private/Injection/InjectionCallbackProxy.cpp
--- a/Engine/Plugins/Experimental/UAF/UAFAnimGraph/Source/UAFAnimGraph/Private/Injection/InjectionCallbackProxy.cpp
+++ b/Engine/Plugins/Experimental/UAF/UAFAnimGraph/Source/UAFAnimGraph/Private/Injection/InjectionCallbackProxy.cpp
@@ -48,7 +48,7 @@ bool UInjectionCallbackProxy::Inject(
 	LifetimeEvents.OnInterrupted.BindUObject(this, &UInjectionCallbackProxy::OnInjectionInterrupted);
 	LifetimeEvents.OnBlendingOut.BindUObject(this, &UInjectionCallbackProxy::OnInjectionBlendingOut);
 
-	PlayingRequest = UE::UAF::FInjectionUtils::Inject(AnimNextComponent, AnimNextComponent->GetModuleHandle(), MoveTemp(RequestArgs), MoveTemp(LifetimeEvents));
+	PlayingRequest = UE::UAF::FInjectionUtils::Inject(AnimNextComponent, AnimNextComponent->GetSystemReference(), MoveTemp(RequestArgs), MoveTemp(LifetimeEvents));
 
 	bWasInterrupted = false;
 
diff --git a/Engine/Plugins/Experimental/UAF/UAFAnimGraph/Source/UAFAnimGraph/Private/Injection/InjectionRequest.cpp b/Engine/Plugins/Experimental/UAF/UAFAnimGraph/Source/UAFAnimGraph/Private/Injection/InjectionRequest.cpp
--- a/Engine/Plugins/Experimental/UAF/UAFAnimGraph/Source/UAFAnimGraph/Private/Injection/InjectionRequest.cpp
+++ b/Engine/Plugins/Experimental/UAF/UAFAnimGraph/Source/UAFAnimGraph/Private/Injection/InjectionRequest.cpp
@@ -4,19 +4,20 @@
 
 #include "Injection/InjectionEvents.h"
 #include "Component/AnimNextComponent.h"
-#include "Component/AnimNextWorldSubsystem.h"
 #include "Graph/AnimNextGraphInstance.h"
 #include "Logging/StructuredLog.h"
 #include "Engine/World.h"
+#include "Module/AnimNextModuleInstance.h"
+#include "Module/UAFWeakSystemReference.h"
 
 #include UE_INLINE_GENERATED_CPP_BY_NAME(InjectionRequest)
 
 namespace UE::UAF
 {
-	bool FInjectionRequest::Inject(FInjectionRequestArgs&& InRequestArgs, FInjectionLifetimeEvents&& InLifetimeEvents, UObject* InHost, FModuleHandle InHandle)
+	bool FInjectionRequest::Inject(FInjectionRequestArgs&& InRequestArgs, FInjectionLifetimeEvents&& InLifetimeEvents, UObject* InHost, FUAFWeakSystemReference InReference)
 	{
 		check(IsInGameThread());
-		if (InHost == nullptr || !InHandle.IsValid())
+		if (InHost == nullptr || !InReference.IsValid())
 		{
 			return false;	// Nothing to play
 		}
@@ -40,16 +41,13 @@ namespace UE::UAF
 		RequestArgs = MoveTemp(InRequestArgs);
 		LifetimeEvents = MoveTemp(InLifetimeEvents);
 		WeakHost = InHost;
-		WeakWorld = World;
-		Handle = InHandle;
+		Reference = InReference;
 		Status = EInjectionStatus::Pending;
 
 		auto InjectEvent = MakeTraitEvent<FInjection_InjectEvent>();
 		InjectEvent->Request = AsShared();
 
-		UAnimNextWorldSubsystem* Subsystem = World->GetSubsystem<UAnimNextWorldSubsystem>();
-		check(Subsystem);
-		Subsystem->QueueInputTraitEventHandle(InHandle, InjectEvent);
+		Reference.QueueInputTraitEvent(InjectEvent);
 
 		InjectionEvent = InjectEvent;
 
@@ -78,12 +76,7 @@ namespace UE::UAF
 		auto UninjectEvent = MakeTraitEvent<FInjection_UninjectEvent>();
 		UninjectEvent->Request = AsShared();
 
-		if (UWorld* World = WeakWorld.Get())
-		{
-			UAnimNextWorldSubsystem* Subsystem = World->GetSubsystem<UAnimNextWorldSubsystem>();
-			check(Subsystem);
-			Subsystem->QueueInputTraitEventHandle(Handle, UninjectEvent);
-		}
+		Reference.QueueInputTraitEvent(UninjectEvent);
 	}
 
 	bool FInjectionRequest::ValidateArgs(const FInjectionRequestArgs& InRequestArgs)
diff --git a/Engine/Plugins/Experimental/UAF/UAFAnimGraph/Source/UAFAnimGraph/Private/Injection/InjectionUtils.cpp b/Engine/Plugins/Experimental/UAF/UAFAnimGraph/Source/UAFAnimGraph/Private/Injection/InjectionUtils.cpp
--- a/Engine/Plugins/Experimental/UAF/UAFAnimGraph/Source/UAFAnimGraph/Private/Injection/InjectionUtils.cpp
+++ b/Engine/Plugins/Experimental/UAF/UAFAnimGraph/Source/UAFAnimGraph/Private/Injection/InjectionUtils.cpp
@@ -10,15 +10,15 @@ namespace UE::UAF
 {
 	FInjectionRequestPtr FInjectionUtils::Inject(UAnimNextComponent* InComponent, FInjectionRequestArgs&& InArgs, FInjectionLifetimeEvents&& InLifetimeEvents)
 	{
-		return Inject(InComponent, InComponent->GetModuleHandle(), MoveTemp(InArgs), MoveTemp(InLifetimeEvents));
+		return Inject(InComponent, InComponent->GetSystemReference(), MoveTemp(InArgs), MoveTemp(InLifetimeEvents));
 	}
 
-	FInjectionRequestPtr FInjectionUtils::Inject(UObject* InHost, FModuleHandle InHandle, FInjectionRequestArgs&& InArgs, FInjectionLifetimeEvents&& InLifetimeEvents)
+	FInjectionRequestPtr FInjectionUtils::Inject(UObject* InHost, FUAFWeakSystemReference InSystemReference, FInjectionRequestArgs&& InArgs, FInjectionLifetimeEvents&& InLifetimeEvents)
 	{
 		check(IsInGameThread());
 
 		FInjectionRequestPtr Request = MakeInjectionRequest();
-		if(Request->Inject(MoveTemp(InArgs), MoveTemp(InLifetimeEvents), InHost, InHandle))
+		if(Request->Inject(MoveTemp(InArgs), MoveTemp(InLifetimeEvents), InHost, InSystemReference))
 		{
 			return Request;
 		}
@@ -46,7 +46,7 @@ namespace UE::UAF
 	{
 		return PlayAnimHandle(
 			InComponent,
-			InComponent->GetModuleHandle(),
+			InComponent->GetSystemReference(),
 			InSite,
 			InAnimSequence,
 			{
@@ -69,7 +69,7 @@ namespace UE::UAF
 	{
 		return PlayAnimHandle(
 			InComponent,
-			InComponent->GetModuleHandle(),
+			InComponent->GetSystemReference(),
 			InSite,
 			InAnimSequence,
 			MoveTemp(InArgs),
@@ -80,7 +80,7 @@ namespace UE::UAF
 
 	FInjectionRequestPtr FInjectionUtils::PlayAnimHandle(
 		UObject* InHost,
-		FModuleHandle InModuleHandle,
+		FUAFWeakSystemReference InSystemReference,
 		const FInjectionSite& InSite,
 		UAnimSequence* InAnimSequence,
 		FPlayAnimArgs&& InArgs,
@@ -127,7 +127,7 @@ namespace UE::UAF
 			RequestArgs.FactoryParams = MoveTemp(InArgs.FactoryParams);
 		}
 
-		return Inject(InHost, InModuleHandle, MoveTemp(RequestArgs), MoveTemp(InLifetimeEvents));
+		return Inject(InHost, InSystemReference, MoveTemp(RequestArgs), MoveTemp(InLifetimeEvents));
 	}
 
 	FInjectionRequestPtr FInjectionUtils::InjectAsset(
@@ -141,7 +141,7 @@ namespace UE::UAF
 	{
 		return InjectAsset(
 			InComponent,
-			InComponent->GetModuleHandle(),
+			InComponent->GetSystemReference(),
 			InSite,
 			InAsset,
 			MoveTemp(InFactoryParams),
@@ -152,7 +152,7 @@ namespace UE::UAF
 
 	FInjectionRequestPtr FInjectionUtils::InjectAsset(
 		UObject* InHost,
-		FModuleHandle InModuleHandle,
+		FUAFWeakSystemReference InSystemReference,
 		const FInjectionSite& InSite,
 		UObject* InAsset,
 		FAnimNextFactoryParams&& InFactoryParams,
@@ -170,7 +170,7 @@ namespace UE::UAF
 		RequestArgs.BlendOutSettings = InBlendOutSettings;
 		RequestArgs.FactoryParams = MoveTemp(InFactoryParams);
 
-		return Inject(InHost, InModuleHandle, MoveTemp(RequestArgs), MoveTemp(InLifetimeEvents));
+		return Inject(InHost, InSystemReference, MoveTemp(RequestArgs), MoveTemp(InLifetimeEvents));
 	}
 
 	FInjectionRequestPtr FInjectionUtils::InjectEvaluationModifier(
@@ -180,14 +180,14 @@ namespace UE::UAF
 	{
 		return InjectEvaluationModifier(
 			InComponent,
-			InComponent->GetModuleHandle(),
+			InComponent->GetSystemReference(),
 			InEvaluationModifier,
 			InSite);
 	}
 
 	FInjectionRequestPtr FInjectionUtils::InjectEvaluationModifier(
 		UObject* InHost,
-		FModuleHandle InModuleHandle,
+		FUAFWeakSystemReference InSystemReference,
 		const TSharedRef<IEvaluationModifier>& InEvaluationModifier,
 		const FInjectionSite& InSite)
 	{
@@ -198,6 +198,6 @@ namespace UE::UAF
 		RequestArgs.Type = EAnimNextInjectionType::EvaluationModifier;
 		RequestArgs.EvaluationModifier = InEvaluationModifier;
 
-		return Inject(InHost, InModuleHandle, MoveTemp(RequestArgs));
+		return Inject(InHost, InSystemReference, MoveTemp(RequestArgs));
 	}
 }
diff --git a/Engine/Plugins/Experimental/UAF/UAFAnimGraph/Source/UAFAnimGraph/Public/Graph/AnimNextAnimGraph.h b/Engine/Plugins/Experimental/UAF/UAFAnimGraph/Source/UAFAnimGraph/Public/Graph/AnimNextAnimGraph.h
--- a/Engine/Plugins/Experimental/UAF/UAFAnimGraph/Source/UAFAnimGraph/Public/Graph/AnimNextAnimGraph.h
+++ b/Engine/Plugins/Experimental/UAF/UAFAnimGraph/Source/UAFAnimGraph/Public/Graph/AnimNextAnimGraph.h
@@ -3,7 +3,6 @@
 #pragma once
 
 #include "CoreMinimal.h"
-#include "Module/ModuleHandle.h"
 #include "StructUtils/StructView.h"
 #include "TraitCore/TraitEvent.h"
 #include "TraitCore/TraitPtr.h"
diff --git a/Engine/Plugins/Experimental/UAF/UAFAnimGraph/Source/UAFAnimGraph/Public/Injection/InjectionRequest.h b/Engine/Plugins/Experimental/UAF/UAFAnimGraph/Source/UAFAnimGraph/Public/Injection/InjectionRequest.h
--- a/Engine/Plugins/Experimental/UAF/UAFAnimGraph/Source/UAFAnimGraph/Public/Injection/InjectionRequest.h
+++ b/Engine/Plugins/Experimental/UAF/UAFAnimGraph/Source/UAFAnimGraph/Public/Injection/InjectionRequest.h
@@ -12,11 +12,12 @@
 #include "TraitInterfaces/ITimeline.h"
 #include "Injection/InjectionStatus.h"
 #include "Injection/InjectionSite.h"
-#include "Module/ModuleHandle.h"
 #include "Containers/SpscQueue.h"
+#include "Module/UAFWeakSystemReference.h"
 
 #include "InjectionRequest.generated.h"
 
+struct FUAFWeakSystemReference;
 class FReferenceCollector;
 struct FAnimNextModuleInjectionComponent;
 struct FAnimNextGraphInstance;
@@ -260,7 +261,7 @@ namespace UE::UAF
 
 	private:
 		// Sends this request to the specified host and it will attempt to play with the requested arguments
-		bool Inject(FInjectionRequestArgs&& InRequestArgs, FInjectionLifetimeEvents&& InLifetimeEvents, UObject* InHost, FModuleHandle InHandle);
+		bool Inject(FInjectionRequestArgs&& InRequestArgs, FInjectionLifetimeEvents&& InLifetimeEvents, UObject* InHost, FUAFWeakSystemReference InReference);
 
 		// Interrupts this request and request that we transition to the source input on the playing injection site
 		void Uninject();
@@ -293,11 +294,8 @@ namespace UE::UAF
 		// The object we are playing on
 		TWeakObjectPtr<UObject> WeakHost;
 
-		// The world within which we are playing
-		TWeakObjectPtr<UWorld> WeakWorld;
-
-		// Handle to the module instance
-		FModuleHandle Handle;
+		// Reference to the system instance
+		FUAFWeakSystemReference Reference;
 
 		// The injection event if we have injected already
 		TWeakPtr<FAnimNextTraitEvent> InjectionEvent;
diff --git a/Engine/Plugins/Experimental/UAF/UAFAnimGraph/Source/UAFAnimGraph/Public/Injection/InjectionUtils.h b/Engine/Plugins/Experimental/UAF/UAFAnimGraph/Source/UAFAnimGraph/Public/Injection/InjectionUtils.h
--- a/Engine/Plugins/Experimental/UAF/UAFAnimGraph/Source/UAFAnimGraph/Public/Injection/InjectionUtils.h
+++ b/Engine/Plugins/Experimental/UAF/UAFAnimGraph/Source/UAFAnimGraph/Public/Injection/InjectionUtils.h
@@ -59,7 +59,7 @@ namespace UE::UAF
 		// Injects an animation sequence into a running AnimNext module
 		// @param InComponent          The UAnimNextComponent that hosts the module being played on
 		// @param InHost               The object that hosts the module (e.g. UAnimNextComponent) being played on
-		// @param InModuleHandle       The module handle being injected into
+		// @param InSystemReference    The system reference being injected into
 		// @param InSite               The injection site to 'play' the animation sequence on
 		// @param InAnimSequence       The animation sequence to play
 		// @param InArgs               Optional args to use
@@ -79,7 +79,7 @@ namespace UE::UAF
 		// @see PlayAnim
 		UAFANIMGRAPH_API static FInjectionRequestPtr PlayAnimHandle(
 			UObject* InHost,
-			FModuleHandle InModuleHandle,
+			FUAFWeakSystemReference InSystemReference,
 			const FInjectionSite& InSite,
 			UAnimSequence* InAnimSequence,
 			FPlayAnimArgs&& InArgs = FPlayAnimArgs(),
@@ -90,12 +90,10 @@ namespace UE::UAF
 		// Injects an asset into a running AnimNext module
 		// @param InComponent          The UAnimNextComponent that hosts the module being played on
 		// @param InHost               The object that hosts the module (e.g. UAnimNextComponent) being played on
-		// @param InHandle             The module handle being injected into
+		// @param InSystemReference    The system reference being injected into
 		// @param InSite               The injection site to 'play' the asset on
 		// @param InAsset              The asset to 'play'
 		// @param InFactoryParams      The params used to manufacture a graph to run the injected object 
-		// @param InBindingComponent   A component to use to supply additional data interface bindings at the injection site
-		// @param InBindingModuleHandle A module to use to supply additional data interface bindings at the injection site
 		// @param InBlendInSettings    The blend in settings
 		// @param InBlendOutSettings   The blend in settings
 		// @param InLifetimeEvents     Delegates called for various request lifetime events
@@ -112,7 +110,7 @@ namespace UE::UAF
 		// @see InjectAsset
 		UAFANIMGRAPH_API static FInjectionRequestPtr InjectAsset(
 			UObject* InHost,
-			FModuleHandle InModuleHandle,
+			FUAFWeakSystemReference InSystemReference,
 			const FInjectionSite& InSite,
 			UObject* InAsset,
 			FAnimNextFactoryParams&& InFactoryParams = FAnimNextFactoryParams(),
@@ -124,7 +122,7 @@ namespace UE::UAF
 		// Does not require that injection sites are publicly visible.
 		// @param InComponent          The UAnimNextComponent that hosts the module being injected into
 		// @param InHost               The object that hosts the module (e.g. UAnimNextComponent) being injected into
-		// @param InModuleHandle       The module handle being injected into
+		// @param InSystemReference    The system reference being injected into
 		// @param InEvaluationModifier The evaluation modifier to inject
 		// @param InSite               The injection site to use. If this is NAME_None, then the first site will be used.
 		// @return the new injection request for tracking, or an invalid request if injection failed. This can be discarded for fire-and-forget style behavior.
@@ -136,14 +134,14 @@ namespace UE::UAF
 		// @see InjectExternalGraph
 		UAFANIMGRAPH_API static FInjectionRequestPtr InjectEvaluationModifier(
 			UObject* InHost,
-			FModuleHandle InModuleHandle,
+			FUAFWeakSystemReference InSystemReference,
 			const TSharedRef<IEvaluationModifier>& InEvaluationModifier,
 			const FInjectionSite& InSite = FInjectionSite());
 
 		// Inject into a running AnimNext module (raw custom args)
 		// @param InComponent          The UAnimNextComponent that hosts the module being injected into
 		// @param InHost               The object that hosts the module (e.g. UAnimNextComponent) being injected into
-		// @param InHandle             The module handle being injected into
+		// @param InSystemReference    The system reference being injected into
 		// @param InArgs               Arguments to the injection request
 		// @param InLifetimeEvents     Delegates called for various request lifetime events
 		// @return the new injection request for tracking, or an invalid request if injection failed. This can be discarded for fire-and-forget style behavior.
@@ -155,7 +153,7 @@ namespace UE::UAF
 		// @see Inject
 		UAFANIMGRAPH_API static FInjectionRequestPtr Inject(
 			UObject* InHost,
-			FModuleHandle InHandle,
+			FUAFWeakSystemReference InSystemReference,
 			FInjectionRequestArgs&& InArgs,
 			FInjectionLifetimeEvents&& InLifetimeEvents = FInjectionLifetimeEvents());
 
diff --git a/Engine/Plugins/Experimental/UAF/UAFAnimGraph/Source/UAFAnimGraph/Public/TraitInterfaces/IBlendStack.h b/Engine/Plugins/Experimental/UAF/UAFAnimGraph/Source/UAFAnimGraph/Public/TraitInterfaces/IBlendStack.h
--- a/Engine/Plugins/Experimental/UAF/UAFAnimGraph/Source/UAFAnimGraph/Public/TraitInterfaces/IBlendStack.h
+++ b/Engine/Plugins/Experimental/UAF/UAFAnimGraph/Source/UAFAnimGraph/Public/TraitInterfaces/IBlendStack.h
@@ -8,7 +8,6 @@
 #include "AlphaBlend.h"
 #include "Factory/AnimNextFactoryParams.h"
 #include "Graph/AnimNextAnimationGraph.h"
-#include "Module/ModuleHandle.h"
 #include "Animation/BlendProfile.h"
 
 #define UE_API UAFANIMGRAPH_API
diff --git a/Engine/Plugins/Experimental/UAF/UAFPoseSearch/Config/DefaultUAFPoseSearch.ini b/Engine/Plugins/Experimental/UAF/UAFPoseSearch/Config/DefaultUAFPoseSearch.ini
--- a/Engine/Plugins/Experimental/UAF/UAFPoseSearch/Config/DefaultUAFPoseSearch.ini
+++ b/Engine/Plugins/Experimental/UAF/UAFPoseSearch/Config/DefaultUAFPoseSearch.ini
@@ -15,4 +15,8 @@
 +StructRedirects=(OldName="/Script/AnimNextPoseSearch.RigUnit_PoseSearchDatabaseGetTags", NewName="/Script/UAFPoseSearch.RigUnit_PoseSearchDatabaseGetTags")
 +StructRedirects=(OldName="/Script/AnimNextPoseSearch.RigUnit_PoseSearchResultGetSelectedDatabase", NewName="/Script/UAFPoseSearch.RigUnit_PoseSearchResultGetSelectedDatabase")
 
-+ClassRedirects=(OldName="/Script/AnimNextPoseSearch.PoseSearchFeatureChannel_DistanceFromAnimNextVar", NewName="/Script/UAFPoseSearch.PoseSearchFeatureChannel_DistanceFromAnimNextVar")
\ No newline at end of file
++ClassRedirects=(OldName="/Script/AnimNextPoseSearch.PoseSearchFeatureChannel_DistanceFromAnimNextVar", NewName="/Script/UAFPoseSearch.PoseSearchFeatureChannel_DistanceFromAnimNextVar")
+
+; System references
++PropertyRedirects=(OldName="/Script/AnimNextPoseSearch.MotionMatchingTraitSharedData.MotionMatchingResultModuleHandle", NewName="MotionMatchingResultSystemReference")
++PropertyRedirects=(OldName="/Script/AnimNextPoseSearch.MotionMatchingTraitSharedData.MotionMatchingResultModuleHandleAlt", NewName="MotionMatchingResultSystemReferenceAlt")
diff --git a/Engine/Plugins/Experimental/UAF/UAFPoseSearch/Source/UAFPoseSearch/Internal/MotionMatchingTraitData.h b/Engine/Plugins/Experimental/UAF/UAFPoseSearch/Source/UAFPoseSearch/Internal/MotionMatchingTraitData.h
--- a/Engine/Plugins/Experimental/UAF/UAFPoseSearch/Source/UAFPoseSearch/Internal/MotionMatchingTraitData.h
+++ b/Engine/Plugins/Experimental/UAF/UAFPoseSearch/Source/UAFPoseSearch/Internal/MotionMatchingTraitData.h
@@ -4,7 +4,7 @@
 
 #include "CoreMinimal.h"
 #include "AlphaBlend.h"
-#include "Module/ModuleHandle.h"
+#include "Module/UAFWeakSystemReference.h"
 #include "PoseSearch/PoseSearchInteractionLibrary.h"
 #include "PoseSearch/PoseSearchLibrary.h"
 #include "TraitCore/TraitSharedData.h"
@@ -87,17 +87,17 @@ struct FMotionMatchingTraitSharedData : public FAnimNextTraitSharedData
 	UPROPERTY(EditAnywhere, Category = Settings)
 	FName AlignmentTransformVariableName;
 
-	// If set along with result variable name, the motion matching result will be written to the variable in this module, otherwise the variable will be written in the module containing the MotionMatchingTrait
+	// If set along with result variable name, the motion matching result will be written to the variable in this system, otherwise the variable will be written in the module containing the MotionMatchingTrait
 	UPROPERTY(EditAnywhere, Category = Settings)
-	FAnimNextModuleHandle MotionMatchingResultModuleHandle;
+	FUAFWeakSystemReference MotionMatchingResultSystemReference;
 
 	// Same as MotionMatchingResultVariableName but for an alternative destination variable
 	UPROPERTY(EditAnywhere, Category = Settings)
 	FName MotionMatchingResultVariableNameAlt = NAME_None;
 
-	// Same as MotionMatchingResultModuleHandle but for an alternative destination variable
+	// Same as MotionMatchingResultSystemReference but for an alternative destination variable
 	UPROPERTY(EditAnywhere, Category = Settings)
-	FAnimNextModuleHandle MotionMatchingResultModuleHandleAlt;
+	FUAFWeakSystemReference MotionMatchingResultSystemReferenceAlt;
 
 	UPROPERTY(EditAnywhere, Category = Interaction, meta = (PinHiddenByDefault))
 	TArray<FPoseSearchInteractionAvailability> Availabilities;
@@ -143,9 +143,9 @@ struct FMotionMatchingTraitSharedData : public FAnimNextTraitSharedData
 		GeneratorMacro(NextUpdateInterruptMode) \
 		GeneratorMacro(MotionMatchingResultVariableName) \
 		GeneratorMacro(AlignmentTransformVariableName) \
-		GeneratorMacro(MotionMatchingResultModuleHandle) \
+		GeneratorMacro(MotionMatchingResultSystemReference) \
 		GeneratorMacro(MotionMatchingResultVariableNameAlt) \
-		GeneratorMacro(MotionMatchingResultModuleHandleAlt) \
+		GeneratorMacro(MotionMatchingResultSystemReferenceAlt) \
 		GeneratorMacro(Availabilities) \
 		GeneratorMacro(bValidateResultAgainstAvailabilities) \
 		GeneratorMacro(bKeepInteractionAlive) \
diff --git a/Engine/Plugins/Experimental/UAF/UAFPoseSearch/Source/UAFPoseSearch/Private/MotionMatchingTrait.cpp b/Engine/Plugins/Experimental/UAF/UAFPoseSearch/Source/UAFPoseSearch/Private/MotionMatchingTrait.cpp
--- a/Engine/Plugins/Experimental/UAF/UAFPoseSearch/Source/UAFPoseSearch/Private/MotionMatchingTrait.cpp
+++ b/Engine/Plugins/Experimental/UAF/UAFPoseSearch/Source/UAFPoseSearch/Private/MotionMatchingTrait.cpp
@@ -441,14 +441,14 @@ PRAGMA_ENABLE_DEPRECATION_WARNINGS
 		}
 #endif // WITH_EDITOR
 		
-		const auto PublishResult = [](FAnimNextModuleInstance* ModuleInstance, const FModuleHandle& ModuleHandle, const FPoseSearchBlueprintResult& SearchResult, FName ResultOutputVariableName)
+		const auto PublishResult = [](FAnimNextModuleInstance* ModuleInstance, const FUAFWeakSystemReference& SystemReference, const FPoseSearchBlueprintResult& SearchResult, FName ResultOutputVariableName)
 			{
 				if (ResultOutputVariableName != NAME_None)
 				{
-					if (ModuleHandle.IsValid())
+					if (SystemReference.IsValid())
 					{
 						// storing SearchResult TObjectPtr(s) as TWeakObjectPtr(s) to avoid referencing deallocated memory 
-						// in case the module containing this trait gets destroyed before the QueueTaskOnOtherModule lambda execution
+						// in case the module containing this trait gets destroyed before the QueueTask lambda execution
 						TWeakObjectPtr<UObject> SelectedAnim = SearchResult.SelectedAnim;
 						float SelectedTime = SearchResult.SelectedTime;
 						bool bIsContinuingPoseSearch = SearchResult.bIsContinuingPoseSearch;
@@ -463,8 +463,7 @@ PRAGMA_ENABLE_DEPRECATION_WARNINGS
 						TArray<FTransform> ActorRootTransforms = SearchResult.ActorRootTransforms;
 						TArray<FTransform> ActorRootBoneTransforms = SearchResult.ActorRootBoneTransforms;
 
-						ModuleInstance->QueueTaskOnOtherModule(
-							ModuleHandle,
+						SystemReference.QueueTask(
 							FRigUnit_AnimNextPrePhysicsEvent::DefaultEventName,
 							[SelectedAnim, SelectedTime, bIsContinuingPoseSearch, WantedPlayRate, bLoop,
 							bIsMirrored, BlendParameters, SelectedDatabase, SearchCost, bIsInteraction,
@@ -509,11 +508,11 @@ PRAGMA_ENABLE_DEPRECATION_WARNINGS
 		FAnimNextModuleInstance* ModuleInstance = Binding.GetTraitPtr().GetNodeInstance()->GetOwner().GetModuleInstance();
 		check(ModuleInstance);
 
-		const FAnimNextModuleHandle ResultModuleHandle = SharedData->GetMotionMatchingResultModuleHandle(Binding);
-		PublishResult(ModuleInstance, ResultModuleHandle.ModuleHandle, MotionMatchingState.SearchResult, ResultOutputVariableName);
+		const FUAFWeakSystemReference ResultSystemReference = SharedData->GetMotionMatchingResultSystemReference(Binding);
+		PublishResult(ModuleInstance, ResultSystemReference, MotionMatchingState.SearchResult, ResultOutputVariableName);
 		
-		const FAnimNextModuleHandle ResultModuleHandleAlt = SharedData->GetMotionMatchingResultModuleHandleAlt(Binding);
-		PublishResult(ModuleInstance, ResultModuleHandleAlt.ModuleHandle, MotionMatchingState.SearchResult, ResultOutputVariableNameAlt);
+		const FUAFWeakSystemReference ResultSystemReferenceAlt = SharedData->GetMotionMatchingResultSystemReferenceAlt(Binding);
+		PublishResult(ModuleInstance, ResultSystemReferenceAlt, MotionMatchingState.SearchResult, ResultOutputVariableNameAlt);
 	}
 
 	void FMotionMatchingTrait::PostEvaluate(FEvaluateTraversalContext& Context, const TTraitBinding<IEvaluate>& Binding) const
diff --git a/Engine/Plugins/Experimental/UAF/UAFPoseSearch/Source/UAFPoseSearch/Private/MotionMatchingTrait.h b/Engine/Plugins/Experimental/UAF/UAFPoseSearch/Source/UAFPoseSearch/Private/MotionMatchingTrait.h
--- a/Engine/Plugins/Experimental/UAF/UAFPoseSearch/Source/UAFPoseSearch/Private/MotionMatchingTrait.h
+++ b/Engine/Plugins/Experimental/UAF/UAFPoseSearch/Source/UAFPoseSearch/Private/MotionMatchingTrait.h
@@ -4,7 +4,6 @@
 
 #include "CoreMinimal.h"
 #include "AlphaBlend.h"
-#include "Module/ModuleHandle.h"
 #include "PoseSearch/PoseSearchInteractionLibrary.h"
 #include "PoseSearch/PoseSearchLibrary.h"
 #include "TraitCore/Trait.h"
