[Backout] - CL47144844
#rnx
Original CL Desc
-----------------------------------------------------------------
UAF: Owning/weak system references

Prerequisite for scene graph-like system reference semantics.
This allows us to better reason about the validity of weak system references without having access to implementation details.
The most significant changes are in reworking the world subsystem code into FSystemReference.
- Pooling is still used for allocation (via a shared ptr custom deleter), but 'handles' are now references that wrap a shared/weak ptr.
- 'Owning references' (FSystemReference) are specifically designed to be single-owner (move-only), despite using TSharedPtr under the hood.
- World subsystems are replaced with an engine subsystem and a single allocation pool, rather than one per world.
- Deferred operations like enable/disable/release now use a SPSC queue of operations rather than locking.

Also switches UAnimNextComponent to use Activate/Deactivate pattern inherited from actor components

diff --git a/Engine/Plugins/Experimental/MovieSceneAnimMixer/Source/MovieSceneAnimMixer/Private/Systems/MovieSceneAnimNextTargetSystem.cpp b/Engine/Plugins/Experimental/MovieSceneAnimMixer/Source/MovieSceneAnimMixer/Private/Systems/MovieSceneAnimNextTargetSystem.cpp
--- a/Engine/Plugins/Experimental/MovieSceneAnimMixer/Source/MovieSceneAnimMixer/Private/Systems/MovieSceneAnimNextTargetSystem.cpp
+++ b/Engine/Plugins/Experimental/MovieSceneAnimMixer/Source/MovieSceneAnimMixer/Private/Systems/MovieSceneAnimNextTargetSystem.cpp
@@ -41,6 +41,7 @@
 
 #include "PrimitiveSceneProxy.h"
 #include "Module/AnimNextModule.h"
+#include "UAF/Private/Component/AnimNextComponentWorldSubsystem.h"
 #include "UObject/FortniteMainBranchObjectVersion.h"
 
 #include UE_INLINE_GENERATED_CPP_BY_NAME(MovieSceneAnimNextTargetSystem)
@@ -91,9 +92,9 @@ namespace UE::MovieScene
 					// TODO the animation needs to be updated, or the post needs to be reset to ref pose when restoring to a previous module.
 					// Currently there is no support for this in UAF.
 					// This is most pronounced in the editor, where the mesh stays in the same pose when the sequence is closed.
-					FAnimMixerUtils::UnregisterUAFSystem(AnimNextComponent);
+					FAnimMixerUtils::UnregisterWithUAFSubsystem(AnimNextComponent);
 					FAnimMixerUtils::SetUAFComponentModule(AnimNextComponent, InOutCachedValue.PreAnimatedModule.Get());
-					FAnimMixerUtils::RegisterUAFSystem(AnimNextComponent);
+					FAnimMixerUtils::RegisterWithUAFSubsystem(AnimNextComponent);
 				}
 			}
 			
@@ -102,7 +103,7 @@ namespace UE::MovieScene
 			{
 				if (UAnimNextComponent* AnimNextComponent = Cast<UAnimNextComponent>(ObjectAndInjectionSite.Key.ResolveObjectPtr()))
 				{
-					AnimNextComponent->SetActive(false);
+					AnimNextComponent->SetEnabled(false);
 				}
 			}
 		}
@@ -267,9 +268,9 @@ namespace UE::MovieScene
 				{
 					const UMovieSceneAnimMixerSettings* AnimMixerSettings = GetDefault<UMovieSceneAnimMixerSettings>();
 
-					FAnimMixerUtils::UnregisterUAFSystem(AnimNextComponent);
+					FAnimMixerUtils::UnregisterWithUAFSubsystem(AnimNextComponent);
 					FAnimMixerUtils::SetUAFComponentModule(AnimNextComponent, AnimMixerSettings->DefaultUAFModule.LoadSynchronous());
-					FAnimMixerUtils::RegisterUAFSystem(AnimNextComponent);
+					FAnimMixerUtils::RegisterWithUAFSubsystem(AnimNextComponent);
 
 					bModuleReplaced = true;
 				}
@@ -299,7 +300,7 @@ namespace UE::MovieScene
 					PreAnimatedStorage->BeginTrackingEntity(EntityID, true, FRootInstanceHandle(), AnimNextComponent, InjectionSite.DesiredSite);
 					PreAnimatedSkelMeshCompStorage->BeginTrackingEntity(EntityID, true, FRootInstanceHandle(), SkeletalMeshComponent);
 
-					bool bAnimNextModuleEnabled = AnimNextComponent->IsActive();
+					bool bAnimNextModuleEnabled = AnimNextComponent->IsEnabled();
 
 					auto OnCacheAnimatedState = [this, TargetData, bAnimNextModuleEnabled, bModuleReplaced, PreAnimatedModule](UAnimNextComponent*)
 					{
@@ -318,7 +319,7 @@ namespace UE::MovieScene
 					// Force enable the anim next module if not currently enabled
 					if (!bAnimNextModuleEnabled)
 					{
-						AnimNextComponent->SetActive(true);
+						AnimNextComponent->SetEnabled(true);
 					}
 				}
 				else if(TargetData->Modifier.IsValid())
diff --git a/Engine/Plugins/Experimental/MovieSceneAnimMixer/Source/MovieSceneAnimMixer/Public/AnimMixerUtils.h b/Engine/Plugins/Experimental/MovieSceneAnimMixer/Source/MovieSceneAnimMixer/Public/AnimMixerUtils.h
--- a/Engine/Plugins/Experimental/MovieSceneAnimMixer/Source/MovieSceneAnimMixer/Public/AnimMixerUtils.h
+++ b/Engine/Plugins/Experimental/MovieSceneAnimMixer/Source/MovieSceneAnimMixer/Public/AnimMixerUtils.h
@@ -7,19 +7,19 @@
 struct FAnimMixerUtils
 {
 
-	static void RegisterUAFSystem(UAnimNextComponent* InComponent)
+	static void RegisterWithUAFSubsystem(UAnimNextComponent* InComponent)
 	{
 		if (InComponent)
 		{
-			InComponent->RegisterSystem();
+			InComponent->RegisterWithSubsystem();
 		}
 	}
 
-	static void UnregisterUAFSystem(UAnimNextComponent* InComponent)
+	static void UnregisterWithUAFSubsystem(UAnimNextComponent* InComponent)
 	{
 		if (InComponent)
 		{
-			InComponent->UnregisterSystem();
+			InComponent->UnregisterWithSubsystem();
 		}
 	}
 
diff --git a/Engine/Plugins/Experimental/UAF/UAF/Config/DefaultUAF.ini b/Engine/Plugins/Experimental/UAF/UAF/Config/DefaultUAF.ini
--- a/Engine/Plugins/Experimental/UAF/UAF/Config/DefaultUAF.ini
+++ b/Engine/Plugins/Experimental/UAF/UAF/Config/DefaultUAF.ini
@@ -296,9 +296,3 @@
 +StructRedirects=(OldName="/Script/SkeletonTemplateFrameworkRuntime.AbstractSkeleton_BoneBinding", NewName="/Script/UAF.AbstractSkeleton_BoneBinding")
 +StructRedirects=(OldName="/Script/SkeletonTemplateFrameworkRuntime.AbstractSkeleton_AttributeBinding", NewName="/Script/UAF.AbstractSkeleton_AttributeBinding")
 +StructRedirects=(OldName="/Script/SkeletonTemplateFrameworkRuntime.AbstractSkeletonSet", NewName="/Script/UAF.AbstractSkeletonSet")
-
-; System I/O
-+StructRedirects=(OldName="/Script/UAF.AnimNextModuleHandle", NewName="/Script/UAF.UAFWeakSystemReference")
-+FunctionRedirects=(OldName="AnimNextComponent.SetEnabled",NewName="ActorComponent.SetActive")
-+FunctionRedirects=(OldName="AnimNextComponent.IsEnabled",NewName="ActorComponent.IsActive")
-+FunctionRedirects=(OldName="AnimNextComponent.BlueprintGetModuleHandle",NewName="AnimNextComponent.GetSystemReference")
\ No newline at end of file
diff --git a/Engine/Plugins/Experimental/UAF/UAF/Source/UAF/Internal/AnimNextDebugDraw.h b/Engine/Plugins/Experimental/UAF/UAF/Source/UAF/Internal/AnimNextDebugDraw.h
--- a/Engine/Plugins/Experimental/UAF/UAF/Source/UAF/Internal/AnimNextDebugDraw.h
+++ b/Engine/Plugins/Experimental/UAF/UAF/Source/UAF/Internal/AnimNextDebugDraw.h
@@ -11,11 +11,11 @@
 class UObject;
 struct FRigVMDrawInterface;
 struct FAnimNextModuleInstance;
+class UAnimNextWorldSubsystem;
 
 namespace UE::UAF
 {
 	struct FModuleEventTickFunction;
-	struct FSystemReference;
 }
 
 namespace UE::UAF::Debug
@@ -90,7 +90,7 @@ private:
 
 	friend struct ::FAnimNextModuleInstance;
 	friend struct UE::UAF::FModuleEventTickFunction;
-	friend UE::UAF::FSystemReference;
+	friend class ::UAnimNextWorldSubsystem;
 };
 
 #endif
diff --git a/Engine/Plugins/Experimental/UAF/UAF/Source/UAF/Internal/Component/AnimNextWorldSubsystem.h b/Engine/Plugins/Experimental/UAF/UAF/Source/UAF/Internal/Component/AnimNextWorldSubsystem.h
new file mode 100644
--- /dev/null
+++ b/Engine/Plugins/Experimental/UAF/UAF/Source/UAF/Internal/Component/AnimNextWorldSubsystem.h
@@ -0,0 +1,228 @@
+// Copyright Epic Games, Inc. All Rights Reserved.
+
+#pragma once
+
+#include "AnimNextPool.h"
+#include "EngineDefines.h"
+#include "Misc/TVariant.h"
+#include "Module/ModuleHandle.h"
+#include "Module/ModuleInstanceUtils.h"
+#include "Module/TaskRunLocation.h"
+#include "Subsystems/WorldSubsystem.h"
+#include "AnimNextWorldSubsystem.generated.h"
+
+#define UE_API UAF_API
+
+class UAnimNextModule;
+class UAnimNextRigVMAsset;
+enum class EAnimNextModuleInitMethod : uint8;
+enum class EPropertyBagResult : uint8;
+struct FAnimNextParamType;
+struct FAnimNextTraitEvent;
+struct FAnimNextVariableReference;
+
+namespace UE::UAF
+{
+struct FInjectionRequest;
+struct FModuleContext;
+struct FModuleEventTickFunction;
+struct FModuleTaskContext;
+
+// A queued action to complete next frame
+struct FModulePendingAction
+{
+	enum class EType : int32
+	{
+		None = 0,
+		ReleaseHandle,			// Payload = FModulePayloadNone
+		EnableHandle,			// Payload = bool
+		EnableDebugDrawing,		// Payload = bool
+	};
+
+	// Marker struct for an 'empty' payload
+	struct FModulePayloadNone
+	{
+	};
+
+	FModulePendingAction() = default;
+
+	FModulePendingAction(FModuleHandle InHandle, EType InType)
+		: Handle(InHandle)
+		, Type(InType)
+	{
+		Payload.Set<FModulePayloadNone>(FModulePayloadNone());
+	}
+
+	template<typename PayloadType>
+	FModulePendingAction(FModuleHandle InHandle, EType InType, PayloadType InPayload)
+		: Handle(InHandle)
+		, Type(InType)
+	{
+		Payload.Set<PayloadType>(InPayload);
+	}
+
+	TVariant<FModulePayloadNone, bool, UE::UAF::FModuleHandle> Payload;
+
+	FModuleHandle Handle;
+
+	EType Type = EType::None;
+};
+
+}
+
+// Represents AnimNext systems to the gameplay framework
+UCLASS(MinimalAPI, Abstract)
+class UAnimNextWorldSubsystem : public UWorldSubsystem
+{
+	GENERATED_BODY()
+
+public:
+	UE_API UAnimNextWorldSubsystem();
+
+	UE_API virtual void BeginDestroy() override;
+
+	// Get the subsystem for the specified object's world
+	static UE_API UAnimNextWorldSubsystem* Get(UObject* InObject);
+
+	static UE_API void AddReferencedObjects(UObject* InThis, FReferenceCollector& Collector);
+
+	// Dependency type passed to AddDependencyHandle etc.
+	enum class EDependency : uint8
+	{
+		// Dependency runs before the specified event
+		Prerequisite,
+		// Dependency runs after the specified event
+		Subsequent
+	};
+
+private:
+	// UWorldSubsystem interface
+	UE_API virtual bool DoesSupportWorldType(const EWorldType::Type WorldType) const override;
+
+	UE_API bool IsValidHandle(UE::UAF::FModuleHandle InHandle) const;
+
+	UE_API void FlushPendingActions();
+
+protected:
+#if WITH_EDITOR
+	// Refresh any entries that use the provided asset as it has been recompiled.
+	UE_API void OnCompileJobFinished(UAnimNextRigVMAsset* InAsset);
+#endif
+
+	// Register a handle to the subsystem
+	UE_API void RegisterHandle(UE::UAF::FModuleHandle& InOutHandle, const UObject* InAsset, UObject* InObject, EAnimNextModuleInitMethod InitMethod);
+
+	// Unregister a handle from the subsystem
+	UE_API void UnregisterHandle(UE::UAF::FModuleHandle& InOutHandle);
+
+	// Returns whether the module represented by the supplied handle is enabled.
+	// If there are pending actions they will take precedent over the actual current state.
+	UE_API bool IsHandleEnabled(UE::UAF::FModuleHandle InHandle) const;
+
+	// Enables or disables the module represented by the supplied handle
+	// This operation is deferred until the next time the schedule ticks
+	UE_API void EnableHandle(UE::UAF::FModuleHandle InHandle, bool bInEnabled);
+
+#if UE_ENABLE_DEBUG_DRAWING
+	// Enables or disables the module's debug drawing
+	// This operation is deferred until the next time the schedule ticks
+	UE_API void ShowDebugDrawingHandle(UE::UAF::FModuleHandle InHandle, bool bInShowDebugDrawing);
+#endif
+
+	// Queue a task to run at a particular point in a schedule
+	// @param	InHandle			The handle to execute the task on
+	// @param	InModuleEventName	The name of the event in the module to run the supplied task relative to. If this is NAME_None, then the first valid event will be used.
+	// @param	InTaskFunction		The function to run
+	// @param	InLocation			Where to run the task, before or after
+	UE_API void QueueTaskHandle(UE::UAF::FModuleHandle InHandle, FName InModuleEventName, TUniqueFunction<void(const UE::UAF::FModuleTaskContext&)>&& InTaskFunction, UE::UAF::ETaskRunLocation InLocation = UE::UAF::ETaskRunLocation::Before);
+
+	// Queue an input trait event on the supplied handle
+	UE_API void QueueInputTraitEventHandle(UE::UAF::FModuleHandle InHandle, TSharedPtr<FAnimNextTraitEvent> Event);
+
+	// Find the const tick function for the specified event. Useful to analyze/log prerequisites. Use AddDependencyHandle, RemoveDependencyHandle to actually modify dependencies
+	// @param	InHandle			The handle of the module
+	// @param	InEventName			The event associated to the wanted tick function
+	UE_API const FTickFunction* FindTickFunctionHandle(UE::UAF::FModuleHandle InHandle, FName InEventName) const;
+
+	// Add a dependency on a tick function to the specified event
+	// @param	InHandle			The handle of the module whose dependencies are being modified
+	// @param	InObject			The object that owns the tick function
+	// @param	InTickFunction		The tick function to depend on
+	// @param	InEventName			The event to add the dependency to
+	// @param	InDependency		The kind of dependency to add
+	UE_API void AddDependencyHandle(UE::UAF::FModuleHandle InHandle, UObject* InObject, FTickFunction& InTickFunction, FName InEventName, EDependency InDependency);
+
+	// Remove a dependency on a tick function from the specified event
+	// @param	InHandle			The handle of the module whose dependencies are being modified
+	// @param	InObject			The object that owns the tick function
+	// @param	InTickFunction		The tick function that was depended on
+	// @param	InEventName			The event to remove the dependency from
+	// @param	InDependency		The kind of dependency to remove
+	UE_API void RemoveDependencyHandle(UE::UAF::FModuleHandle InHandle, UObject* InObject, FTickFunction& InTickFunction, FName InEventName, EDependency InDependency);
+
+	// Add a dependency on a module event to the specified module event
+	// @param	InHandle			The handle of the module whose dependencies are being modified
+	// @param	InEventName			The event to add the dependency to/from
+	// @param	OtherHandle			The handle of the other module whose dependencies are being modified
+	// @param	OtherEventName		The other module's event to add the dependency to/from
+	// @param	InDependency		The kind of dependency to add
+	UE_API void AddModuleEventDependencyHandle(UE::UAF::FModuleHandle InHandle, FName InEventName, UE::UAF::FModuleHandle OtherHandle, FName OtherEventName, EDependency InDependency);
+
+	// Remove a dependency on a module event from the specified module event
+	// @param	InHandle			The handle of the module whose dependencies are being modified
+	// @param	InEventName			The event to remove the dependency to/from
+	// @param	OtherHandle			The handle of the other module whose dependencies are being modified
+	// @param	OtherEventName		The other module's event to remove the dependency to/from
+	// @param	InDependency		The kind of dependency to remove
+	UE_API void RemoveModuleEventDependencyHandle(UE::UAF::FModuleHandle InHandle, FName InEventName, UE::UAF::FModuleHandle OtherHandle, FName OtherEventName, EDependency InDependency);
+
+	// Set the value of the specified variable
+	// @param	InHandle			The handle of the module whose variable we are setting
+	// @param	InVariable			The variable we want to set
+	// @param	InType				The type of the variable we want to set
+	// @param	InData				The data to set the variable with
+	// @return true if the variable was set successfully
+	UE_API EPropertyBagResult SetVariableHandle(UE::UAF::FModuleHandle InHandle, const FAnimNextVariableReference& InVariable, const FAnimNextParamType& InType, TConstArrayView<uint8> InData);
+
+	// Accesses the variable of the specified name for writing.
+	// General read access is not provided via this API due to the current double-buffering strategy used to communicate variable writes to worker threads.
+	// This is intended to allow for copy-free writing of larger data structures & arrays, rather than read access.
+	// @param	InHandle			The handle of the module whose variable we are setting
+	// @param	InVariable			The variable we want to access
+	// @param	InType				The type of the variable we want to access
+	// @param	InFunction			Function that will be called to allow modification of the variables
+	// @return true if the variable was accessed successfully
+	UE_API EPropertyBagResult WriteVariableHandle(UE::UAF::FModuleHandle InHandle, const FAnimNextVariableReference& InVariable, const FAnimNextParamType& InType, TFunctionRef<void(TArrayView<uint8>)> InFunction);
+
+protected:
+	friend struct FRigUnit_AnimNextInitializeEvent;
+	friend struct UE::UAF::FModuleEventTickFunction;
+	friend struct UE::UAF::FInjectionRequest;
+	friend EPropertyBagResult UE::UAF::SetModuleVariable(UE::UAF::FModuleContext&&, const FAnimNextVariableReference&, const FAnimNextParamType&, TConstArrayView<uint8>);
+	friend EPropertyBagResult UE::UAF::WriteModuleVariable(UE::UAF::FModuleContext&&, const FAnimNextVariableReference&, const FAnimNextParamType&, TFunctionRef<void(TArrayView<uint8>)>);
+
+	// Currently running instances, pooled
+	UE::UAF::TPool<FAnimNextModuleInstance> Instances;
+
+	// Queued actions
+	TArray<UE::UAF::FModulePendingAction> PendingActions;
+
+	// Locks and validators for concurrent modifications
+	FRWLock InstancesLock;
+	FRWLock PendingLock;
+	UE_MT_DECLARE_RW_FULLY_RECURSIVE_ACCESS_DETECTOR(InstancesAccessDetector);
+	UE_MT_DECLARE_RW_ACCESS_DETECTOR(PendingActionsAccessDetector);
+
+#if WITH_EDITOR
+	// Handle used to hook asset compilation
+	FDelegateHandle OnCompileJobFinishedHandle;
+#endif
+
+	// Handle used to hook into pre-world tick
+	FDelegateHandle OnWorldPreActorTickHandle;
+
+	// Cached delta time
+	float DeltaTime;
+};
+
+#undef UE_API
\ No newline at end of file
diff --git a/Engine/Plugins/Experimental/UAF/UAF/Source/UAF/Internal/Module/AnimNextModuleInstance.h b/Engine/Plugins/Experimental/UAF/UAF/Source/UAF/Internal/Module/AnimNextModuleInstance.h
--- a/Engine/Plugins/Experimental/UAF/UAF/Source/UAF/Internal/Module/AnimNextModuleInstance.h
+++ b/Engine/Plugins/Experimental/UAF/UAF/Source/UAF/Internal/Module/AnimNextModuleInstance.h
@@ -4,7 +4,7 @@
 
 #include "CoreMinimal.h"
 #include "AnimNextModule.h"
-#include "AnimNextPoolHandle.h"
+#include "Module/ModuleHandle.h"
 #include "Variables/UAFInstanceVariableDataProxy.h"
 #include "UAFAssetInstance.h"
 #include "Module/ModuleTickFunction.h"
@@ -16,25 +16,20 @@
 #include "Module/TaskRunLocation.h"
 #include "AnimNextModuleInstance.generated.h"
 
-#define UE_API UAF_API
-
 struct FAnimNextModuleInstance;
 struct FRigUnit_AnimNextRunAnimationGraph_v1;
 struct FRigUnit_AnimNextRunAnimationGraph_v2;
-class UUAFEngineSubsystem;
+class UAnimNextWorldSubsystem;
 struct FRigUnit_AnimNextInitializeEvent;
 struct FRigUnit_CopyModuleProxyVariables;
 struct FAnimNextNotifyDispatcherComponent;
 struct FAnimNextGraphInstance;
 struct FUAFSystemFactoryParams;
-struct FUAFWeakSystemReference;
 
 namespace UE::UAF
 {
 	struct FModuleEventTickFunction;
 	struct FModuleEndTickFunction;
-	struct FSystemReference;
-	struct FCustomSystemDeleter;
 }
 
 namespace UE::UAF::Debug
@@ -53,93 +48,96 @@ namespace UE::UAF::Tests
 	class FVariables_UOLBindings;
 }
 
-// Represents an instance of a UAF system.
+// Root memory owner of a parameterized schedule 
 USTRUCT()
 struct FAnimNextModuleInstance : public FUAFAssetInstance
-#if CPP
-	, TSharedFromThis<FAnimNextModuleInstance>
-#endif
 {
 	GENERATED_BODY()
 
 	FAnimNextModuleInstance();
 
 	// Standalone constructor
-	UE_API FAnimNextModuleInstance(
-		const UAnimNextModule* InModule,
+	UAF_API FAnimNextModuleInstance(
+	const UAnimNextModule* InModule,
 		EAnimNextModuleInitMethod InInitMethod);
 
 	// Standalone constructor, UObject host
-	UE_API FAnimNextModuleInstance(
-		const UAnimNextModule* InModule,
+	UAF_API FAnimNextModuleInstance(
+	const UAnimNextModule* InModule,
 		UObject* InObject,
 		EAnimNextModuleInitMethod InInitMethod);
 
-	UE_API ~FAnimNextModuleInstance();
+	UAF_API ~FAnimNextModuleInstance();
 
 	// Checks to see if this entry is ticking
 	bool IsEnabled() const;
 
 	// Enables/disables the ticking of this entry
-	void SetEnabled(bool bInEnabled);
+	void Enable(bool bInEnabled);
 
 	// Queues an input trait event
 	// Input events will be processed in the next graph update after they are queued
-	UE_API void QueueInputTraitEvent(FAnimNextTraitEventPtr Event);
+	UAF_API void QueueInputTraitEvent(FAnimNextTraitEventPtr Event);
 
 	// Queues an output trait event
 	// Output events will be processed at the end of the schedule tick
-	UE_API void QueueOutputTraitEvent(FAnimNextTraitEventPtr Event);
+	UAF_API void QueueOutputTraitEvent(FAnimNextTraitEventPtr Event);
 
 	// Get the object that this module is bound to
 	UObject* GetObject() const { return Object; }
 
-	// Get the system that this instance represents
-	UAF_API const UAnimNextModule* GetSystem() const;
+	// Get the module that this instance represents
+	UAF_API const UAnimNextModule* GetModule() const;
 
 #if UE_ENABLE_DEBUG_DRAWING
 	// Get the debug draw interface
-	UE_API FRigVMDrawInterface* GetDebugDrawInterface();
+	UAF_API FRigVMDrawInterface* GetDebugDrawInterface();
 
 	// Set whether to show the module's debug drawing instructions in the viewport
-	UE_API void ShowDebugDrawing(bool bInShowDebugDrawing);
+	UAF_API void ShowDebugDrawing(bool bInShowDebugDrawing);
 #endif
 
 	// Run a simple task on the GT via FFunctionGraphTask::CreateAndDispatchWhenReady
-	UE_API static void RunTaskOnGameThread(TUniqueFunction<void(void)>&& InFunction);
+	UAF_API static void RunTaskOnGameThread(TUniqueFunction<void(void)>&& InFunction);
 
 	// Find the tick function for the specified event
 	UE::UAF::FModuleEventTickFunction* FindTickFunctionByName(FName InEventName);
 	const UE::UAF::FModuleEventTickFunction* FindTickFunctionByName(FName InEventName) const;
 
 	// Find the first 'user' tick function
-	UE_API UE::UAF::FModuleEventTickFunction* FindFirstUserTickFunction();
+	UAF_API UE::UAF::FModuleEventTickFunction* FindFirstUserTickFunction();
 
 	// Run the specified script event using a module context
-	UE_API void RunScriptEvent(const UE::UAF::FScriptEvent& InEvent, float InDeltaTime);
+	UAF_API void RunScriptEvent(const UE::UAF::FScriptEvent& InEvent, float InDeltaTime);
 
 	// Run the specified script event using a module context, using name lookup alone (slower)
-	UE_API void RunScriptEventByName(FName InEventName, float InDeltaTime);
+	UAF_API void RunScriptEventByName(FName InEventName, float InDeltaTime);
 
 	// Get the world type that this module was instantiated within
 	EWorldType::Type GetWorldType() const { return WorldType; }
 
 	// Get tick functions for this module instance
-	UE_API TArrayView<UE::UAF::FModuleEventTickFunction> GetTickFunctions();
+	UAF_API TArrayView<UE::UAF::FModuleEventTickFunction> GetTickFunctions();
 
-	// Get a weak, non-owning reference to this instance
-	UE_API FUAFWeakSystemReference GetReference() const;
+	UE::UAF::FModuleHandle GetHandle() const { return Handle; }
+
+	// Queue a task to run at a particular module event
+	// @param	InModuleEvent		The event in the module to run the supplied task relative to. If this is NAME_None or not found then the first 'user' event will be used.
+	// @param	InTaskFunction		The function to run
+	// @param	InLocation			Where to run the task, before or after
+	UAF_API void QueueTask(FName InEventName, TUniqueFunction<void(const UE::UAF::FModuleTaskContext&)>&& InTaskFunction, UE::UAF::ETaskRunLocation InLocation = UE::UAF::ETaskRunLocation::Before);
 
-	// Queue a task to run at a particular system event
-	// @param	InEventName			The event in the system to run the supplied task relative to. If this is NAME_None or not found then the first 'user' event will be used.
+	// Queue a task to run at a particular module event on some other module
+	// @param	InModuleHandle		The handle for the module to queue the task on
+	// @param	InModuleEventName	The event in the module to run the supplied task relative to. If this is NAME_None or not found then the first user 'event' will be used.
 	// @param	InTaskFunction		The function to run
 	// @param	InLocation			Where to run the task, before or after
-	UE_API void QueueTask(FName InEventName, TUniqueFunction<void(const UE::UAF::FModuleTaskContext&)>&& InTaskFunction, UE::UAF::ETaskRunLocation InLocation = UE::UAF::ETaskRunLocation::Before);
+	UAF_API void QueueTaskOnOtherModule(const UE::UAF::FModuleHandle InOtherModuleHandle, FName InEventName, TUniqueFunction<void(const UE::UAF::FModuleTaskContext&)>&& InTaskFunction, UE::UAF::ETaskRunLocation InLocation = UE::UAF::ETaskRunLocation::Before);
 
 private:
 	// Setup the entry
 	// Note that we perform initialization in two phases because we may need to access weak references to 'this' during initialization
-	UE_API void Initialize(const FUAFSystemFactoryParams& InFactoryParams);
+	UAF_API void Initialize(const FUAFSystemFactoryParams& InFactoryParams, UE::UAF::TPool<FAnimNextModuleInstance>* InPool = nullptr);
 
 	// Clear data that binds the schedule to a runtime (e.g. tick functions) and any instance data
 	void ResetBindingsAndInstanceData();
@@ -152,7 +150,7 @@ private:
 
 #if UAF_TRACE_ENABLED
 	// trace this module if it hasn't already been traced this frame
-	UE_API void Trace();
+	UAF_API void Trace();
 #endif
 
 #if WITH_EDITOR
@@ -174,25 +172,28 @@ private:
 	EPropertyBagResult WriteProxyVariable(const FAnimNextVariableReference& InVariable, const FAnimNextParamType& InType, TFunctionRef<void(TArrayView<uint8>)> InFunction);
 
 private:
-	friend UE::UAF::FSystemReference;
-	friend FUAFWeakSystemReference;
 	friend UE::UAF::FModuleEventTickFunction;
 	friend UE::UAF::FModuleEndTickFunction;
 	friend FRigUnit_AnimNextRunAnimationGraph_v1;
 	friend FRigUnit_AnimNextRunAnimationGraph_v2;
-	friend UUAFEngineSubsystem;
+	friend UAnimNextWorldSubsystem;
 	friend FRigUnit_AnimNextInitializeEvent;
 	friend FRigUnit_CopyModuleProxyVariables;
 	friend FAnimNextNotifyDispatcherComponent;
 	friend FAnimNextGraphInstance;
 	friend UE::UAF::Tests::FVariables_Shared;
 	friend UE::UAF::Tests::FVariables_UOLBindings;
-	friend UE::UAF::FCustomSystemDeleter;
 
 	// Object this entry is bound to
 	UPROPERTY(Transient)
 	TObjectPtr<UObject> Object = nullptr;
 
+	// The pool that this module instance exists in
+	UE::UAF::TPool<FAnimNextModuleInstance>* Pool = nullptr;
+
+	// Copy of the handle that represents this entry to client systems
+	UE::UAF::FModuleHandle Handle;
+
 	// Pre-allocated graph of tick functions
 	TArray<UE::UAF::FModuleEventTickFunction> TickFunctions;
 
@@ -291,5 +292,3 @@ struct TStructOpsTypeTraits<FAnimNextModuleInstance> : public TStructOpsTypeTrai
 #else
 #define TRACE_UAF_MODULE(ModuleInstance)
 #endif
-
-#undef UE_API
\ No newline at end of file
diff --git a/Engine/Plugins/Experimental/UAF/UAF/Source/UAF/Private/AnimNextModuleImpl.cpp b/Engine/Plugins/Experimental/UAF/UAF/Source/UAF/Private/AnimNextModuleImpl.cpp
--- a/Engine/Plugins/Experimental/UAF/UAF/Source/UAF/Private/AnimNextModuleImpl.cpp
+++ b/Engine/Plugins/Experimental/UAF/UAF/Source/UAF/Private/AnimNextModuleImpl.cpp
@@ -14,7 +14,7 @@
 #include "Curves/CurveFloat.h"
 #include "Graph/AnimNext_LODPose.h"
 #include "Module/AnimNextModuleInstance.h"
-#include "Module/UAFWeakSystemReference.h"
+#include "Module/ModuleHandle.h"
 #include "Modules/ModuleManager.h"
 #include "Param/AnimNextActorLocatorFragment.h"
 #include "Param/AnimNextComponentLocatorFragment.h"
@@ -78,7 +78,7 @@ namespace UE::UAF
 			FRigVMGraphFunctionHeader::StaticStruct(),
 			TBaseStructure<FGuid>::Get(),
 			FRigVMVariant::StaticStruct(),
-			FUAFWeakSystemReference::StaticStruct(),
+			FAnimNextModuleHandle::StaticStruct(),
 			FAnimNextGraphLODPose::StaticStruct(),
 			FAnimNextGraphReferencePose::StaticStruct(),
 		};
diff --git a/Engine/Plugins/Experimental/UAF/UAF/Source/UAF/Private/Component/AnimNextComponent.cpp b/Engine/Plugins/Experimental/UAF/UAF/Source/UAF/Private/Component/AnimNextComponent.cpp
--- a/Engine/Plugins/Experimental/UAF/UAF/Source/UAF/Private/Component/AnimNextComponent.cpp
+++ b/Engine/Plugins/Experimental/UAF/UAF/Source/UAF/Private/Component/AnimNextComponent.cpp
@@ -2,14 +2,13 @@
 
 #include "Component/AnimNextComponent.h"
 
+#include "AnimNextComponentWorldSubsystem.h"
 #include "Blueprint/BlueprintExceptionInfo.h"
 #include "Engine/World.h"
 #include "Module/ModuleTaskContext.h"
 #include "UObject/FortniteMainBranchObjectVersion.h"
 #include "Variables/AnimNextVariableReference.h"
 #include "Module/AnimNextModule.h"
-#include "Module/SystemDependency.h"
-#include "Module/UAFWeakSystemReference.h"
 
 #include UE_INLINE_GENERATED_CPP_BY_NAME(AnimNextComponent)
 
@@ -19,13 +18,15 @@ void UAnimNextComponent::OnRegister()
 
 	Super::OnRegister();
 
-	if (Asset)
+	Subsystem = GetWorld()->GetSubsystem<UAnimNextComponentWorldSubsystem>();
+
+	if (Subsystem && Asset)
 	{
-		check(!SystemReference.IsValid());
+		check(!ModuleHandle.IsValid());
 
-		SystemReference = FSystemReference(Asset, this, InitMethod);
+		Subsystem->Register(this);
 #if UE_ENABLE_DEBUG_DRAWING
-		SystemReference.ShowDebugDrawing(bShowDebugDrawing);
+		Subsystem->ShowDebugDrawing(this, bShowDebugDrawing);
 #endif
 	}
 }
@@ -34,7 +35,11 @@ void UAnimNextComponent::OnUnregister()
 {
 	Super::OnUnregister();
 
-	SystemReference.Reset();
+	if(Subsystem)
+	{
+		Subsystem->Unregister(this);
+		Subsystem = nullptr;
+	}
 }
 
 void UAnimNextComponent::BeginPlay()
@@ -47,7 +52,7 @@ void UAnimNextComponent::BeginPlay()
 	if (InitMethod == EAnimNextModuleInitMethod::InitializeAndRun
 	|| (InitMethod == EAnimNextModuleInitMethod::InitializeAndPauseInEditor && World->WorldType != EWorldType::Editor && World->WorldType != EWorldType::EditorPreview))
 	{
-		Activate();
+		SetEnabled(true);
 	}
 }
 
@@ -55,7 +60,7 @@ void UAnimNextComponent::EndPlay(const EEndPlayReason::Type EndPlayReason)
 {
 	Super::EndPlay(EndPlayReason);
 
-	Deactivate();
+	SetEnabled(false);
 }
 
 void UAnimNextComponent::Serialize(FArchive& Ar)
@@ -270,31 +275,36 @@ DEFINE_FUNCTION(UAnimNextComponent::execBlueprintSetVariableReference)
 
 EPropertyBagResult UAnimNextComponent::SetVariableInternal(const FAnimNextVariableReference& InVariable, const FAnimNextParamType& InType, TConstArrayView<uint8> InNewValue)
 {
-	return SystemReference.SetVariable(InVariable, InType, InNewValue);
+	if (Subsystem)
+	{
+		return Subsystem->SetVariable(this, InVariable, InType, InNewValue);
+	}
+	return EPropertyBagResult::PropertyNotFound;
 }
 
 EPropertyBagResult UAnimNextComponent::WriteVariableInternal(const FAnimNextVariableReference& InVariable, const FAnimNextParamType& InType, TFunctionRef<void(TArrayView<uint8>)> InFunction)
 {
-	return SystemReference.WriteVariable(InVariable, InType, InFunction);
+	if (Subsystem)
+	{
+		return Subsystem->WriteVariable(this, InVariable, InType, InFunction);
+	}
+	return EPropertyBagResult::PropertyNotFound;
 }
 
-void UAnimNextComponent::Activate(bool bReset)
+bool UAnimNextComponent::IsEnabled() const
 {
-	if (bReset || ShouldActivate())
+	if (Subsystem)
 	{
-		SystemReference.SetEnabled(true);
-		SetActiveFlag(true);
-		OnComponentActivated.Broadcast(this, bReset);
+		return Subsystem->IsEnabled(this);
 	}
+	return false;
 }
 
-void UAnimNextComponent::Deactivate()
+void UAnimNextComponent::SetEnabled(bool bEnabled)
 {
-	if (!ShouldActivate())
+	if(Subsystem)
 	{
-		SystemReference.SetEnabled(false);
-		SetActiveFlag(false);
-		OnComponentDeactivated.Broadcast(this);
+		Subsystem->SetEnabled(this, bEnabled);
 	}
 }
 
@@ -302,7 +312,10 @@ void UAnimNextComponent::ShowDebugDrawing(bool bInShowDebugDrawing)
 {
 #if UE_ENABLE_DEBUG_DRAWING
 	bShowDebugDrawing = bInShowDebugDrawing;
-	SystemReference.ShowDebugDrawing(bInShowDebugDrawing);
+	if(Subsystem)
+	{
+		Subsystem->ShowDebugDrawing(this, bInShowDebugDrawing);
+	}
 #endif
 }
 
@@ -310,24 +323,37 @@ void UAnimNextComponent::QueueTask(FName InModuleEventName, TUniqueFunction<void
 {
 	using namespace UE::UAF;
 
-	SystemReference.QueueTask(InModuleEventName, MoveTemp(InTaskFunction), InLocation);
+	if(Subsystem)
+	{
+		Subsystem->QueueTask(this, InModuleEventName, MoveTemp(InTaskFunction), InLocation);
+	}
 }
 
 void UAnimNextComponent::QueueInputTraitEvent(FAnimNextTraitEventPtr Event)
 {
 	using namespace UE::UAF;
 
-	SystemReference.QueueInputTraitEvent(Event);
+	if(Subsystem)
+	{
+		Subsystem->QueueInputTraitEvent(this, Event);
+	}
 }
 
 const FTickFunction* UAnimNextComponent::FindTickFunction(FName InEventName) const
 {
-	return SystemReference.FindTickFunction(InEventName);
+	if(Subsystem)
+	{
+		return Subsystem->FindTickFunction(this, InEventName);
+	}
+	return nullptr;
 }
 
 void UAnimNextComponent::AddPrerequisite(UObject* InObject, FTickFunction& InTickFunction, FName InEventName)
 {
-	SystemReference.AddDependency(InObject, InTickFunction, InEventName, UE::UAF::ESystemDependency::Prerequisite);
+	if(Subsystem)
+	{
+		Subsystem->AddDependency(this, InObject, InTickFunction, InEventName, UAnimNextWorldSubsystem::EDependency::Prerequisite);
+	}
 }
 
 void UAnimNextComponent::AddComponentPrerequisite(UActorComponent* InComponent, FName InEventName)
@@ -340,7 +366,10 @@ void UAnimNextComponent::AddComponentPrerequisite(UActorComponent* InComponent,
 
 void UAnimNextComponent::AddSubsequent(UObject* InObject, FTickFunction& InTickFunction, FName InEventName)
 {
-	SystemReference.AddDependency(InObject, InTickFunction, InEventName, UE::UAF::ESystemDependency::Subsequent);
+	if(Subsystem)
+	{
+		Subsystem->AddDependency(this, InObject, InTickFunction, InEventName, UAnimNextWorldSubsystem::EDependency::Subsequent);
+	}
 }
 
 void UAnimNextComponent::AddComponentSubsequent(UActorComponent* InComponent, FName InEventName)
@@ -353,7 +382,10 @@ void UAnimNextComponent::AddComponentSubsequent(UActorComponent* InComponent, FN
 
 void UAnimNextComponent::RemovePrerequisite(UObject* InObject, FTickFunction& InTickFunction, FName InEventName)
 {
-	SystemReference.RemoveDependency(InObject, InTickFunction, InEventName, UE::UAF::ESystemDependency::Prerequisite);
+	if(Subsystem)
+	{
+		Subsystem->RemoveDependency(this, InObject, InTickFunction, InEventName, UAnimNextWorldSubsystem::EDependency::Prerequisite);
+	}
 }
 
 void UAnimNextComponent::RemoveComponentPrerequisite(UActorComponent* InComponent, FName InEventName)
@@ -366,7 +398,10 @@ void UAnimNextComponent::RemoveComponentPrerequisite(UActorComponent* InComponen
 
 void UAnimNextComponent::RemoveSubsequent(UObject* InObject, FTickFunction& InTickFunction, FName InEventName)
 {
-	SystemReference.RemoveDependency(InObject, InTickFunction, InEventName, UE::UAF::ESystemDependency::Subsequent);
+	if(Subsystem)
+	{
+		Subsystem->RemoveDependency(this, InObject, InTickFunction, InEventName, UAnimNextWorldSubsystem::EDependency::Subsequent);
+	}
 }
 
 void UAnimNextComponent::RemoveComponentSubsequent(UActorComponent* InComponent, FName InEventName)
@@ -387,9 +422,9 @@ void UAnimNextComponent::AddModuleEventPrerequisite(FName InEventName, UAnimNext
 	{
 		UE_LOG(LogAnimation, Warning, TEXT("UAFComponent::AddModuleEventPrerequisite called using the same component"));
 	}
-	else
+	else if (Subsystem)
 	{
-		SystemReference.AddSystemEventDependency(InEventName, OtherAnimNextComponent->GetSystemReference(), OtherEventName, UE::UAF::ESystemDependency::Prerequisite);
+		Subsystem->AddModuleEventDependency(this, InEventName, OtherAnimNextComponent, OtherEventName, UAnimNextWorldSubsystem::EDependency::Prerequisite);
 	}
 }
 
@@ -403,9 +438,9 @@ void UAnimNextComponent::AddModuleEventSubsequent(FName InEventName, UAnimNextCo
 	{
 		UE_LOG(LogAnimation, Warning, TEXT("UAFComponent::AddModuleEventPrerequisite called using the same component"));
 	}
-	else
+	else if (Subsystem)
 	{
-		SystemReference.AddSystemEventDependency(InEventName, OtherAnimNextComponent->GetSystemReference(), OtherEventName, UE::UAF::ESystemDependency::Subsequent);
+		Subsystem->AddModuleEventDependency(this, InEventName, OtherAnimNextComponent, OtherEventName, UAnimNextWorldSubsystem::EDependency::Subsequent);
 	}
 }
 
@@ -419,9 +454,9 @@ void UAnimNextComponent::RemoveModuleEventPrerequisite(FName InEventName, UAnimN
 	{
 		UE_LOG(LogAnimation, Warning, TEXT("UAnimNextComponent::AddModuleEventPrerequisite called using the same component"));
 	}
-	else
+	else if (Subsystem)
 	{
-		SystemReference.RemoveSystemEventDependency(InEventName, OtherAnimNextComponent->GetSystemReference(), OtherEventName, UE::UAF::ESystemDependency::Prerequisite);
+		Subsystem->RemoveModuleEventDependency(this, InEventName, OtherAnimNextComponent, OtherEventName, UAnimNextWorldSubsystem::EDependency::Prerequisite);
 	}
 }
 
@@ -435,15 +470,15 @@ void UAnimNextComponent::RemoveModuleEventSubsequent(FName InEventName, UAnimNex
 	{
 		UE_LOG(LogAnimation, Warning, TEXT("UAnimNextComponent::AddModuleEventPrerequisite called using the same component"));
 	}
-	else
+	else if (Subsystem)
 	{
-		SystemReference.RemoveSystemEventDependency(InEventName, OtherAnimNextComponent->GetSystemReference(), OtherEventName, UE::UAF::ESystemDependency::Subsequent);
+		Subsystem->RemoveModuleEventDependency(this, InEventName, OtherAnimNextComponent, OtherEventName, UAnimNextWorldSubsystem::EDependency::Subsequent);
 	}
 }
 
-FUAFWeakSystemReference UAnimNextComponent::GetSystemReference() const
+FAnimNextModuleHandle UAnimNextComponent::BlueprintGetModuleHandle() const
 {
-	return FUAFWeakSystemReference(SystemReference);
+	return FAnimNextModuleHandle(ModuleHandle);
 }
 
 void UAnimNextComponent::SetAsset(TObjectPtr<UObject> InAsset)
@@ -451,25 +486,25 @@ void UAnimNextComponent::SetAsset(TObjectPtr<UObject> InAsset)
 	Asset = InAsset;
 }
 
-void UAnimNextComponent::RegisterSystem()
+void UAnimNextComponent::RegisterWithSubsystem()
 {
-	using namespace UE::UAF;
-
-	if (Asset)
+	if (Subsystem && Asset)
 	{
-		check(!SystemReference.IsValid());
-		SystemReference = FSystemReference(Asset, this, InitMethod);
+		Subsystem->Register(this);
 	}
 }
 
-void UAnimNextComponent::UnregisterSystem()
+void UAnimNextComponent::UnregisterWithSubsystem()
 {
-	SystemReference.Reset();
+	if (Subsystem)
+	{
+		Subsystem->Unregister(this);
+	}
 }
 
 bool UAnimNextComponent::IsModuleValid()
 {
-	return SystemReference.IsValid();
+	return ModuleHandle.IsValid();
 }
 
 TObjectPtr<UObject> UAnimNextComponent::GetAsset() const
diff --git a/Engine/Plugins/Experimental/UAF/UAF/Source/UAF/Private/Component/AnimNextComponentWorldSubsystem.cpp b/Engine/Plugins/Experimental/UAF/UAF/Source/UAF/Private/Component/AnimNextComponentWorldSubsystem.cpp
new file mode 100644
--- /dev/null
+++ b/Engine/Plugins/Experimental/UAF/UAF/Source/UAF/Private/Component/AnimNextComponentWorldSubsystem.cpp
@@ -0,0 +1,124 @@
+// Copyright Epic Games, Inc. All Rights Reserved.
+
+#include "AnimNextComponentWorldSubsystem.h"
+
+#include "Component/AnimNextComponent.h"
+#include "Engine/World.h"
+#include "Logging/StructuredLog.h"
+#include "Module/AnimNextModule.h"
+#include "Module/ModuleHandle.h"
+#include "Module/ModuleTaskContext.h"
+
+#include UE_INLINE_GENERATED_CPP_BY_NAME(AnimNextComponentWorldSubsystem)
+
+void UAnimNextComponentWorldSubsystem::Register(UAnimNextComponent* InComponent)
+{
+	using namespace UE::UAF;
+	check(IsInGameThread());
+	check(InComponent);
+	RegisterHandle(InComponent->ModuleHandle, InComponent->Asset, InComponent, InComponent->InitMethod);
+}
+
+void UAnimNextComponentWorldSubsystem::Unregister(UAnimNextComponent* InComponent)
+{
+	using namespace UE::UAF;
+	check(IsInGameThread());
+	check(InComponent);
+	UnregisterHandle(InComponent->ModuleHandle);
+}
+
+void UAnimNextComponentWorldSubsystem::SetEnabled(UAnimNextComponent* InComponent, bool bInEnabled)
+{
+	using namespace UE::UAF;
+	check(IsInGameThread());
+	check(InComponent);
+	EnableHandle(InComponent->ModuleHandle, bInEnabled);
+}
+
+bool UAnimNextComponentWorldSubsystem::IsEnabled(const UAnimNextComponent* InComponent) const
+{
+	using namespace UE::UAF;
+	check(IsInGameThread());
+	check(InComponent);
+	return IsHandleEnabled(InComponent->ModuleHandle);
+}
+
+#if UE_ENABLE_DEBUG_DRAWING
+void UAnimNextComponentWorldSubsystem::ShowDebugDrawing(UAnimNextComponent* InComponent, bool bInShowDebugDrawing)
+{
+	using namespace UE::UAF;
+	check(IsInGameThread());
+	check(InComponent);
+	ShowDebugDrawingHandle(InComponent->ModuleHandle, bInShowDebugDrawing);
+}
+#endif
+
+void UAnimNextComponentWorldSubsystem::QueueTask(UAnimNextComponent* InComponent, FName InModuleEventName, TUniqueFunction<void(const UE::UAF::FModuleTaskContext&)>&& InTaskFunction, UE::UAF::ETaskRunLocation InLocation)
+{
+	using namespace UE::UAF;
+	check(IsInGameThread());
+	check(InComponent);
+	QueueTaskHandle(InComponent->ModuleHandle, InModuleEventName, MoveTemp(InTaskFunction), InLocation);
+}
+
+void UAnimNextComponentWorldSubsystem::QueueInputTraitEvent(UAnimNextComponent* InComponent, FAnimNextTraitEventPtr Event)
+{
+	using namespace UE::UAF;
+	check(IsInGameThread());
+	check(InComponent);
+	QueueInputTraitEventHandle(InComponent->ModuleHandle, Event);
+}
+
+const FTickFunction* UAnimNextComponentWorldSubsystem::FindTickFunction(const UAnimNextComponent* InComponent, FName InEventName) const
+{
+	using namespace UE::UAF;
+	check(IsInGameThread());
+	check(InComponent);
+	return FindTickFunctionHandle(InComponent->ModuleHandle, InEventName);
+}
+
+void UAnimNextComponentWorldSubsystem::AddDependency(UAnimNextComponent* InComponent, UObject* InObject, FTickFunction& InTickFunction, FName InEventName, EDependency InDependency)
+{
+	using namespace UE::UAF;
+	check(IsInGameThread());
+	check(InComponent);
+	AddDependencyHandle(InComponent->ModuleHandle, InObject, InTickFunction, InEventName, InDependency);
+}
+
+void UAnimNextComponentWorldSubsystem::RemoveDependency(UAnimNextComponent* InComponent, UObject* InObject, FTickFunction& InTickFunction, FName InEventName, EDependency InDependency)
+{
+	using namespace UE::UAF;
+	check(IsInGameThread());
+	check(InComponent);
+	RemoveDependencyHandle(InComponent->ModuleHandle, InObject, InTickFunction, InEventName, InDependency);
+}
+
+void UAnimNextComponentWorldSubsystem::AddModuleEventDependency(UAnimNextComponent* InComponent, FName InEventName, UAnimNextComponent* OtherComponent, FName OtherEventName, EDependency InDependency)
+{
+	using namespace UE::UAF;
+	check(IsInGameThread());
+	check(InComponent && OtherComponent && InComponent != OtherComponent);
+	AddModuleEventDependencyHandle(InComponent->ModuleHandle, InEventName, OtherComponent->ModuleHandle, OtherEventName, InDependency);
+}
+
+void UAnimNextComponentWorldSubsystem::RemoveModuleEventDependency(UAnimNextComponent* InComponent, FName InEventName, UAnimNextComponent* OtherComponent, FName OtherEventName, EDependency InDependency)
+{
+	using namespace UE::UAF;
+	check(IsInGameThread());
+	check(InComponent && OtherComponent && InComponent != OtherComponent);
+	RemoveModuleEventDependencyHandle(InComponent->ModuleHandle, InEventName, OtherComponent->ModuleHandle, OtherEventName, InDependency);
+}
+
+EPropertyBagResult UAnimNextComponentWorldSubsystem::SetVariable(UAnimNextComponent* InComponent, const FAnimNextVariableReference& InVariable, const FAnimNextParamType& InType, TConstArrayView<uint8> InData)
+{
+	using namespace UE::UAF;
+	check(InComponent);
+	return SetVariableHandle(InComponent->ModuleHandle, InVariable, InType, InData);
+}
+
+EPropertyBagResult UAnimNextComponentWorldSubsystem::WriteVariable(UAnimNextComponent* InComponent, const FAnimNextVariableReference& InVariable, const FAnimNextParamType& InType, TFunctionRef<void(TArrayView<uint8>)> InFunction)
+{
+	using namespace UE::UAF;
+	check(InComponent);
+	return WriteVariableHandle(InComponent->ModuleHandle, InVariable, InType, InFunction);
+}
diff --git a/Engine/Plugins/Experimental/UAF/UAF/Source/UAF/Private/Component/AnimNextComponentWorldSubsystem.h b/Engine/Plugins/Experimental/UAF/UAF/Source/UAF/Private/Component/AnimNextComponentWorldSubsystem.h
new file mode 100644
--- /dev/null
+++ b/Engine/Plugins/Experimental/UAF/UAF/Source/UAF/Private/Component/AnimNextComponentWorldSubsystem.h
@@ -0,0 +1,104 @@
+// Copyright Epic Games, Inc. All Rights Reserved.
+
+#pragma once
+
+#include "CoreMinimal.h"
+#include "AnimNextPool.h"
+#include "Component/AnimNextWorldSubsystem.h"
+#include "Module/ModuleHandle.h"
+#include "Module/ModuleTaskContext.h"
+#include "AnimNextComponentWorldSubsystem.generated.h"
+
+class UAnimNextComponent;
+
+// Represents AnimNext systems to the AActor/UActorComponent gameplay framework
+UCLASS()
+class UAnimNextComponentWorldSubsystem : public UAnimNextWorldSubsystem
+{
+	GENERATED_BODY()
+
+	friend class UAnimNextComponent;
+
+	// Register a component to the subsystem
+	void Register(UAnimNextComponent* InComponent);
+
+	// Unregister a component to the subsystem
+	// The full release of the module referenced by the component's handle will be deferred after this call is made
+	void Unregister(UAnimNextComponent* InComponent);
+
+	// Returns whether the module is enabled.
+	bool IsEnabled(const UAnimNextComponent* InComponent) const;
+
+	// Enables or disables the module represented by the supplied handle
+	// This operation is deferred until the next time the schedule ticks
+	void SetEnabled(UAnimNextComponent* InComponent, bool bInEnabled);
+
+#if UE_ENABLE_DEBUG_DRAWING
+	// Enables or disables debug drawing for the module represented by the supplied handle
+	// This operation is deferred until the next time the schedule ticks
+	void ShowDebugDrawing(UAnimNextComponent* InComponent, bool bInShowDebugDrawing);
+#endif
+
+	// Queue a task to run at a particular point in a schedule
+	// @param	InComponent			The component to execute the task on
+	// @param	InModuleEventName	The name of the event in the module to run the supplied task relative to. If this is NAME_None, then the first valid event will be used.
+	// @param	InTaskFunction		The function to run
+	// @param	InLocation			Where to run the task, before or after
+	void QueueTask(UAnimNextComponent* InComponent, FName InModuleEventName, TUniqueFunction<void(const UE::UAF::FModuleTaskContext&)>&& InTaskFunction, UE::UAF::ETaskRunLocation InLocation = UE::UAF::ETaskRunLocation::Before);
+
+	// Queue an input trait event
+	void QueueInputTraitEvent(UAnimNextComponent* InComponent, FAnimNextTraitEventPtr Event);
+
+	// Find the component tick function for the specified event
+	// @param	InComponent			The component being searched for the tick function
+	// @param	InEventName			The event associated to the wanted tick function
+	const FTickFunction* FindTickFunction(const UAnimNextComponent* InComponent, FName InEventName) const;
+
+	// Add a dependency on a tick function to the specified event
+	// @param	InComponent			The component whose dependencies are being modified
+	// @param	InObject			The object that owns the tick function
+	// @param	InTickFunction		The tick function to depend on
+	// @param	InEventName			The event to add the dependency to
+	void AddDependency(UAnimNextComponent* InComponent, UObject* InObject, FTickFunction& InTickFunction, FName InEventName, EDependency InDependency);
+
+	// Remove a dependency on a tick function from the specified event
+	// @param	InComponent			The component whose dependencies are being modified
+	// @param	InObject			The object that owns the tick function
+	// @param	InTickFunction		The tick function that was depended on
+	// @param	InEventName			The event to remove the dependency from
+	void RemoveDependency(UAnimNextComponent* InComponent, UObject* InObject, FTickFunction& InTickFunction, FName InEventName, EDependency InDependency);
+
+	// Add a dependency on a module event to the specified module event
+	// @param	InComponent			The component whose dependencies are being modified
+	// @param	InEventName			The event to add the dependency to/from
+	// @param	OtherComponent		The other component whose dependencies are being modified
+	// @param	OtherEventName		The other module's event to add the dependency to/from
+	// @param	InDependency		The kind of dependency to add
+	void AddModuleEventDependency(UAnimNextComponent* InComponent, FName InEventName, UAnimNextComponent* OtherComponent, FName OtherEventName, EDependency InDependency);
+
+	// Remove a dependency on a module event to the specified module event
+	// @param	InComponent			The component whose dependencies are being modified
+	// @param	InEventName			The event to remove the dependency to/from
+	// @param	OtherComponent		The other component whose dependencies are being modified
+	// @param	OtherEventName		The other module's event to remove the dependency to/from
+	// @param	InDependency		The kind of dependency to remove
+	void RemoveModuleEventDependency(UAnimNextComponent* InComponent, FName InEventName, UAnimNextComponent* OtherComponent, FName OtherEventName, EDependency InDependency);
+
+	// Set the value of the specified variable
+	// @param	InComponent			The component whose variable we are setting
+	// @param	InVariable			The variable we want to set
+	// @param	InType				The type of the variable we want to set
+	// @param	InData				The data to set the variable with
+	// @return true if the variable was set successfully
+	EPropertyBagResult SetVariable(UAnimNextComponent* InComponent, const FAnimNextVariableReference& InVariable, const FAnimNextParamType& InType, TConstArrayView<uint8> InData);
+
+	// Accesses the variable of the specified name for writing.
+	// General read access is not provided via this API due to the current double-buffering strategy used to communicate variable writes to worker threads.
+	// This is intended to allow for copy-free writing of larger data structures & arrays, rather than read access. 
+	// @param	InComponent			The component whose variable we are setting
+	// @param	InVariable			The variable we want to access
+	// @param	InType				The type of the variable we want to access
+	// @param	InFunction			Function that will be called to allow modification of the variable
+	// @return true if the variable was accessed successfully
+	EPropertyBagResult WriteVariable(UAnimNextComponent* InComponent, const FAnimNextVariableReference& InVariable, const FAnimNextParamType& InType, TFunctionRef<void(TArrayView<uint8>)> InFunction);
+};
\ No newline at end of file
diff --git a/Engine/Plugins/Experimental/UAF/UAF/Source/UAF/Private/Component/AnimNextWorldSubsystem.cpp b/Engine/Plugins/Experimental/UAF/UAF/Source/UAF/Private/Component/AnimNextWorldSubsystem.cpp
new file mode 100644
--- /dev/null
+++ b/Engine/Plugins/Experimental/UAF/UAF/Source/UAF/Private/Component/AnimNextWorldSubsystem.cpp
@@ -0,0 +1,451 @@
+// Copyright Epic Games, Inc. All Rights Reserved.
+
+#include "Component/AnimNextWorldSubsystem.h"
+
+#include "AnimNextDebugDraw.h"
+#include "AnimNextRigVMAsset.h"
+#include "Engine/World.h"
+#include "Factory/SystemFactory.h"
+#include "Logging/StructuredLog.h"
+#include "Module/AnimNextModuleInstance.h"
+#include "Module/ModuleTaskContext.h"
+
+#include UE_INLINE_GENERATED_CPP_BY_NAME(AnimNextWorldSubsystem)
+
+UAnimNextWorldSubsystem::UAnimNextWorldSubsystem()
+{
+	if (!HasAnyFlags(RF_ClassDefaultObject))
+	{
+#if WITH_EDITOR
+		OnCompileJobFinishedHandle = UAnimNextRigVMAsset::OnCompileJobFinished().AddUObject(this, &UAnimNextWorldSubsystem::OnCompileJobFinished);
+#endif
+
+		// Kick off root task at the start of each world tick
+		OnWorldPreActorTickHandle = FWorldDelegates::OnWorldPreActorTick.AddLambda([this](UWorld* InWorld, ELevelTick InTickType, float InDeltaSeconds)
+		{
+			if (InTickType == LEVELTICK_All || InTickType == LEVELTICK_ViewportsOnly)
+			{
+				// Flush actions here as they require game thread callbacks (e.g. to reconfigure tick functions)
+				FlushPendingActions();
+				DeltaTime = InDeltaSeconds;
+			}
+		});
+	}
+}
+
+void UAnimNextWorldSubsystem::BeginDestroy()
+{
+	Super::BeginDestroy();
+
+	if (!HasAnyFlags(RF_ClassDefaultObject))
+	{
+#if WITH_EDITOR
+		UAnimNextRigVMAsset::OnCompileJobFinished().Remove(OnCompileJobFinishedHandle);
+#endif
+
+		FWorldDelegates::OnWorldPreActorTick.Remove(OnWorldPreActorTickHandle);
+
+		{
+			FRWScopeLock InstancesLockScope(InstancesLock, SLT_Write);
+			UE_MT_SCOPED_WRITE_ACCESS(InstancesAccessDetector);
+
+			for (FAnimNextModuleInstance& Instance : Instances)
+			{
+				Instance.RemoveAllTickDependencies();
+				Instance.ReleaseComponents();
+			}
+			Instances = UE::UAF::TPool<FAnimNextModuleInstance>(); // Force instances destruction
+		}
+	}
+}
+
+UAnimNextWorldSubsystem* UAnimNextWorldSubsystem::Get(UObject* InObject)
+{
+	if (InObject == nullptr)
+	{
+		return nullptr;
+	}
+
+	UWorld* World = InObject->GetWorld();
+	if (World == nullptr)
+	{
+		return nullptr;
+	}
+
+	return World->GetSubsystem<UAnimNextWorldSubsystem>();
+}
+
+void UAnimNextWorldSubsystem::AddReferencedObjects(UObject* InThis, FReferenceCollector& Collector)
+{
+	Super::AddReferencedObjects(InThis, Collector);
+
+	UAnimNextWorldSubsystem* This = CastChecked<UAnimNextWorldSubsystem>(InThis);
+	UE_MT_SCOPED_READ_ACCESS(This->InstancesAccessDetector);
+
+	for (FAnimNextModuleInstance& Instance : This->Instances)
+	{
+		Collector.AddPropertyReferencesWithStructARO(FAnimNextModuleInstance::StaticStruct(), &Instance, InThis);
+	}
+}
+
+bool UAnimNextWorldSubsystem::DoesSupportWorldType(const EWorldType::Type WorldType) const
+{
+	switch (WorldType)
+	{
+	case EWorldType::Game:
+	case EWorldType::Editor:
+	case EWorldType::PIE:
+	case EWorldType::EditorPreview:
+	case EWorldType::GamePreview:
+		return true;
+	}
+
+	return false;
+}
+
+bool UAnimNextWorldSubsystem::IsValidHandle(UE::UAF::FModuleHandle InHandle) const
+{
+	UE_MT_SCOPED_READ_ACCESS(InstancesAccessDetector);
+	return Instances.IsValidHandle(InHandle);
+}
+
+void UAnimNextWorldSubsystem::FlushPendingActions()
+{
+	using namespace UE::UAF;
+
+	FRWScopeLock PendingLockScope(PendingLock, SLT_Write);
+	UE_MT_SCOPED_WRITE_ACCESS(PendingActionsAccessDetector);
+
+	if (PendingActions.Num() > 0)
+	{
+		for (FModulePendingAction& PendingAction : PendingActions)
+		{
+			switch (PendingAction.Type)
+			{
+			case FModulePendingAction::EType::ReleaseHandle:
+				if (IsValidHandle(PendingAction.Handle))
+				{
+					FRWScopeLock InstancesLockScope(InstancesLock, SLT_Write);
+					UE_MT_SCOPED_WRITE_ACCESS(InstancesAccessDetector);
+					Instances.Release(PendingAction.Handle);
+				}
+				break;
+			case FModulePendingAction::EType::EnableHandle:
+				if (IsValidHandle(PendingAction.Handle))
+				{
+					UE_MT_SCOPED_READ_ACCESS(InstancesAccessDetector);
+					FAnimNextModuleInstance& Instance = Instances.Get(PendingAction.Handle);
+					Instance.Enable(PendingAction.Payload.Get<bool>());
+				}
+				break;
+			case FModulePendingAction::EType::EnableDebugDrawing:
+#if UE_ENABLE_DEBUG_DRAWING
+				if (IsValidHandle(PendingAction.Handle))
+				{
+					UE_MT_SCOPED_READ_ACCESS(InstancesAccessDetector);
+					FAnimNextModuleInstance& Instance = Instances.Get(PendingAction.Handle);
+					Instance.ShowDebugDrawing(PendingAction.Payload.Get<bool>());
+				}
+#endif
+				break;
+			default:
+				checkNoEntry();
+				break;
+			}
+		}
+
+		PendingActions.Reset();
+	}
+}
+
+void UAnimNextWorldSubsystem::RegisterHandle(UE::UAF::FModuleHandle& InOutHandle, const UObject* InAsset, UObject* InObject, EAnimNextModuleInitMethod InitMethod)
+{
+	using namespace UE::UAF;
+	check(IsInGameThread());
+	check(!IsValidHandle(InOutHandle));	// Shouldnt br registering without unregistering first
+	{
+		FWriteScopeLock InstancesLockScope(InstancesLock);
+		UE_MT_SCOPED_WRITE_ACCESS(InstancesAccessDetector);
+
+		const UAnimNextModule* Module = Cast<UAnimNextModule>(InAsset);
+		FUAFSystemFactoryParams FactoryParams;
+		if (Module == nullptr)
+		{
+			FactoryParams = FSystemFactory::GetDefaultParamsForObject(InAsset);
+			Module = FSystemFactory::BuildSystem(FactoryParams.GetBuilder());
+		}
+
+		InOutHandle = Instances.Emplace(Module, InObject, InitMethod);
+		FAnimNextModuleInstance& Instance = Instances.Get(InOutHandle);
+		Instance.Handle = InOutHandle;
+		Instance.Initialize(FactoryParams, &Instances);
+	}
+}
+
+void UAnimNextWorldSubsystem::UnregisterHandle(UE::UAF::FModuleHandle& InOutHandle)
+{
+	using namespace UE::UAF;
+	check(IsInGameThread());
+
+	if (IsValidHandle(InOutHandle))
+	{
+		UE_MT_SCOPED_READ_ACCESS(InstancesAccessDetector);
+		FAnimNextModuleInstance& Instance = Instances.Get(InOutHandle);
+
+#if UE_ENABLE_DEBUG_DRAWING
+		// Remove debug drawing immediately as the renderer will need to know about this before EOF
+		Instance.DebugDraw->RemovePrimitive();
+#endif
+
+		// Remove all tick dependencies immediately, as once the handle has been invalidated there is no way for external systems to remove their dependencies
+		Instance.RemoveAllTickDependencies();
+
+		{
+			FRWScopeLock PendingLockScope(PendingLock, SLT_Write);
+			UE_MT_SCOPED_WRITE_ACCESS(PendingActionsAccessDetector);
+			PendingActions.Emplace(InOutHandle, FModulePendingAction::EType::ReleaseHandle);
+		}
+		InOutHandle.Reset();
+	}
+}
+
+bool UAnimNextWorldSubsystem::IsHandleEnabled(UE::UAF::FModuleHandle InHandle) const
+{
+	using namespace UE::UAF;
+	check(IsInGameThread());
+	if (IsValidHandle(InHandle))
+	{
+		// The last pending action takes precedent here if present
+		{
+			UE_MT_SCOPED_READ_ACCESS(PendingActionsAccessDetector);
+			int32 PendingActionIndex = PendingActions.FindLastByPredicate([&InHandle](const FModulePendingAction& PendingAction) { return PendingAction.Handle == InHandle && PendingAction.Type == FModulePendingAction::EType::EnableHandle;});
+
+			if (PendingActionIndex != INDEX_NONE)
+			{
+				return PendingActions[PendingActionIndex].Payload.Get<bool>();
+			}
+		}
+
+		// Otherwise return the current value on the instance
+		{
+			UE_MT_SCOPED_READ_ACCESS(InstancesAccessDetector);
+			const FAnimNextModuleInstance& Instance = Instances.Get(InHandle);
+			return Instance.IsEnabled();
+		}
+	}
+	return false;
+}
+
+void UAnimNextWorldSubsystem::EnableHandle(UE::UAF::FModuleHandle InHandle, bool bInEnabled)
+{
+	using namespace UE::UAF;
+	check(IsInGameThread());
+	if (IsValidHandle(InHandle))
+	{
+		UE_MT_SCOPED_WRITE_ACCESS(PendingActionsAccessDetector);
+		PendingActions.Emplace(InHandle, FModulePendingAction::EType::EnableHandle, bInEnabled);
+	}
+}
+
+#if UE_ENABLE_DEBUG_DRAWING
+void UAnimNextWorldSubsystem::ShowDebugDrawingHandle(UE::UAF::FModuleHandle InHandle, bool bInShowDebugDrawing)
+{
+	using namespace UE::UAF;
+	check(IsInGameThread());
+	if (IsValidHandle(InHandle))
+	{
+		UE_MT_SCOPED_READ_ACCESS(PendingActionsAccessDetector);
+		PendingActions.Emplace(InHandle, FModulePendingAction::EType::EnableDebugDrawing, bInShowDebugDrawing);
+	}
+}
+#endif
+
+void UAnimNextWorldSubsystem::QueueTaskHandle(UE::UAF::FModuleHandle InOutHandle, FName InModuleEventName, TUniqueFunction<void(const UE::UAF::FModuleTaskContext&)>&& InTaskFunction, UE::UAF::ETaskRunLocation InLocation)
+{
+	using namespace UE::UAF;
+	check(IsInGameThread());
+	if (IsValidHandle(InOutHandle))
+	{
+		UE_MT_SCOPED_READ_ACCESS(InstancesAccessDetector);
+		FAnimNextModuleInstance& Instance = Instances.Get(InOutHandle);
+		Instance.QueueTask(InModuleEventName, MoveTemp(InTaskFunction), InLocation);
+	}
+}
+
+#if WITH_EDITOR
+
+void UAnimNextWorldSubsystem::OnCompileJobFinished(UAnimNextRigVMAsset* InAsset)
+{
+	// Cant do this while we are running in a world tick
+	check(!GetWorld()->bInTick);
+	UE_MT_SCOPED_READ_ACCESS(InstancesAccessDetector);
+
+	for (FAnimNextModuleInstance& Instance : Instances)
+	{
+		if (Instance.GetModule() == InAsset)
+		{
+			Instance.OnCompileJobFinished();
+		}
+	}
+}
+
+#endif
+
+void UAnimNextWorldSubsystem::QueueInputTraitEventHandle(UE::UAF::FModuleHandle InHandle, FAnimNextTraitEventPtr Event)
+{
+	using namespace UE::UAF;
+
+	QueueTaskHandle(InHandle, NAME_None, [Event = MoveTemp(Event)](const FModuleTaskContext& InContext)
+	{
+		InContext.QueueInputTraitEvent(Event);
+	},
+	ETaskRunLocation::Before);
+}
+
+const FTickFunction* UAnimNextWorldSubsystem::FindTickFunctionHandle(UE::UAF::FModuleHandle InHandle, FName InEventName) const
+{
+	using namespace UE::UAF;
+	check(IsInGameThread());
+	if (IsValidHandle(InHandle))
+	{
+		UE_MT_SCOPED_READ_ACCESS(InstancesAccessDetector);
+
+		const FAnimNextModuleInstance& Instance = Instances.Get(InHandle);
+		const FModuleEventTickFunction* TickFunction = Instance.FindTickFunctionByName(InEventName);
+		if (TickFunction == nullptr)
+		{
+			UE_LOGFMT(LogAnimation, Warning, "FindTickFunctionHandle: Could not find event '{EventName}' in module '{ModuleName}'", InEventName, Instance.GetAssetName());
+			return nullptr;
+		}
+
+		if (TickFunction->bUserEvent)
+		{
+			return TickFunction;
+		}
+
+		UE_LOGFMT(LogAnimation, Warning, "FindTickFunctionHandle: Event '{EventName}' in module '{ModuleName}' is not a bUserEvent, therefore cannot be exposed", InEventName, Instance.GetAssetName());
+	}
+	return nullptr;
+}
+
+void UAnimNextWorldSubsystem::AddDependencyHandle(UE::UAF::FModuleHandle InHandle, UObject* InObject, FTickFunction& InTickFunction, FName InEventName, EDependency InDependency)
+{
+	using namespace UE::UAF;
+	check(IsInGameThread());
+	if (IsValidHandle(InHandle))
+	{
+		UE_MT_SCOPED_READ_ACCESS(InstancesAccessDetector);
+
+		FAnimNextModuleInstance& Instance = Instances.Get(InHandle);
+		FModuleEventTickFunction* TickFunction = Instance.FindTickFunctionByName(InEventName);
+		if (TickFunction == nullptr)
+		{
+			UE_LOGFMT(LogAnimation, Warning, "AddDependencyHandle: Could not find event '{EventName}' in module '{ModuleName}'", InEventName, Instance.GetAssetName());
+			return;
+		}
+
+		if (InDependency == EDependency::Prerequisite)
+		{
+			TickFunction->AddPrerequisite(InObject, InTickFunction);
+		}
+		else
+		{
+			TickFunction->AddSubsequent(InObject, InTickFunction);
+		}
+	}
+}
+
+void UAnimNextWorldSubsystem::RemoveDependencyHandle(UE::UAF::FModuleHandle InHandle, UObject* InObject, FTickFunction& InTickFunction, FName InEventName, EDependency InDependency)
+{
+	using namespace UE::UAF;
+	check(IsInGameThread());
+	if (IsValidHandle(InHandle))
+	{
+		UE_MT_SCOPED_READ_ACCESS(InstancesAccessDetector);
+
+		FAnimNextModuleInstance& Instance = Instances.Get(InHandle);
+		FModuleEventTickFunction* TickFunction = Instance.FindTickFunctionByName(InEventName);
+		if (TickFunction == nullptr)
+		{
+			UE_LOGFMT(LogAnimation, Warning, "RemoveDependencyHandle: Could not find event '{EventName}' in module '{ModuleName}'", InEventName, Instance.GetAssetName());
+			return;
+		}
+
+		if (InDependency == EDependency::Prerequisite)
+		{
+			TickFunction->RemovePrerequisite(InObject, InTickFunction);
+		}
+		else
+		{
+			TickFunction->RemoveSubsequent(InObject, InTickFunction);
+		}
+	}
+}
+
+void UAnimNextWorldSubsystem::AddModuleEventDependencyHandle(UE::UAF::FModuleHandle InHandle, FName InEventName, UE::UAF::FModuleHandle OtherHandle, FName OtherEventName, EDependency InDependency)
+{
+	using namespace UE::UAF;
+	check(IsInGameThread());
+	check(InHandle != OtherHandle);
+	if (IsValidHandle(OtherHandle))
+	{
+		UE_MT_SCOPED_READ_ACCESS(InstancesAccessDetector);
+
+		FAnimNextModuleInstance& Instance = Instances.Get(OtherHandle);
+		FModuleEventTickFunction* TickFunction = Instance.FindTickFunctionByName(OtherEventName);
+		if (TickFunction == nullptr)
+		{
+			UE_LOGFMT(LogAnimation, Warning, "AddDependencyHandle: Could not find event '{EventName}' in module '{ModuleName}'", OtherEventName, Instance.GetAssetName());
+			return;
+		}
+
+		AddDependencyHandle(InHandle, Instance.GetObject(), *TickFunction, InEventName, InDependency);
+	}
+}
+
+void UAnimNextWorldSubsystem::RemoveModuleEventDependencyHandle(UE::UAF::FModuleHandle InHandle, FName InEventName, UE::UAF::FModuleHandle OtherHandle, FName OtherEventName, EDependency InDependency)
+{
+	using namespace UE::UAF;
+	check(IsInGameThread());
+	check(InHandle != OtherHandle);
+	if (IsValidHandle(InHandle))
+	{
+		UE_MT_SCOPED_READ_ACCESS(InstancesAccessDetector);
+
+		FAnimNextModuleInstance& Instance = Instances.Get(InHandle);
+		FModuleEventTickFunction* TickFunction = Instance.FindTickFunctionByName(InEventName);
+		if (TickFunction == nullptr)
+		{
+			UE_LOGFMT(LogAnimation, Warning, "RemoveDependencyHandle: Could not find event '{EventName}' in module '{ModuleName}'", OtherEventName, Instance.GetAssetName());
+			return;
+		}
+
+		RemoveDependencyHandle(InHandle, Instance.GetObject(), *TickFunction, InEventName, InDependency);
+	}
+}
+
+EPropertyBagResult UAnimNextWorldSubsystem::SetVariableHandle(UE::UAF::FModuleHandle InHandle, const FAnimNextVariableReference& InVariable, const FAnimNextParamType& InType, TConstArrayView<uint8> InData)
+{
+	using namespace UE::UAF;
+	if (IsValidHandle(InHandle))
+	{
+		UE_MT_SCOPED_READ_ACCESS(InstancesAccessDetector);
+
+		FAnimNextModuleInstance& Instance = Instances.Get(InHandle);
+		return Instance.SetProxyVariable(InVariable, InType, InData);
+	}
+	return EPropertyBagResult::PropertyNotFound;
+}
+
+EPropertyBagResult UAnimNextWorldSubsystem::WriteVariableHandle(UE::UAF::FModuleHandle InHandle, const FAnimNextVariableReference& InVariable, const FAnimNextParamType& InType, TFunctionRef<void(TArrayView<uint8>)> InFunction)
+{
+	using namespace UE::UAF;
+	if (IsValidHandle(InHandle))
+	{
+		UE_MT_SCOPED_READ_ACCESS(InstancesAccessDetector);
+
+		FAnimNextModuleInstance& Instance = Instances.Get(InHandle);
+		return Instance.WriteProxyVariable(InVariable, InType, InFunction);
+	}
+	return EPropertyBagResult::PropertyNotFound;
+}
diff --git a/Engine/Plugins/Experimental/UAF/UAF/Source/UAF/Private/Module/AnimNextActorComponentReferenceComponent.cpp b/Engine/Plugins/Experimental/UAF/UAF/Source/UAF/Private/Module/AnimNextActorComponentReferenceComponent.cpp
--- a/Engine/Plugins/Experimental/UAF/UAF/Source/UAF/Private/Module/AnimNextActorComponentReferenceComponent.cpp
+++ b/Engine/Plugins/Experimental/UAF/UAF/Source/UAF/Private/Module/AnimNextActorComponentReferenceComponent.cpp
@@ -4,7 +4,6 @@
 
 #include "Component/AnimNextComponent.h"
 #include "Module/AnimNextModuleInstance.h"
-#include "Module/UAFWeakSystemReference.h"
 
 #include UE_INLINE_GENERATED_CPP_BY_NAME(AnimNextActorComponentReferenceComponent)
 
@@ -21,7 +20,7 @@ void FAnimNextActorComponentReferenceComponent::OnInitializeHelper(UScriptStruct
 	UAnimNextComponent* AnimNextComponent = Cast<UAnimNextComponent>(ModuleInstance->GetObject());
 
 	FAnimNextModuleInstance::RunTaskOnGameThread(
-		[Reference = ModuleInstance->GetReference(), WeakAnimNextComponent = TWeakObjectPtr<UAnimNextComponent>(AnimNextComponent), ComponentType = ComponentType, InScriptStruct]()
+		[WeakAnimNextComponent = TWeakObjectPtr<UAnimNextComponent>(AnimNextComponent), ComponentType = ComponentType, InScriptStruct]()
 		{
 			UAnimNextComponent* PinnedAnimNextComponent = WeakAnimNextComponent.Get();
 			if(PinnedAnimNextComponent == nullptr)
@@ -36,7 +35,7 @@ void FAnimNextActorComponentReferenceComponent::OnInitializeHelper(UScriptStruct
 			}
 
 			UActorComponent* FoundComponent = Owner->FindComponentByClass(ComponentType);
-			Reference.QueueTask(NAME_None, [FoundComponent, InScriptStruct](const UE::UAF::FModuleTaskContext& InContext)
+			PinnedAnimNextComponent->QueueTask(NAME_None, [FoundComponent, InScriptStruct](const UE::UAF::FModuleTaskContext& InContext)
 			{
 				InContext.TryAccessComponent(InScriptStruct, [FoundComponent](FUAFModuleInstanceComponent& InComponent)
 				{
diff --git a/Engine/Plugins/Experimental/UAF/UAF/Source/UAF/Private/Module/AnimNextModule.cpp b/Engine/Plugins/Experimental/UAF/UAF/Source/UAF/Private/Module/AnimNextModule.cpp
--- a/Engine/Plugins/Experimental/UAF/UAF/Source/UAF/Private/Module/AnimNextModule.cpp
+++ b/Engine/Plugins/Experimental/UAF/UAF/Source/UAF/Private/Module/AnimNextModule.cpp
@@ -1,7 +1,6 @@
 // Copyright Epic Games, Inc. All Rights Reserved.
 
 #include "Module/AnimNextModule.h"
-#include "AnimNextPool.h"
 #include "RigVMCore/RigVMExecuteContext.h"
 #include "UObject/AssetRegistryTagsContext.h"
 #include "UObject/ObjectSaveContext.h"
diff --git a/Engine/Plugins/Experimental/UAF/UAF/Source/UAF/Private/Module/AnimNextModuleInstance.cpp b/Engine/Plugins/Experimental/UAF/UAF/Source/UAF/Private/Module/AnimNextModuleInstance.cpp
--- a/Engine/Plugins/Experimental/UAF/UAF/Source/UAF/Private/Module/AnimNextModuleInstance.cpp
+++ b/Engine/Plugins/Experimental/UAF/UAF/Source/UAF/Private/Module/AnimNextModuleInstance.cpp
@@ -18,8 +18,6 @@
 #include "Misc/EnumerateRange.h"
 #include "Module/RigUnit_AnimNextModuleEvents.h"
 #include "Module/ModuleTickFunction.h"
-#include "Module/SystemReference.h"
-#include "Module/UAFWeakSystemReference.h"
 #include "UObject/UObjectIterator.h"
 #include "RewindDebugger/UAFTrace.h"
 #include "Script/UAFScriptComponent.h"
@@ -72,9 +70,10 @@ FAnimNextModuleInstance::~FAnimNextModuleInstance()
 
 	Object = nullptr;
 	Asset = nullptr;
+	Handle.Reset();
 }
 
-void FAnimNextModuleInstance::Initialize(const FUAFSystemFactoryParams& InFactoryParams)
+void FAnimNextModuleInstance::Initialize(const FUAFSystemFactoryParams& InFactoryParams, UE::UAF::TPool<FAnimNextModuleInstance>* InPool)
 {
 	SCOPE_CYCLE_COUNTER(STAT_AnimNext_InitializeInstance);
 	
@@ -82,11 +81,13 @@ void FAnimNextModuleInstance::Initialize(const FUAFSystemFactoryParams& InFactor
 
 	check(IsInGameThread());
 
-	const UAnimNextModule* System = GetSystem();
-	if (System)
+	Pool = InPool;
+
+	const UAnimNextModule* Module = GetModule();
+	if (Module)
 	{
 #if WITH_EDITOR
-		if (System->CompilationState == EAnimNextRigVMAssetState::CompiledWithErrors)
+		if (Module->CompilationState == EAnimNextRigVMAssetState::CompiledWithErrors)
 		{
 			return;
 		}
@@ -157,7 +158,7 @@ void FAnimNextModuleInstance::Initialize(const FUAFSystemFactoryParams& InFactor
 				PrevTickFunction = &TickFunction;
 
 				// Set up dependencies, if any
-				for (const TInstancedStruct<FRigVMTrait_ModuleEventDependency>& DependencyInstance : System->Dependencies)
+				for (const TInstancedStruct<FRigVMTrait_ModuleEventDependency>& DependencyInstance : Module->Dependencies)
 				{
 					const FRigVMTrait_ModuleEventDependency* Dependency = DependencyInstance.GetPtr<FRigVMTrait_ModuleEventDependency>();
 					if (Dependency != nullptr && Dependency->EventName == ModuleEvent.Event.GetEventName())
@@ -260,7 +261,7 @@ void FAnimNextModuleInstance::Initialize(const FUAFSystemFactoryParams& InFactor
 		// Just pause now if we arent needing an initial update
 		if(InitMethod == EAnimNextModuleInitMethod::None)
 		{
-			SetEnabled(false);
+			Enable(false);
 		}
 #if WITH_EDITOR
 		else if(World)
@@ -335,18 +336,12 @@ bool FAnimNextModuleInstance::IsEnabled() const
 	return RunState == ERunState::Running;
 }
 
-void FAnimNextModuleInstance::SetEnabled(bool bInEnabled)
+void FAnimNextModuleInstance::Enable(bool bInEnabled)
 {
 	using namespace UE::UAF;
 
 	check(IsInGameThread());
 
-	// Ensure that we synchronize activation state with actor component if we are hosted there
-	if (UActorComponent* Component = Cast<UActorComponent>(Object))
-	{
-		Component->SetActiveFlag(bInEnabled);
-	}
-
 	if(RunState == ERunState::Paused || RunState == ERunState::Running)
 	{
 		for (FModuleEventTickFunction& TickFunction : TickFunctions)
@@ -436,7 +431,7 @@ void FAnimNextModuleInstance::Trace()
 }
 #endif
 
-const UAnimNextModule* FAnimNextModuleInstance::GetSystem() const
+const UAnimNextModule* FAnimNextModuleInstance::GetModule() const
 {
 	return CastChecked<UAnimNextModule>(Asset, ECastCheckedType::NullAllowed);
 }
@@ -583,13 +578,6 @@ TArrayView<UE::UAF::FModuleEventTickFunction> FAnimNextModuleInstance::GetTickFu
 	return TickFunctions;
 }
 
-FUAFWeakSystemReference FAnimNextModuleInstance::GetReference() const
-{
-	FUAFWeakSystemReference Handle;
-	Handle.WeakPtr = SharedThis(const_cast<FAnimNextModuleInstance*>(this));
-	return Handle;
-}
-
 UE::UAF::FModuleEventTickFunction* FAnimNextModuleInstance::FindFirstUserTickFunction()
 {
 	UE::UAF::FModuleEventTickFunction* FoundTickFunction = TickFunctions.FindByPredicate([](const UE::UAF::FModuleEventTickFunction& InTickFunction)
@@ -649,3 +637,19 @@ void FAnimNextModuleInstance::QueueTask(FName InEventName, TUniqueFunction<void(
 		UE_LOGFMT(LogAnimation, Warning, "QueueTask: Could not find event '{EventName}' in module '{ModuleName}'", InEventName, GetAssetName());
 	}
 }
+
+void FAnimNextModuleInstance::QueueTaskOnOtherModule(const UE::UAF::FModuleHandle InOtherModuleHandle, const FName InEventName, TUniqueFunction<void(const UE::UAF::FModuleTaskContext&)>&& InTaskFunction, UE::UAF::ETaskRunLocation InLocation)
+{
+	if (Pool == nullptr)
+	{
+		return;
+	}
+
+	FAnimNextModuleInstance* OtherModuleInstance = Pool->TryGet(InOtherModuleHandle);
+	if (OtherModuleInstance == nullptr)
+	{
+		return;
+	}
+
+	OtherModuleInstance->QueueTask(InEventName, MoveTemp(InTaskFunction), InLocation);
+}
diff --git a/Engine/Plugins/Experimental/UAF/UAF/Source/UAF/Private/Module/ModuleInstanceUtils.cpp b/Engine/Plugins/Experimental/UAF/UAF/Source/UAF/Private/Module/ModuleInstanceUtils.cpp
new file mode 100644
--- /dev/null
+++ b/Engine/Plugins/Experimental/UAF/UAF/Source/UAF/Private/Module/ModuleInstanceUtils.cpp
@@ -0,0 +1,47 @@
+// Copyright Epic Games, Inc. All Rights Reserved.
+
+#include "Module/ModuleInstanceUtils.h"
+
+#include "Component/AnimNextComponentWorldSubsystem.h"
+#include "Engine/Engine.h"
+#include "StructUtils/PropertyBag.h"
+#include "Variables/AnimNextVariableReference.h"
+
+namespace UE::UAF
+{
+
+EPropertyBagResult SetModuleVariable(FModuleContext&& InModuleContext, const FAnimNextVariableReference& InVariable
+	, const FAnimNextParamType& InType, TConstArrayView<uint8> InData)
+{
+	if (const UWorld* World = GEngine->GetWorldFromContextObject(InModuleContext.ContextObject, EGetWorldErrorMode::ReturnNull))
+	{
+		if (UAnimNextComponentWorldSubsystem* Subsystem = World->GetSubsystem<UAnimNextComponentWorldSubsystem>())
+		{
+			return Subsystem->SetVariableHandle(InModuleContext.Handle, InVariable, InType, InData);
+		}
+	}
+
+	UE_LOG(LogAnimation, Error, TEXT("Unable to set variable '%s'. Cause: unable to access world system from the provided context object '%s'")
+		, *InVariable.GetName().ToString()
+		, *GetFullNameSafe(InModuleContext.ContextObject));
+	return EPropertyBagResult::PropertyNotFound;
+}
+
+EPropertyBagResult WriteModuleVariable(FModuleContext&& InModuleContext, const FAnimNextVariableReference& InVariable
+	, const FAnimNextParamType& InType, TFunctionRef<void(TArrayView<uint8>)> InFunction)
+{
+	if (const UWorld* World = GEngine->GetWorldFromContextObject(InModuleContext.ContextObject, EGetWorldErrorMode::ReturnNull))
+	{
+		if (UAnimNextComponentWorldSubsystem* Subsystem = World->GetSubsystem<UAnimNextComponentWorldSubsystem>())
+		{
+			return Subsystem->WriteVariableHandle(InModuleContext.Handle, InVariable, InType, InFunction);
+		}
+	}
+
+	UE_LOG(LogAnimation, Error, TEXT("Unable to set variable '%s'. Cause: unable to access world system from the provided context object '%s'")
+		, *InVariable.GetName().ToString()
+		, *GetFullNameSafe(InModuleContext.ContextObject));
+	return EPropertyBagResult::PropertyNotFound;
+}
+
+} // UE::UAF
\ No newline at end of file
diff --git a/Engine/Plugins/Experimental/UAF/UAF/Source/UAF/Private/Module/ModuleTickFunction.cpp b/Engine/Plugins/Experimental/UAF/UAF/Source/UAF/Private/Module/ModuleTickFunction.cpp
--- a/Engine/Plugins/Experimental/UAF/UAF/Source/UAF/Private/Module/ModuleTickFunction.cpp
+++ b/Engine/Plugins/Experimental/UAF/UAF/Source/UAF/Private/Module/ModuleTickFunction.cpp
@@ -6,13 +6,12 @@
 #include "AnimNextStats.h"
 #include "RigVMRuntimeDataRegistry.h"
 #include "Algo/TopologicalSort.h"
+#include "Component/AnimNextWorldSubsystem.h"
 #include "Module/AnimNextModule.h"
 #include "Module/AnimNextModuleInstance.h"
 #include "Module/ModuleTaskContext.h"
 #include "Module/RigUnit_AnimNextModuleEvents.h"
 #include "Module/ModuleEvents.h"
-#include "Module/SystemReference.h"
-#include "Module/UAFWeakSystemReference.h"
 #include "TraitCore/TraitEventList.h"
 
 namespace UE::UAF
@@ -93,8 +92,16 @@ void FModuleEventTickFunction::EndTick(float DeltaTime)
 #endif
 		)
 		{
-			// Disable ourselves (deferred)
-			FSystemReference(*ModuleInstance).SetEnabled(false);
+			// Queue task to disable ourselves
+			RunTaskOnGameThread([ModuleHandle = ModuleInstance->Handle, WeakObject = TWeakObjectPtr<UObject>(ModuleInstance->Object)]()
+			{
+				check(IsInGameThread());
+				if (UObject* Object = WeakObject.Get())
+				{
+					UAnimNextWorldSubsystem* WorldSubsystem = UAnimNextWorldSubsystem::Get(Object);
+					WorldSubsystem->EnableHandle(ModuleHandle, false);
+				}
+			});
 		}
 	}
 
diff --git a/Engine/Plugins/Experimental/UAF/UAF/Source/UAF/Private/Module/RigUnit_AnimNextModuleEvents.cpp b/Engine/Plugins/Experimental/UAF/UAF/Source/UAF/Private/Module/RigUnit_AnimNextModuleEvents.cpp
--- a/Engine/Plugins/Experimental/UAF/UAF/Source/UAF/Private/Module/RigUnit_AnimNextModuleEvents.cpp
+++ b/Engine/Plugins/Experimental/UAF/UAF/Source/UAF/Private/Module/RigUnit_AnimNextModuleEvents.cpp
@@ -2,6 +2,7 @@
 
 #include "Module/RigUnit_AnimNextModuleEvents.h"
 
+#include "Component/AnimNextWorldSubsystem.h"
 #include "Engine/World.h"
 #include "Module/AnimNextModuleInstance.h"
 #include "Script/ScriptEventTickFunctionBindings.h"
diff --git a/Engine/Plugins/Experimental/UAF/UAF/Source/UAF/Private/Module/SystemReference.cpp b/Engine/Plugins/Experimental/UAF/UAF/Source/UAF/Private/Module/SystemReference.cpp
deleted file mode 100644
diff --git a/Engine/Plugins/Experimental/UAF/UAF/Source/UAF/Private/Module/UAFWeakSystemReference.cpp b/Engine/Plugins/Experimental/UAF/UAF/Source/UAF/Private/Module/UAFWeakSystemReference.cpp
deleted file mode 100644
diff --git a/Engine/Plugins/Experimental/UAF/UAF/Source/UAF/Public/Component/AnimNextComponent.h b/Engine/Plugins/Experimental/UAF/UAF/Source/UAF/Public/Component/AnimNextComponent.h
--- a/Engine/Plugins/Experimental/UAF/UAF/Source/UAF/Public/Component/AnimNextComponent.h
+++ b/Engine/Plugins/Experimental/UAF/UAF/Source/UAF/Public/Component/AnimNextComponent.h
@@ -3,7 +3,7 @@
 #pragma once
 
 #include "Components/ActorComponent.h"
-#include "Module/SystemReference.h"
+#include "Module/ModuleHandle.h"
 #include "Param/ParamType.h"
 #include "TraitCore/TraitEvent.h"
 #include "Module/AnimNextModuleInitMethod.h"
@@ -18,11 +18,11 @@ namespace UE::Workspace
 }
 
 struct FAnimNextComponentInstanceData;
+class UAnimNextComponentWorldSubsystem;
 class UAnimNextModule;
 struct AnimNextVariableReference;
 // TODO: if a public API is implemented, the forward declaration of this struct would no longer be needed.
 struct FAnimMixerUtils;
-struct FUAFWeakSystemReference;
 
 namespace UE::UAF::UncookedOnly
 {
@@ -47,8 +47,6 @@ class UAnimNextComponent : public UActorComponent
 	virtual void OnUnregister() override;
 	virtual void BeginPlay() override;
 	virtual void EndPlay(const EEndPlayReason::Type EndPlayReason) override;
-	virtual void Activate(bool bReset = false) override;
-	virtual void Deactivate() override;
 
 	// UObject interface
 	virtual void Serialize(FArchive& Ar) override;
@@ -92,6 +90,14 @@ public:
 		}) == EPropertyBagResult::Success;
 	}
 
+	// Whether this component is currently updating
+	UFUNCTION(BlueprintCallable, Category = "UAF")
+	UAF_API bool IsEnabled() const;
+
+	// Enable or disable this component's update
+	UFUNCTION(BlueprintCallable, Category = "UAF")
+	UAF_API void SetEnabled(bool bEnabled);
+
 	// Enable or disable debug drawing. Note only works in builds with UE_ENABLE_DEBUG_DRAWING enabled
 	UFUNCTION(BlueprintCallable, Category = "UAF")
 	UAF_API void ShowDebugDrawing(bool bShowDebugDrawing);
@@ -103,9 +109,8 @@ public:
 	// Input events will be processed in the next graph update after they are queued
 	UAF_API void QueueInputTraitEvent(FAnimNextTraitEventPtr Event);
 
-	// Get a reference to the registered system
-	UFUNCTION(BlueprintCallable, Category = "UAF")
-	UAF_API FUAFWeakSystemReference GetSystemReference() const;
+	// Get the handle to the registered module
+	UE::UAF::FModuleHandle GetModuleHandle() const { return ModuleHandle; }
 
 	// Find the tick function for the specified event
 	// @param	InEventName			The event associated to the wanted tick function
@@ -189,13 +194,16 @@ public:
 	UFUNCTION(BlueprintCallable, Category = "UAF")
 	UAF_API void RemoveModuleEventSubsequent(FName InEventName, UAnimNextComponent* OtherAnimNextComponent, FName OtherEventName);
 
+	UFUNCTION(BlueprintCallable, Category = "UAF", meta=(DisplayName = "Get Module Handle"))
+	UAF_API FAnimNextModuleHandle BlueprintGetModuleHandle() const;
+
 	UAF_API TObjectPtr<UObject> GetAsset() const;
 
 private:
 	UAF_API void SetAsset(TObjectPtr<UObject> InAsset);
 
-	UAF_API void RegisterSystem();
-	UAF_API void UnregisterSystem();
+	UAF_API void RegisterWithSubsystem();
+	UAF_API void UnregisterWithSubsystem();
 
 	UAF_API bool IsModuleValid();
 
@@ -211,6 +219,7 @@ private:
 private:
 	friend struct UE::UAF::UncookedOnly::FUtils;
 	friend struct FAssetBlueprintTests;
+	friend class UAnimNextComponentWorldSubsystem;
 	friend class UE::UAF::Editor::FAssetWizard;
 	// TODO ideally there would be a public API for the things FAnimMixerUtils does, so it wouldn't be necessary.
 	friend struct FAnimMixerUtils;
@@ -227,8 +236,11 @@ private:
 	TObjectPtr<UAnimNextModule> Module_DEPRECATED = nullptr;
 #endif
 
-	// Owning reference to the registered system
-	UE::UAF::FSystemReference SystemReference;
+	UPROPERTY(Transient)
+	TObjectPtr<UAnimNextComponentWorldSubsystem> Subsystem = nullptr;
+
+	// Handle to the registered module
+	UE::UAF::FModuleHandle ModuleHandle;
 
 	// How to initialize the system
 	UPROPERTY(EditAnywhere, Category="System")
diff --git a/Engine/Plugins/Experimental/UAF/UAF/Source/UAF/Public/Module/ModuleHandle.h b/Engine/Plugins/Experimental/UAF/UAF/Source/UAF/Public/Module/ModuleHandle.h
new file mode 100644
--- /dev/null
+++ b/Engine/Plugins/Experimental/UAF/UAF/Source/UAF/Public/Module/ModuleHandle.h
@@ -0,0 +1,33 @@
+// Copyright Epic Games, Inc. All Rights Reserved.
+
+#pragma once
+
+#include "AnimNextPoolHandle.h"
+
+#include "ModuleHandle.generated.h"
+
+struct FAnimNextModuleInstance;
+
+namespace UE::UAF
+{
+
+// Opaque handle representing a module instance
+using FModuleHandle = TPoolHandle<FAnimNextModuleInstance>;
+
+}
+
+USTRUCT(BlueprintType)
+struct FAnimNextModuleHandle
+{
+	GENERATED_BODY()
+
+	FAnimNextModuleHandle() = default;
+
+	explicit FAnimNextModuleHandle(const UE::UAF::FModuleHandle InModuleHandle)
+		: ModuleHandle(InModuleHandle)
+	{
+	}
+
+public:
+	UE::UAF::FModuleHandle ModuleHandle;
+};
\ No newline at end of file
diff --git a/Engine/Plugins/Experimental/UAF/UAF/Source/UAF/Public/Module/ModuleInstanceUtils.h b/Engine/Plugins/Experimental/UAF/UAF/Source/UAF/Public/Module/ModuleInstanceUtils.h
new file mode 100644
--- /dev/null
+++ b/Engine/Plugins/Experimental/UAF/UAF/Source/UAF/Public/Module/ModuleInstanceUtils.h
@@ -0,0 +1,48 @@
+// Copyright Epic Games, Inc. All Rights Reserved.
+
+#pragma once
+
+#include "CoreTypes.h"
+#include "ModuleHandle.h"
+
+struct FAnimNextParamType;
+struct FAnimNextVariableReference;
+struct FAnimNextModuleInstance;
+enum class EPropertyBagResult : uint8;
+
+#define UE_API UAF_API
+
+namespace UE::UAF
+{
+
+struct FModuleContext
+{
+	FModuleHandle Handle;
+	TObjectPtr<UObject> ContextObject;
+};
+
+/**
+ * Sets the value of the specified variable in the specified animation module
+ * @param InModuleContext	The context providing the handle of the module whose variable we are setting and some additional information to access the module.
+ * @param InVariable		The variable to set
+ * @param InType			The type of the variable to set
+ * @param InData			The data to set the variable with
+ * @return Whether the variable was set successfully
+ */
+UE_API EPropertyBagResult SetModuleVariable(FModuleContext&& InModuleContext, const FAnimNextVariableReference& InVariable, const FAnimNextParamType& InType, TConstArrayView<uint8> InData);
+
+/**
+ * Accesses the variable of the specified name for writing.
+ * General read access is not provided via this API due to the current double-buffering strategy used to communicate variable writes to worker threads.
+ * This is intended to allow for copy-free writing of larger data structures & arrays, rather than read access.
+ * @param InModuleContext	The context providing the handle of the module whose variable we are setting and some additional information to access the module.
+ * @param InVariable		The variable we want to access
+ * @param InType			The type of the variable we want to access
+ * @param InFunction		Function that will be called to allow modification of the variables
+ * @return Whether the variable was accessed successfully
+ */
+UE_API EPropertyBagResult WriteModuleVariable(FModuleContext&& InModuleContext, const FAnimNextVariableReference& InVariable, const FAnimNextParamType& InType, TFunctionRef<void(TArrayView<uint8>)> InFunction);
+
+} // UE::UAF
+
+#undef UE_API
\ No newline at end of file
diff --git a/Engine/Plugins/Experimental/UAF/UAF/Source/UAF/Public/Module/SystemDependency.h b/Engine/Plugins/Experimental/UAF/UAF/Source/UAF/Public/Module/SystemDependency.h
deleted file mode 100644
diff --git a/Engine/Plugins/Experimental/UAF/UAF/Source/UAF/Public/Module/SystemReference.h b/Engine/Plugins/Experimental/UAF/UAF/Source/UAF/Public/Module/SystemReference.h
deleted file mode 100644
diff --git a/Engine/Plugins/Experimental/UAF/UAF/Source/UAF/Public/Module/UAFWeakSystemReference.h b/Engine/Plugins/Experimental/UAF/UAF/Source/UAF/Public/Module/UAFWeakSystemReference.h
deleted file mode 100644
diff --git a/Engine/Plugins/Experimental/UAF/UAF/Source/UAFEditor/Private/AnimNextEditorModule.cpp b/Engine/Plugins/Experimental/UAF/UAF/Source/UAFEditor/Private/AnimNextEditorModule.cpp
--- a/Engine/Plugins/Experimental/UAF/UAF/Source/UAFEditor/Private/AnimNextEditorModule.cpp
+++ b/Engine/Plugins/Experimental/UAF/UAF/Source/UAFEditor/Private/AnimNextEditorModule.cpp
@@ -614,7 +614,6 @@ void FAnimNextEditorModule::ShutdownModule()
 		PropertyModule.UnregisterCustomPropertyTypeLayout("AnimNextVariableReference");
 		PropertyModule.UnregisterCustomPropertyTypeLayout("AnimNextSoftVariableReference");
 		PropertyModule.UnregisterCustomPropertyTypeLayout("NameProperty", ModuleEventPropertyTypeIdentifier);
-		PropertyModule.UnregisterCustomPropertyTypeLayout("UAFSystemReference");
 		PropertyModule.UnregisterCustomClassLayout("AnimNextVariableEntry");
 		PropertyModule.UnregisterCustomClassLayout("AnimNextRigVMAssetEditorData");
 	}
diff --git a/Engine/Plugins/Experimental/UAF/UAFAnimGraph/Source/UAFAnimGraph/Private/Injection/InjectionCallbackProxy.cpp b/Engine/Plugins/Experimental/UAF/UAFAnimGraph/Source/UAFAnimGraph/Private/Injection/InjectionCallbackProxy.cpp
--- a/Engine/Plugins/Experimental/UAF/UAFAnimGraph/Source/UAFAnimGraph/Private/Injection/InjectionCallbackProxy.cpp
+++ b/Engine/Plugins/Experimental/UAF/UAFAnimGraph/Source/UAFAnimGraph/Private/Injection/InjectionCallbackProxy.cpp
@@ -48,7 +48,7 @@ bool UInjectionCallbackProxy::Inject(
 	LifetimeEvents.OnInterrupted.BindUObject(this, &UInjectionCallbackProxy::OnInjectionInterrupted);
 	LifetimeEvents.OnBlendingOut.BindUObject(this, &UInjectionCallbackProxy::OnInjectionBlendingOut);
 
-	PlayingRequest = UE::UAF::FInjectionUtils::Inject(AnimNextComponent, AnimNextComponent->GetSystemReference(), MoveTemp(RequestArgs), MoveTemp(LifetimeEvents));
+	PlayingRequest = UE::UAF::FInjectionUtils::Inject(AnimNextComponent, AnimNextComponent->GetModuleHandle(), MoveTemp(RequestArgs), MoveTemp(LifetimeEvents));
 
 	bWasInterrupted = false;
 
diff --git a/Engine/Plugins/Experimental/UAF/UAFAnimGraph/Source/UAFAnimGraph/Private/Injection/InjectionRequest.cpp b/Engine/Plugins/Experimental/UAF/UAFAnimGraph/Source/UAFAnimGraph/Private/Injection/InjectionRequest.cpp
--- a/Engine/Plugins/Experimental/UAF/UAFAnimGraph/Source/UAFAnimGraph/Private/Injection/InjectionRequest.cpp
+++ b/Engine/Plugins/Experimental/UAF/UAFAnimGraph/Source/UAFAnimGraph/Private/Injection/InjectionRequest.cpp
@@ -4,20 +4,19 @@
 
 #include "Injection/InjectionEvents.h"
 #include "Component/AnimNextComponent.h"
+#include "Component/AnimNextWorldSubsystem.h"
 #include "Graph/AnimNextGraphInstance.h"
 #include "Logging/StructuredLog.h"
 #include "Engine/World.h"
-#include "Module/AnimNextModuleInstance.h"
-#include "Module/UAFWeakSystemReference.h"
 
 #include UE_INLINE_GENERATED_CPP_BY_NAME(InjectionRequest)
 
 namespace UE::UAF
 {
-	bool FInjectionRequest::Inject(FInjectionRequestArgs&& InRequestArgs, FInjectionLifetimeEvents&& InLifetimeEvents, UObject* InHost, FUAFWeakSystemReference InReference)
+	bool FInjectionRequest::Inject(FInjectionRequestArgs&& InRequestArgs, FInjectionLifetimeEvents&& InLifetimeEvents, UObject* InHost, FModuleHandle InHandle)
 	{
 		check(IsInGameThread());
-		if (InHost == nullptr || !InReference.IsValid())
+		if (InHost == nullptr || !InHandle.IsValid())
 		{
 			return false;	// Nothing to play
 		}
@@ -41,13 +40,16 @@ namespace UE::UAF
 		RequestArgs = MoveTemp(InRequestArgs);
 		LifetimeEvents = MoveTemp(InLifetimeEvents);
 		WeakHost = InHost;
-		Reference = InReference;
+		WeakWorld = World;
+		Handle = InHandle;
 		Status = EInjectionStatus::Pending;
 
 		auto InjectEvent = MakeTraitEvent<FInjection_InjectEvent>();
 		InjectEvent->Request = AsShared();
 
-		Reference.QueueInputTraitEvent(InjectEvent);
+		UAnimNextWorldSubsystem* Subsystem = World->GetSubsystem<UAnimNextWorldSubsystem>();
+		check(Subsystem);
+		Subsystem->QueueInputTraitEventHandle(InHandle, InjectEvent);
 
 		InjectionEvent = InjectEvent;
 
@@ -76,7 +78,12 @@ namespace UE::UAF
 		auto UninjectEvent = MakeTraitEvent<FInjection_UninjectEvent>();
 		UninjectEvent->Request = AsShared();
 
-		Reference.QueueInputTraitEvent(UninjectEvent);
+		if (UWorld* World = WeakWorld.Get())
+		{
+			UAnimNextWorldSubsystem* Subsystem = World->GetSubsystem<UAnimNextWorldSubsystem>();
+			check(Subsystem);
+			Subsystem->QueueInputTraitEventHandle(Handle, UninjectEvent);
+		}
 	}
 
 	bool FInjectionRequest::ValidateArgs(const FInjectionRequestArgs& InRequestArgs)
diff --git a/Engine/Plugins/Experimental/UAF/UAFAnimGraph/Source/UAFAnimGraph/Private/Injection/InjectionUtils.cpp b/Engine/Plugins/Experimental/UAF/UAFAnimGraph/Source/UAFAnimGraph/Private/Injection/InjectionUtils.cpp
--- a/Engine/Plugins/Experimental/UAF/UAFAnimGraph/Source/UAFAnimGraph/Private/Injection/InjectionUtils.cpp
+++ b/Engine/Plugins/Experimental/UAF/UAFAnimGraph/Source/UAFAnimGraph/Private/Injection/InjectionUtils.cpp
@@ -10,15 +10,15 @@ namespace UE::UAF
 {
 	FInjectionRequestPtr FInjectionUtils::Inject(UAnimNextComponent* InComponent, FInjectionRequestArgs&& InArgs, FInjectionLifetimeEvents&& InLifetimeEvents)
 	{
-		return Inject(InComponent, InComponent->GetSystemReference(), MoveTemp(InArgs), MoveTemp(InLifetimeEvents));
+		return Inject(InComponent, InComponent->GetModuleHandle(), MoveTemp(InArgs), MoveTemp(InLifetimeEvents));
 	}
 
-	FInjectionRequestPtr FInjectionUtils::Inject(UObject* InHost, FUAFWeakSystemReference InSystemReference, FInjectionRequestArgs&& InArgs, FInjectionLifetimeEvents&& InLifetimeEvents)
+	FInjectionRequestPtr FInjectionUtils::Inject(UObject* InHost, FModuleHandle InHandle, FInjectionRequestArgs&& InArgs, FInjectionLifetimeEvents&& InLifetimeEvents)
 	{
 		check(IsInGameThread());
 
 		FInjectionRequestPtr Request = MakeInjectionRequest();
-		if(Request->Inject(MoveTemp(InArgs), MoveTemp(InLifetimeEvents), InHost, InSystemReference))
+		if(Request->Inject(MoveTemp(InArgs), MoveTemp(InLifetimeEvents), InHost, InHandle))
 		{
 			return Request;
 		}
@@ -46,7 +46,7 @@ namespace UE::UAF
 	{
 		return PlayAnimHandle(
 			InComponent,
-			InComponent->GetSystemReference(),
+			InComponent->GetModuleHandle(),
 			InSite,
 			InAnimSequence,
 			{
@@ -69,7 +69,7 @@ namespace UE::UAF
 	{
 		return PlayAnimHandle(
 			InComponent,
-			InComponent->GetSystemReference(),
+			InComponent->GetModuleHandle(),
 			InSite,
 			InAnimSequence,
 			MoveTemp(InArgs),
@@ -80,7 +80,7 @@ namespace UE::UAF
 
 	FInjectionRequestPtr FInjectionUtils::PlayAnimHandle(
 		UObject* InHost,
-		FUAFWeakSystemReference InSystemReference,
+		FModuleHandle InModuleHandle,
 		const FInjectionSite& InSite,
 		UAnimSequence* InAnimSequence,
 		FPlayAnimArgs&& InArgs,
@@ -127,7 +127,7 @@ namespace UE::UAF
 			RequestArgs.FactoryParams = MoveTemp(InArgs.FactoryParams);
 		}
 
-		return Inject(InHost, InSystemReference, MoveTemp(RequestArgs), MoveTemp(InLifetimeEvents));
+		return Inject(InHost, InModuleHandle, MoveTemp(RequestArgs), MoveTemp(InLifetimeEvents));
 	}
 
 	FInjectionRequestPtr FInjectionUtils::InjectAsset(
@@ -141,7 +141,7 @@ namespace UE::UAF
 	{
 		return InjectAsset(
 			InComponent,
-			InComponent->GetSystemReference(),
+			InComponent->GetModuleHandle(),
 			InSite,
 			InAsset,
 			MoveTemp(InFactoryParams),
@@ -152,7 +152,7 @@ namespace UE::UAF
 
 	FInjectionRequestPtr FInjectionUtils::InjectAsset(
 		UObject* InHost,
-		FUAFWeakSystemReference InSystemReference,
+		FModuleHandle InModuleHandle,
 		const FInjectionSite& InSite,
 		UObject* InAsset,
 		FAnimNextFactoryParams&& InFactoryParams,
@@ -170,7 +170,7 @@ namespace UE::UAF
 		RequestArgs.BlendOutSettings = InBlendOutSettings;
 		RequestArgs.FactoryParams = MoveTemp(InFactoryParams);
 
-		return Inject(InHost, InSystemReference, MoveTemp(RequestArgs), MoveTemp(InLifetimeEvents));
+		return Inject(InHost, InModuleHandle, MoveTemp(RequestArgs), MoveTemp(InLifetimeEvents));
 	}
 
 	FInjectionRequestPtr FInjectionUtils::InjectEvaluationModifier(
@@ -180,14 +180,14 @@ namespace UE::UAF
 	{
 		return InjectEvaluationModifier(
 			InComponent,
-			InComponent->GetSystemReference(),
+			InComponent->GetModuleHandle(),
 			InEvaluationModifier,
 			InSite);
 	}
 
 	FInjectionRequestPtr FInjectionUtils::InjectEvaluationModifier(
 		UObject* InHost,
-		FUAFWeakSystemReference InSystemReference,
+		FModuleHandle InModuleHandle,
 		const TSharedRef<IEvaluationModifier>& InEvaluationModifier,
 		const FInjectionSite& InSite)
 	{
@@ -198,6 +198,6 @@ namespace UE::UAF
 		RequestArgs.Type = EAnimNextInjectionType::EvaluationModifier;
 		RequestArgs.EvaluationModifier = InEvaluationModifier;
 
-		return Inject(InHost, InSystemReference, MoveTemp(RequestArgs));
+		return Inject(InHost, InModuleHandle, MoveTemp(RequestArgs));
 	}
 }
diff --git a/Engine/Plugins/Experimental/UAF/UAFAnimGraph/Source/UAFAnimGraph/Public/Graph/AnimNextAnimGraph.h b/Engine/Plugins/Experimental/UAF/UAFAnimGraph/Source/UAFAnimGraph/Public/Graph/AnimNextAnimGraph.h
--- a/Engine/Plugins/Experimental/UAF/UAFAnimGraph/Source/UAFAnimGraph/Public/Graph/AnimNextAnimGraph.h
+++ b/Engine/Plugins/Experimental/UAF/UAFAnimGraph/Source/UAFAnimGraph/Public/Graph/AnimNextAnimGraph.h
@@ -3,6 +3,7 @@
 #pragma once
 
 #include "CoreMinimal.h"
+#include "Module/ModuleHandle.h"
 #include "StructUtils/StructView.h"
 #include "TraitCore/TraitEvent.h"
 #include "TraitCore/TraitPtr.h"
diff --git a/Engine/Plugins/Experimental/UAF/UAFAnimGraph/Source/UAFAnimGraph/Public/Injection/InjectionRequest.h b/Engine/Plugins/Experimental/UAF/UAFAnimGraph/Source/UAFAnimGraph/Public/Injection/InjectionRequest.h
--- a/Engine/Plugins/Experimental/UAF/UAFAnimGraph/Source/UAFAnimGraph/Public/Injection/InjectionRequest.h
+++ b/Engine/Plugins/Experimental/UAF/UAFAnimGraph/Source/UAFAnimGraph/Public/Injection/InjectionRequest.h
@@ -12,12 +12,11 @@
 #include "TraitInterfaces/ITimeline.h"
 #include "Injection/InjectionStatus.h"
 #include "Injection/InjectionSite.h"
+#include "Module/ModuleHandle.h"
 #include "Containers/SpscQueue.h"
-#include "Module/UAFWeakSystemReference.h"
 
 #include "InjectionRequest.generated.h"
 
-struct FUAFWeakSystemReference;
 class FReferenceCollector;
 struct FAnimNextModuleInjectionComponent;
 struct FAnimNextGraphInstance;
@@ -261,7 +260,7 @@ namespace UE::UAF
 
 	private:
 		// Sends this request to the specified host and it will attempt to play with the requested arguments
-		bool Inject(FInjectionRequestArgs&& InRequestArgs, FInjectionLifetimeEvents&& InLifetimeEvents, UObject* InHost, FUAFWeakSystemReference InReference);
+		bool Inject(FInjectionRequestArgs&& InRequestArgs, FInjectionLifetimeEvents&& InLifetimeEvents, UObject* InHost, FModuleHandle InHandle);
 
 		// Interrupts this request and request that we transition to the source input on the playing injection site
 		void Uninject();
@@ -294,8 +293,11 @@ namespace UE::UAF
 		// The object we are playing on
 		TWeakObjectPtr<UObject> WeakHost;
 
-		// Reference to the system instance
-		FUAFWeakSystemReference Reference;
+		// The world within which we are playing
+		TWeakObjectPtr<UWorld> WeakWorld;
+
+		// Handle to the module instance
+		FModuleHandle Handle;
 
 		// The injection event if we have injected already
 		TWeakPtr<FAnimNextTraitEvent> InjectionEvent;
diff --git a/Engine/Plugins/Experimental/UAF/UAFAnimGraph/Source/UAFAnimGraph/Public/Injection/InjectionUtils.h b/Engine/Plugins/Experimental/UAF/UAFAnimGraph/Source/UAFAnimGraph/Public/Injection/InjectionUtils.h
--- a/Engine/Plugins/Experimental/UAF/UAFAnimGraph/Source/UAFAnimGraph/Public/Injection/InjectionUtils.h
+++ b/Engine/Plugins/Experimental/UAF/UAFAnimGraph/Source/UAFAnimGraph/Public/Injection/InjectionUtils.h
@@ -59,7 +59,7 @@ namespace UE::UAF
 		// Injects an animation sequence into a running AnimNext module
 		// @param InComponent          The UAnimNextComponent that hosts the module being played on
 		// @param InHost               The object that hosts the module (e.g. UAnimNextComponent) being played on
-		// @param InSystemReference    The system reference being injected into
+		// @param InModuleHandle       The module handle being injected into
 		// @param InSite               The injection site to 'play' the animation sequence on
 		// @param InAnimSequence       The animation sequence to play
 		// @param InArgs               Optional args to use
@@ -79,7 +79,7 @@ namespace UE::UAF
 		// @see PlayAnim
 		UAFANIMGRAPH_API static FInjectionRequestPtr PlayAnimHandle(
 			UObject* InHost,
-			FUAFWeakSystemReference InSystemReference,
+			FModuleHandle InModuleHandle,
 			const FInjectionSite& InSite,
 			UAnimSequence* InAnimSequence,
 			FPlayAnimArgs&& InArgs = FPlayAnimArgs(),
@@ -90,10 +90,12 @@ namespace UE::UAF
 		// Injects an asset into a running AnimNext module
 		// @param InComponent          The UAnimNextComponent that hosts the module being played on
 		// @param InHost               The object that hosts the module (e.g. UAnimNextComponent) being played on
-		// @param InSystemReference    The system reference being injected into
+		// @param InHandle             The module handle being injected into
 		// @param InSite               The injection site to 'play' the asset on
 		// @param InAsset              The asset to 'play'
 		// @param InFactoryParams      The params used to manufacture a graph to run the injected object 
+		// @param InBindingComponent   A component to use to supply additional data interface bindings at the injection site
+		// @param InBindingModuleHandle A module to use to supply additional data interface bindings at the injection site
 		// @param InBlendInSettings    The blend in settings
 		// @param InBlendOutSettings   The blend in settings
 		// @param InLifetimeEvents     Delegates called for various request lifetime events
@@ -110,7 +112,7 @@ namespace UE::UAF
 		// @see InjectAsset
 		UAFANIMGRAPH_API static FInjectionRequestPtr InjectAsset(
 			UObject* InHost,
-			FUAFWeakSystemReference InSystemReference,
+			FModuleHandle InModuleHandle,
 			const FInjectionSite& InSite,
 			UObject* InAsset,
 			FAnimNextFactoryParams&& InFactoryParams = FAnimNextFactoryParams(),
@@ -122,7 +124,7 @@ namespace UE::UAF
 		// Does not require that injection sites are publicly visible.
 		// @param InComponent          The UAnimNextComponent that hosts the module being injected into
 		// @param InHost               The object that hosts the module (e.g. UAnimNextComponent) being injected into
-		// @param InSystemReference    The system reference being injected into
+		// @param InModuleHandle       The module handle being injected into
 		// @param InEvaluationModifier The evaluation modifier to inject
 		// @param InSite               The injection site to use. If this is NAME_None, then the first site will be used.
 		// @return the new injection request for tracking, or an invalid request if injection failed. This can be discarded for fire-and-forget style behavior.
@@ -134,14 +136,14 @@ namespace UE::UAF
 		// @see InjectExternalGraph
 		UAFANIMGRAPH_API static FInjectionRequestPtr InjectEvaluationModifier(
 			UObject* InHost,
-			FUAFWeakSystemReference InSystemReference,
+			FModuleHandle InModuleHandle,
 			const TSharedRef<IEvaluationModifier>& InEvaluationModifier,
 			const FInjectionSite& InSite = FInjectionSite());
 
 		// Inject into a running AnimNext module (raw custom args)
 		// @param InComponent          The UAnimNextComponent that hosts the module being injected into
 		// @param InHost               The object that hosts the module (e.g. UAnimNextComponent) being injected into
-		// @param InSystemReference    The system reference being injected into
+		// @param InHandle             The module handle being injected into
 		// @param InArgs               Arguments to the injection request
 		// @param InLifetimeEvents     Delegates called for various request lifetime events
 		// @return the new injection request for tracking, or an invalid request if injection failed. This can be discarded for fire-and-forget style behavior.
@@ -153,7 +155,7 @@ namespace UE::UAF
 		// @see Inject
 		UAFANIMGRAPH_API static FInjectionRequestPtr Inject(
 			UObject* InHost,
-			FUAFWeakSystemReference InSystemReference,
+			FModuleHandle InHandle,
 			FInjectionRequestArgs&& InArgs,
 			FInjectionLifetimeEvents&& InLifetimeEvents = FInjectionLifetimeEvents());
 
diff --git a/Engine/Plugins/Experimental/UAF/UAFAnimGraph/Source/UAFAnimGraph/Public/TraitInterfaces/IBlendStack.h b/Engine/Plugins/Experimental/UAF/UAFAnimGraph/Source/UAFAnimGraph/Public/TraitInterfaces/IBlendStack.h
--- a/Engine/Plugins/Experimental/UAF/UAFAnimGraph/Source/UAFAnimGraph/Public/TraitInterfaces/IBlendStack.h
+++ b/Engine/Plugins/Experimental/UAF/UAFAnimGraph/Source/UAFAnimGraph/Public/TraitInterfaces/IBlendStack.h
@@ -8,6 +8,7 @@
 #include "AlphaBlend.h"
 #include "Factory/AnimNextFactoryParams.h"
 #include "Graph/AnimNextAnimationGraph.h"
+#include "Module/ModuleHandle.h"
 #include "Animation/BlendProfile.h"
 
 #define UE_API UAFANIMGRAPH_API
diff --git a/Engine/Plugins/Experimental/UAF/UAFPoseSearch/Config/DefaultUAFPoseSearch.ini b/Engine/Plugins/Experimental/UAF/UAFPoseSearch/Config/DefaultUAFPoseSearch.ini
--- a/Engine/Plugins/Experimental/UAF/UAFPoseSearch/Config/DefaultUAFPoseSearch.ini
+++ b/Engine/Plugins/Experimental/UAF/UAFPoseSearch/Config/DefaultUAFPoseSearch.ini
@@ -15,8 +15,4 @@
 +StructRedirects=(OldName="/Script/AnimNextPoseSearch.RigUnit_PoseSearchDatabaseGetTags", NewName="/Script/UAFPoseSearch.RigUnit_PoseSearchDatabaseGetTags")
 +StructRedirects=(OldName="/Script/AnimNextPoseSearch.RigUnit_PoseSearchResultGetSelectedDatabase", NewName="/Script/UAFPoseSearch.RigUnit_PoseSearchResultGetSelectedDatabase")
 
-+ClassRedirects=(OldName="/Script/AnimNextPoseSearch.PoseSearchFeatureChannel_DistanceFromAnimNextVar", NewName="/Script/UAFPoseSearch.PoseSearchFeatureChannel_DistanceFromAnimNextVar")
-
-; System references
-+PropertyRedirects=(OldName="/Script/AnimNextPoseSearch.MotionMatchingTraitSharedData.MotionMatchingResultModuleHandle", NewName="MotionMatchingResultSystemReference")
-+PropertyRedirects=(OldName="/Script/AnimNextPoseSearch.MotionMatchingTraitSharedData.MotionMatchingResultModuleHandleAlt", NewName="MotionMatchingResultSystemReferenceAlt")
++ClassRedirects=(OldName="/Script/AnimNextPoseSearch.PoseSearchFeatureChannel_DistanceFromAnimNextVar", NewName="/Script/UAFPoseSearch.PoseSearchFeatureChannel_DistanceFromAnimNextVar")
\ No newline at end of file
diff --git a/Engine/Plugins/Experimental/UAF/UAFPoseSearch/Source/UAFPoseSearch/Internal/MotionMatchingTraitData.h b/Engine/Plugins/Experimental/UAF/UAFPoseSearch/Source/UAFPoseSearch/Internal/MotionMatchingTraitData.h
--- a/Engine/Plugins/Experimental/UAF/UAFPoseSearch/Source/UAFPoseSearch/Internal/MotionMatchingTraitData.h
+++ b/Engine/Plugins/Experimental/UAF/UAFPoseSearch/Source/UAFPoseSearch/Internal/MotionMatchingTraitData.h
@@ -4,7 +4,7 @@
 
 #include "CoreMinimal.h"
 #include "AlphaBlend.h"
-#include "Module/UAFWeakSystemReference.h"
+#include "Module/ModuleHandle.h"
 #include "PoseSearch/PoseSearchInteractionLibrary.h"
 #include "PoseSearch/PoseSearchLibrary.h"
 #include "TraitCore/TraitSharedData.h"
@@ -87,17 +87,17 @@ struct FMotionMatchingTraitSharedData : public FAnimNextTraitSharedData
 	UPROPERTY(EditAnywhere, Category = Settings)
 	FName AlignmentTransformVariableName;
 
-	// If set along with result variable name, the motion matching result will be written to the variable in this system, otherwise the variable will be written in the module containing the MotionMatchingTrait
+	// If set along with result variable name, the motion matching result will be written to the variable in this module, otherwise the variable will be written in the module containing the MotionMatchingTrait
 	UPROPERTY(EditAnywhere, Category = Settings)
-	FUAFWeakSystemReference MotionMatchingResultSystemReference;
+	FAnimNextModuleHandle MotionMatchingResultModuleHandle;
 
 	// Same as MotionMatchingResultVariableName but for an alternative destination variable
 	UPROPERTY(EditAnywhere, Category = Settings)
 	FName MotionMatchingResultVariableNameAlt = NAME_None;
 
-	// Same as MotionMatchingResultSystemReference but for an alternative destination variable
+	// Same as MotionMatchingResultModuleHandle but for an alternative destination variable
 	UPROPERTY(EditAnywhere, Category = Settings)
-	FUAFWeakSystemReference MotionMatchingResultSystemReferenceAlt;
+	FAnimNextModuleHandle MotionMatchingResultModuleHandleAlt;
 
 	UPROPERTY(EditAnywhere, Category = Interaction, meta = (PinHiddenByDefault))
 	TArray<FPoseSearchInteractionAvailability> Availabilities;
@@ -143,9 +143,9 @@ struct FMotionMatchingTraitSharedData : public FAnimNextTraitSharedData
 		GeneratorMacro(NextUpdateInterruptMode) \
 		GeneratorMacro(MotionMatchingResultVariableName) \
 		GeneratorMacro(AlignmentTransformVariableName) \
-		GeneratorMacro(MotionMatchingResultSystemReference) \
+		GeneratorMacro(MotionMatchingResultModuleHandle) \
 		GeneratorMacro(MotionMatchingResultVariableNameAlt) \
-		GeneratorMacro(MotionMatchingResultSystemReferenceAlt) \
+		GeneratorMacro(MotionMatchingResultModuleHandleAlt) \
 		GeneratorMacro(Availabilities) \
 		GeneratorMacro(bValidateResultAgainstAvailabilities) \
 		GeneratorMacro(bKeepInteractionAlive) \
diff --git a/Engine/Plugins/Experimental/UAF/UAFPoseSearch/Source/UAFPoseSearch/Private/MotionMatchingTrait.cpp b/Engine/Plugins/Experimental/UAF/UAFPoseSearch/Source/UAFPoseSearch/Private/MotionMatchingTrait.cpp
--- a/Engine/Plugins/Experimental/UAF/UAFPoseSearch/Source/UAFPoseSearch/Private/MotionMatchingTrait.cpp
+++ b/Engine/Plugins/Experimental/UAF/UAFPoseSearch/Source/UAFPoseSearch/Private/MotionMatchingTrait.cpp
@@ -441,14 +441,14 @@ PRAGMA_ENABLE_DEPRECATION_WARNINGS
 		}
 #endif // WITH_EDITOR
 		
-		const auto PublishResult = [](FAnimNextModuleInstance* ModuleInstance, const FUAFWeakSystemReference& SystemReference, const FPoseSearchBlueprintResult& SearchResult, FName ResultOutputVariableName)
+		const auto PublishResult = [](FAnimNextModuleInstance* ModuleInstance, const FModuleHandle& ModuleHandle, const FPoseSearchBlueprintResult& SearchResult, FName ResultOutputVariableName)
 			{
 				if (ResultOutputVariableName != NAME_None)
 				{
-					if (SystemReference.IsValid())
+					if (ModuleHandle.IsValid())
 					{
 						// storing SearchResult TObjectPtr(s) as TWeakObjectPtr(s) to avoid referencing deallocated memory 
-						// in case the module containing this trait gets destroyed before the QueueTask lambda execution
+						// in case the module containing this trait gets destroyed before the QueueTaskOnOtherModule lambda execution
 						TWeakObjectPtr<UObject> SelectedAnim = SearchResult.SelectedAnim;
 						float SelectedTime = SearchResult.SelectedTime;
 						bool bIsContinuingPoseSearch = SearchResult.bIsContinuingPoseSearch;
@@ -463,7 +463,8 @@ PRAGMA_ENABLE_DEPRECATION_WARNINGS
 						TArray<FTransform> ActorRootTransforms = SearchResult.ActorRootTransforms;
 						TArray<FTransform> ActorRootBoneTransforms = SearchResult.ActorRootBoneTransforms;
 
-						SystemReference.QueueTask(
+						ModuleInstance->QueueTaskOnOtherModule(
+							ModuleHandle,
 							FRigUnit_AnimNextPrePhysicsEvent::DefaultEventName,
 							[SelectedAnim, SelectedTime, bIsContinuingPoseSearch, WantedPlayRate, bLoop,
 							bIsMirrored, BlendParameters, SelectedDatabase, SearchCost, bIsInteraction,
@@ -508,11 +509,11 @@ PRAGMA_ENABLE_DEPRECATION_WARNINGS
 		FAnimNextModuleInstance* ModuleInstance = Binding.GetTraitPtr().GetNodeInstance()->GetOwner().GetModuleInstance();
 		check(ModuleInstance);
 
-		const FUAFWeakSystemReference ResultSystemReference = SharedData->GetMotionMatchingResultSystemReference(Binding);
-		PublishResult(ModuleInstance, ResultSystemReference, MotionMatchingState.SearchResult, ResultOutputVariableName);
+		const FAnimNextModuleHandle ResultModuleHandle = SharedData->GetMotionMatchingResultModuleHandle(Binding);
+		PublishResult(ModuleInstance, ResultModuleHandle.ModuleHandle, MotionMatchingState.SearchResult, ResultOutputVariableName);
 		
-		const FUAFWeakSystemReference ResultSystemReferenceAlt = SharedData->GetMotionMatchingResultSystemReferenceAlt(Binding);
-		PublishResult(ModuleInstance, ResultSystemReferenceAlt, MotionMatchingState.SearchResult, ResultOutputVariableNameAlt);
+		const FAnimNextModuleHandle ResultModuleHandleAlt = SharedData->GetMotionMatchingResultModuleHandleAlt(Binding);
+		PublishResult(ModuleInstance, ResultModuleHandleAlt.ModuleHandle, MotionMatchingState.SearchResult, ResultOutputVariableNameAlt);
 	}
 
 	void FMotionMatchingTrait::PostEvaluate(FEvaluateTraversalContext& Context, const TTraitBinding<IEvaluate>& Binding) const
diff --git a/Engine/Plugins/Experimental/UAF/UAFPoseSearch/Source/UAFPoseSearch/Private/MotionMatchingTrait.h b/Engine/Plugins/Experimental/UAF/UAFPoseSearch/Source/UAFPoseSearch/Private/MotionMatchingTrait.h
--- a/Engine/Plugins/Experimental/UAF/UAFPoseSearch/Source/UAFPoseSearch/Private/MotionMatchingTrait.h
+++ b/Engine/Plugins/Experimental/UAF/UAFPoseSearch/Source/UAFPoseSearch/Private/MotionMatchingTrait.h
@@ -4,6 +4,7 @@
 
 #include "CoreMinimal.h"
 #include "AlphaBlend.h"
+#include "Module/ModuleHandle.h"
 #include "PoseSearch/PoseSearchInteractionLibrary.h"
 #include "PoseSearch/PoseSearchLibrary.h"
 #include "TraitCore/Trait.h"
