Update UAF asset wizard to not prompt for workspaces

diff --git a/Engine/Plugins/Experimental/UAF/UAF/Source/UAFEditor/Private/Wizard/AssetWizard.cpp b/Engine/Plugins/Experimental/UAF/UAF/Source/UAFEditor/Private/Wizard/AssetWizard.cpp
--- a/Engine/Plugins/Experimental/UAF/UAF/Source/UAFEditor/Private/Wizard/AssetWizard.cpp
+++ b/Engine/Plugins/Experimental/UAF/UAF/Source/UAFEditor/Private/Wizard/AssetWizard.cpp
@@ -22,29 +22,16 @@
 #include "Kismet2/KismetEditorUtilities.h"
 #include "Workspace/AnimNextWorkspaceEditorMode.h"
 #include "Workspace/AnimNextWorkspaceFactory.h"
-#include "IWorkspacePicker.h"
 #include "ObjectTools.h"
 #include "TemplateConfig.h"
+#include "ContentBrowserModule.h"
+#include "IContentBrowserSingleton.h"
 
 #define LOCTEXT_NAMESPACE "UE::UAF::Editor::Wizard"
 
 namespace UE::UAF::Editor
 {
-	TObjectPtr<UWorkspace> FAssetWizard::ShowWorkspacePicker()
-	{
-		const Workspace::IWorkspaceEditorModule& WorkspaceEditorModule = FModuleManager::Get().LoadModuleChecked<Workspace::IWorkspaceEditorModule>("WorkspaceEditor");
-
-		Workspace::IWorkspacePicker::FConfig WorkspacePickerConfig;
-		WorkspacePickerConfig.HintText = Workspace::IWorkspacePicker::EHintText::CreateOrUseExistingWorkspace;
-		WorkspacePickerConfig.WorkspaceFactoryClass = UAnimNextWorkspaceFactory::StaticClass();
-		
-		const TSharedPtr<Workspace::IWorkspacePicker> WorkspacePicker = WorkspaceEditorModule.CreateWorkspacePicker(WorkspacePickerConfig);
-		WorkspacePicker->ShowModal();
-
-		return Cast<UWorkspace>(WorkspacePicker->GetSelectedWorkspace());
-	}
-
-	void FAssetWizard::HandleTemplateSelected(const TObjectPtr<const UUAFTemplateDataAsset> Template, const TObjectPtr<const UUAFTemplateConfig> TemplateConfig, const TObjectPtr<UWorkspace> Workspace)
+	void FAssetWizard::HandleTemplateSelected(const TObjectPtr<const UUAFTemplateDataAsset> Template, const TObjectPtr<const UUAFTemplateConfig> TemplateConfig)
 	{
 		struct FAssetPair
 		{
@@ -109,11 +96,6 @@ namespace UE::UAF::Editor
 			FArchiveReplaceObjectRef<UObject> ReplaceAr(AssetToRepair, ReplaceMap);
 			AssetToRepair->PostEditChange();
 		}
-
-		for (const FAssetPair& AssetPair : Assets)
-		{
-			Workspace->AddAsset(AssetPair.UserAsset);
-		}
 		
 		if (Template->AssetToOpen)
 		{
@@ -149,13 +131,19 @@ namespace UE::UAF::Editor
 
 					TObjectPtr<UObject> NewAsset = AssetTools.CreateAsset(TemplateConfig->BlueprintAssetName, TemplateConfig->OutputPath.Path, UBlueprint::StaticClass(), BlueprintFactory);
 
-					return CastChecked<UBlueprint>(NewAsset);
+					return Cast<UBlueprint>(NewAsset);
 				}
 				else
 				{
 					return TemplateConfig->BlueprintToModify.LoadSynchronous();
 				}
 			}();
+
+			if (!Blueprint)
+			{
+				UE_LOG(LogAnimation, Error, TEXT("UAF asset wizard failed to create/update a blueprint"));
+				return;
+			}
 			
 			TObjectPtr<AActor> Actor = CastChecked<AActor>(Blueprint->GeneratedClass->GetDefaultObject());
 
@@ -223,12 +211,6 @@ namespace UE::UAF::Editor
 
 	void FAssetWizard::ShowTemplatePicker()
 	{
-		const TObjectPtr<UWorkspace> Workspace = ShowWorkspacePicker();
-		if (!Workspace)
-		{
-			return;
-		}
-		
 		TSharedRef<SWindow> TemplateWindow = SNew(SWindow)
 			.Title(LOCTEXT("AssetWizard", "UAF Asset Wizard"))
 			.ClientSize(FVector2D(1200, 800))
@@ -237,9 +219,9 @@ namespace UE::UAF::Editor
 
 		TemplateWindow->SetContent(
 			SNew(STemplatePicker)
-			.OnTemplateSelected_Lambda([TemplateWindow, Workspace](const TObjectPtr<const UUAFTemplateDataAsset> InTemplate, const TObjectPtr<const UUAFTemplateConfig> InTemplateConfig) mutable
+			.OnTemplateSelected_Lambda([TemplateWindow](const TObjectPtr<const UUAFTemplateDataAsset> InTemplate, const TObjectPtr<const UUAFTemplateConfig> InTemplateConfig) mutable
 			{
-				HandleTemplateSelected(InTemplate, InTemplateConfig, Workspace);
+				HandleTemplateSelected(InTemplate, InTemplateConfig);
 				FSlateApplication::Get().RequestDestroyWindow(TemplateWindow);
 			})
 			.OnCancel_Lambda([TemplateWindow]()
@@ -248,8 +230,13 @@ namespace UE::UAF::Editor
 			})
 		);
 
-		GetMutableDefault<UUAFTemplateConfig>()->OutputPath.Path = FAssetData(Workspace).PackagePath.ToString();
-		
+		FContentBrowserModule& ContentBrowserModule = FModuleManager::LoadModuleChecked<FContentBrowserModule>("ContentBrowser");
+
+		TArray<FString> SelectedPaths;
+		const FContentBrowserItemPath CurrentPath = ContentBrowserModule.Get().GetCurrentPath();
+
+		GetMutableDefault<UUAFTemplateConfig>()->OutputPath.Path = CurrentPath.HasInternalPath() ? CurrentPath.GetInternalPathString() : TEXT("/Game");
+
 		FSlateApplication::Get().AddWindow(TemplateWindow);
 	}
 	
diff --git a/Engine/Plugins/Experimental/UAF/UAF/Source/UAFEditor/Private/Wizard/AssetWizard.h b/Engine/Plugins/Experimental/UAF/UAF/Source/UAFEditor/Private/Wizard/AssetWizard.h
--- a/Engine/Plugins/Experimental/UAF/UAF/Source/UAFEditor/Private/Wizard/AssetWizard.h
+++ b/Engine/Plugins/Experimental/UAF/UAF/Source/UAFEditor/Private/Wizard/AssetWizard.h
@@ -16,10 +16,8 @@ namespace UE::UAF::Editor
 		static void Launch();
 		
 	private:
-		static TObjectPtr<UWorkspace> ShowWorkspacePicker();
-
 		static void ShowTemplatePicker();
 		
-		static void HandleTemplateSelected(const TObjectPtr<const UUAFTemplateDataAsset> Template, const TObjectPtr<const UUAFTemplateConfig> TemplateConfig, const TObjectPtr<UWorkspace> Workspace);
+		static void HandleTemplateSelected(const TObjectPtr<const UUAFTemplateDataAsset> Template, const TObjectPtr<const UUAFTemplateConfig> TemplateConfig);
 	};
 }
\ No newline at end of file
