UAF: Minor BugHawk patches for shared/weak system references

diff --git a/Engine/Plugins/Experimental/UAF/UAF/Source/UAF/Private/Module/SystemReference.cpp b/Engine/Plugins/Experimental/UAF/UAF/Source/UAF/Private/Module/SystemReference.cpp
--- a/Engine/Plugins/Experimental/UAF/UAF/Source/UAF/Private/Module/SystemReference.cpp
+++ b/Engine/Plugins/Experimental/UAF/UAF/Source/UAF/Private/Module/SystemReference.cpp
@@ -157,6 +157,8 @@ void FSystemReference::Destroy()
 
 void FSystemReference::Reset()
 {
+	check(IsInGameThread());
+
 	if (Ptr)
 	{
 #if UE_ENABLE_DEBUG_DRAWING
diff --git a/Engine/Plugins/Experimental/UAF/UAF/Source/UAF/Public/Module/UAFWeakSystemReference.h b/Engine/Plugins/Experimental/UAF/UAF/Source/UAF/Public/Module/UAFWeakSystemReference.h
--- a/Engine/Plugins/Experimental/UAF/UAF/Source/UAF/Public/Module/UAFWeakSystemReference.h
+++ b/Engine/Plugins/Experimental/UAF/UAF/Source/UAF/Public/Module/UAFWeakSystemReference.h
@@ -75,16 +75,6 @@ struct FUAFWeakSystemReference
 	// @return see EPropertyBagResult
 	UE_API EPropertyBagResult WriteVariable(const FAnimNextVariableReference& InVariable, const FAnimNextParamType& InType, TFunctionRef<void(TArrayView<uint8>)> InFunction) const;
 
-	// Access a variable's value for read only.
-	// This performs non-recursive locks on the system instance to ensure that concurrent access is respected, so care must be taken calling this from
-	// animation worker tasks as deadlocks can occur.
-	// Type must match strictly, no conversions are performed.
-	// @param	InVariable			The variable to get the value of
-	// @param	InType				The type of the variable
-	// @param	InFunction			Function that will be called if no errors occur
-	// @return see EPropertyBagResult
-	UE_API EPropertyBagResult ReadVariable(const FAnimNextVariableReference& InVariable, const FAnimNextParamType& InType, TFunctionRef<void(TConstArrayView<uint8>)> InFunction) const;
-
 private:
 	friend UE::UAF::FSystemReference;
 	friend FAnimNextModuleInstance;
