UAF: Expose a way for users to add additional components to the editor data and merge them with other required components that will be allocated for instances of the asset

diff --git a/Engine/Plugins/Experimental/UAF/UAF/Source/UAFUncookedOnly/Internal/AnimNextRigVMAssetEditorData.h b/Engine/Plugins/Experimental/UAF/UAF/Source/UAFUncookedOnly/Internal/AnimNextRigVMAssetEditorData.h
--- a/Engine/Plugins/Experimental/UAF/UAF/Source/UAFUncookedOnly/Internal/AnimNextRigVMAssetEditorData.h
+++ b/Engine/Plugins/Experimental/UAF/UAF/Source/UAFUncookedOnly/Internal/AnimNextRigVMAssetEditorData.h
@@ -697,8 +697,13 @@ protected:
 	UPROPERTY()
 	TArray<FEditedDocumentInfo> LastEditedDocuments;
 
+	// Additional components that should be allocated for instances of this asset
+	UPROPERTY(EditAnywhere, Category = Components)
+	TArray<TInstancedStruct<FUAFAssetInstanceComponent>> AdditionalComponents;
+
 	// Components that will be allocated for instances of this asset
-	UPROPERTY(VisibleAnywhere, Category=Components)
+	// This is a combination of required components based on the asset and additional components added by the user 
+	UPROPERTY(VisibleAnywhere, Category=Components, AdvancedDisplay, meta = (DisplayName = "Required Components"))
 	TArray<TInstancedStruct<FUAFAssetInstanceComponent>> Components;
 
 	UPROPERTY(transient, DuplicateTransient)
diff --git a/Engine/Plugins/Experimental/UAF/UAF/Source/UAFUncookedOnly/Private/AnimNextRigVMAssetEditorData.cpp b/Engine/Plugins/Experimental/UAF/UAF/Source/UAFUncookedOnly/Private/AnimNextRigVMAssetEditorData.cpp
--- a/Engine/Plugins/Experimental/UAF/UAF/Source/UAFUncookedOnly/Private/AnimNextRigVMAssetEditorData.cpp
+++ b/Engine/Plugins/Experimental/UAF/UAF/Source/UAFUncookedOnly/Private/AnimNextRigVMAssetEditorData.cpp
@@ -1236,6 +1236,18 @@ void UAnimNextRigVMAssetEditorData::GatherDefaultComponents(UAnimNextRigVMAsset*
 		}
 	}
 
+	// Add additional components added by the user
+	for (const TInstancedStruct<FUAFAssetInstanceComponent>& AdditionalComponent : AdditionalComponents)
+	{
+		if (AdditionalComponent.IsValid() && !Components.ContainsByPredicate([&AdditionalComponent](const TInstancedStruct<FUAFAssetInstanceComponent>& InComponent)
+			{
+				return AdditionalComponent.GetScriptStruct() == InComponent.GetScriptStruct();
+			}))
+		{
+			Components.Add(AdditionalComponent);
+		}
+	}
+
 	// Copy to the asset
 	InAsset->Components = Components;
 }
