Ensure referenced SharedAsset is Postloaded when trying to verify variable references

diff --git a/Engine/Plugins/Experimental/UAF/UAF/Source/UAFUncookedOnly/Private/AnimNextSharedVariableNode.cpp b/Engine/Plugins/Experimental/UAF/UAF/Source/UAFUncookedOnly/Private/AnimNextSharedVariableNode.cpp
--- a/Engine/Plugins/Experimental/UAF/UAF/Source/UAFUncookedOnly/Private/AnimNextSharedVariableNode.cpp
+++ b/Engine/Plugins/Experimental/UAF/UAF/Source/UAFUncookedOnly/Private/AnimNextSharedVariableNode.cpp
@@ -41,6 +41,8 @@ void UAnimNextSharedVariableNode::PostLoad()
 				{
 					if (const UAnimNextSharedVariables* SharedVariablesAsset = Asset.Get())
 					{
+						const_cast<UAnimNextSharedVariables*>(SharedVariablesAsset)->ConditionalPostLoad();
+						
 						const UAnimNextSharedVariables_EditorData* EditorData = UE::UAF::UncookedOnly::FUtils::GetEditorData<UAnimNextSharedVariables_EditorData>(SharedVariablesAsset);
 						const FRigVMTemplateArgumentType ExpectedVariableType(FName(*GetCPPType()), GetCPPTypeObject());
 
