UAF: Added copy constructors and assignment operator for TLODPose to ensure that view is always kept in line with transforms array.

diff --git a/Engine/Plugins/Experimental/UAF/UAF/Source/UAF/Public/LODPose.h b/Engine/Plugins/Experimental/UAF/UAF/Source/UAF/Public/LODPose.h
--- a/Engine/Plugins/Experimental/UAF/UAF/Source/UAF/Public/LODPose.h
+++ b/Engine/Plugins/Experimental/UAF/UAF/Source/UAF/Public/LODPose.h
@@ -335,12 +335,25 @@ struct TLODPose : public FLODPose
 	TTransformArray<AllocatorType> LocalTransforms;
 
 	TLODPose() = default;
+	
+	TLODPose(const TLODPose& Other) : FLODPose(Other), LocalTransforms(Other.LocalTransforms)
+	{
+		LocalTransformsView = LocalTransforms.GetView();
+	}
 
 	TLODPose(const FReferencePose& InRefPose, int32 InLODLevel, bool bSetRefPose = true, bool bAdditive = false)
 	{
 		PrepareForLOD(InRefPose, InLODLevel, bSetRefPose, bAdditive);
 	}
 
+	TLODPose& operator=(const TLODPose& Other)
+	{
+		FLODPose::operator=(Other);
+		LocalTransforms = Other.LocalTransforms;
+		LocalTransformsView = LocalTransforms.GetView();
+		return *this;
+	}
+
 	bool ShouldPrepareForLOD(const FReferencePose& InRefPose, int32 InLODLevel, bool bAdditive = false) const
 	{
 		return
