UAF: Allow GetActorTransform to be used in animation graphs

diff --git a/Engine/Plugins/Experimental/UAF/UAF/Source/UAF/Private/Graph/RigUnit_GetActorTransform.cpp b/Engine/Plugins/Experimental/UAF/UAF/Source/UAF/Private/Graph/RigUnit_GetActorTransform.cpp
--- a/Engine/Plugins/Experimental/UAF/UAF/Source/UAF/Private/Graph/RigUnit_GetActorTransform.cpp
+++ b/Engine/Plugins/Experimental/UAF/UAF/Source/UAF/Private/Graph/RigUnit_GetActorTransform.cpp
@@ -9,8 +9,13 @@
 
 void FAnimNextActorTransformComponent::OnBindToInstance()
 {
-	FAnimNextModuleInstance& ModuleInstance = GetModuleInstance();
-	UAnimNextComponent* AnimNextComponent = Cast<UAnimNextComponent>(ModuleInstance.GetObject());
+	FUAFAssetInstance& AssetInstance = GetAssetInstance();
+	if (AssetInstance.GetRootInstance() == nullptr)
+	{
+		return;
+	}
+
+	UAnimNextComponent* AnimNextComponent = Cast<UAnimNextComponent>(AssetInstance.GetRootInstance()->GetObject());
 
 	FAnimNextModuleInstance::RunTaskOnGameThread(
 		[this, WeakAnimNextComponent = TWeakObjectPtr<UAnimNextComponent>(AnimNextComponent)]()
@@ -48,10 +53,11 @@ void FAnimNextActorTransformComponent::OnBindToInstance()
 
 				Component->QueueTask(NAME_None, [Transform = UpdatedComponentOwner->GetActorTransform()](const UE::UAF::FModuleTaskContext& InContext)
 				{
-					InContext.TryAccessComponent<FAnimNextActorTransformComponent>([&Transform](FAnimNextActorTransformComponent& InComponent)
+					FAnimNextModuleInstance* ModuleInstance = InContext.GetModuleInstance();
+					if (FAnimNextActorTransformComponent* ActorTransformComponent = ModuleInstance->TryGetComponent<FAnimNextActorTransformComponent>())
 					{
-						InComponent.ActorTransform = Transform;
-					});
+						ActorTransformComponent->ActorTransform = Transform;
+					}
 				});
 			});
 		});
@@ -60,9 +66,12 @@ void FAnimNextActorTransformComponent::OnBindToInstance()
 FRigUnit_GetActorTransform_Execute()
 {
 	using namespace UE::UAF;
-
-	const FAnimNextModuleContextData& ModuleContextData = ExecuteContext.GetContextData<FAnimNextModuleContextData>();
-	FAnimNextModuleInstance& ModuleInstance = ModuleContextData.GetModuleInstance();
-	const FAnimNextActorTransformComponent& TransformComponent = ModuleInstance.GetComponent<FAnimNextActorTransformComponent>();
+	const FUAFAssetContextData& AssetContextData = ExecuteContext.GetContextData<FUAFAssetContextData>();
+	FUAFAssetInstance& AssetInstance = AssetContextData.GetInstance();
+	if (AssetInstance.GetRootInstance() == nullptr)
+        {
+		return;
+        }
+	const FAnimNextActorTransformComponent& TransformComponent = AssetInstance.GetRootInstance()->GetComponent<FAnimNextActorTransformComponent>();
 	Transform = TransformComponent.GetActorTransform();
 }
