Scoped compile job changes:

> Renamed FScopedCompileJob to FCompilationScope and changed its behaviour
* Its purpose is to ensure that any modified assets within its (or child scopes) are (if required) compiled in order and relevant systems are notified of their compile start/finish
* It now holds an array for all relevant assets, and a mapping to per-object FCompilerResultsLog (created-on-demand) to differentiate/group messages according to their relevant asset
* Renamed files to UAFCompilationScope.h/cpp

> Introduced Compilation helper struct
* RequestAssetCompilation for single or set of assets, ensures a compilation scope is set-up
~ Ideally this is used by user-code rather than calling ::RecompileVM directly

> FAssetCompilationHandler no longer tracks number of warnings/errors itself, and switched RecompileVM to Compilation::RequestAssetCompilation

diff --git a/Engine/Plugins/Experimental/UAF/UAF/Source/UAF/Public/AnimNextRigVMAsset.h b/Engine/Plugins/Experimental/UAF/UAF/Source/UAF/Public/AnimNextRigVMAsset.h
--- a/Engine/Plugins/Experimental/UAF/UAF/Source/UAF/Public/AnimNextRigVMAsset.h
+++ b/Engine/Plugins/Experimental/UAF/UAF/Source/UAF/Public/AnimNextRigVMAsset.h
@@ -37,7 +37,7 @@ namespace UE::UAF
 namespace UE::UAF::UncookedOnly
 {
 	struct FUtils;
-	class FScopedCompileJob;
+	class FCompilationScope;
 }
 
 #if WITH_EDITORONLY_DATA
@@ -68,7 +68,7 @@ protected:
 	friend UE::UAF::FInjectionInfo;
 	friend UE::UAF::FAnimGraphBuilderContext;
 	friend UE::UAF::FSystemBuilderContext;
-	friend UE::UAF::UncookedOnly::FScopedCompileJob;
+	friend UE::UAF::UncookedOnly::FCompilationScope;
 	friend FUAFAssetInstance;
 
 	UE_API UAnimNextRigVMAsset(const FObjectInitializer& ObjectInitializer);
diff --git a/Engine/Plugins/Experimental/UAF/UAF/Source/UAFEditor/Private/AssetCompilationHandler.cpp b/Engine/Plugins/Experimental/UAF/UAF/Source/UAFEditor/Private/AssetCompilationHandler.cpp
--- a/Engine/Plugins/Experimental/UAF/UAF/Source/UAFEditor/Private/AssetCompilationHandler.cpp
+++ b/Engine/Plugins/Experimental/UAF/UAF/Source/UAFEditor/Private/AssetCompilationHandler.cpp
@@ -3,32 +3,21 @@
 #include "AssetCompilationHandler.h"
 #include "AnimNextRigVMAsset.h"
 #include "AnimNextRigVMAssetEditorData.h"
-#include "AnimNextScopedCompileJob.h"
+#include "UAFCompilationScope.h"
 #include "UncookedOnlyUtils.h"
 
 namespace UE::UAF::Editor
 {
 
-FAssetCompilationHandler::FAssetCompilationHandler(UObject* InAsset)
+FAssetCompilationHandler::FAssetCompilationHandler(const UObject* InAsset)
 {
-	UAnimNextRigVMAsset* AnimNextRigVMAsset = CastChecked<UAnimNextRigVMAsset>(InAsset);
-	UAnimNextRigVMAssetEditorData* EditorData = UE::UAF::UncookedOnly::FUtils::GetEditorData<UAnimNextRigVMAssetEditorData>(AnimNextRigVMAsset);
-
-	NumErrors = EditorData->bErrorsDuringCompilation ? 1 : 0;
-	NumWarnings = EditorData->bWarningsDuringCompilation ? 1 : 0;
+	check(InAsset && InAsset->IsA<UAnimNextRigVMAsset>());
 }
 
 void FAssetCompilationHandler::Compile(TSharedRef<Workspace::IWorkspaceEditor> InWorkspaceEditor, UObject* InAsset)
 {
 	UAnimNextRigVMAsset* AnimNextRigVMAsset = CastChecked<UAnimNextRigVMAsset>(InAsset);
-	UAnimNextRigVMAssetEditorData* EditorData = UE::UAF::UncookedOnly::FUtils::GetEditorData<UAnimNextRigVMAssetEditorData>(AnimNextRigVMAsset);
-
-	UncookedOnly::FScopedCompileJob CompileJob(InAsset);
-	EditorData->RecompileVM();
-
-	NumErrors = CompileJob.GetLog().NumErrors;
-	NumWarnings = CompileJob.GetLog().NumWarnings;
-
+	UncookedOnly::Compilation::RequestAssetCompilation(AnimNextRigVMAsset);
 	OnCompileStatusChanged().ExecuteIfBound();
 }
 
@@ -50,12 +39,12 @@ ECompileStatus FAssetCompilationHandler::GetCompileStatus(TSharedRef<Workspace::
 {
 	const UAnimNextRigVMAsset* AnimNextRigVMAsset = CastChecked<UAnimNextRigVMAsset>(InAsset);
 	UAnimNextRigVMAssetEditorData* EditorData = UE::UAF::UncookedOnly::FUtils::GetEditorData<UAnimNextRigVMAssetEditorData>(AnimNextRigVMAsset);
-	if(EditorData->bErrorsDuringCompilation || NumErrors > 0)
+	if(EditorData->bErrorsDuringCompilation)
 	{
 		return ECompileStatus::Error;
 	}
 
-	if(EditorData->bWarningsDuringCompilation || NumWarnings > 0)
+	if(EditorData->bWarningsDuringCompilation)
 	{
 		return ECompileStatus::Warning;
 	}
diff --git a/Engine/Plugins/Experimental/UAF/UAF/Source/UAFEditor/Private/Common/GraphEditorSchemaActions.cpp b/Engine/Plugins/Experimental/UAF/UAF/Source/UAFEditor/Private/Common/GraphEditorSchemaActions.cpp
--- a/Engine/Plugins/Experimental/UAF/UAF/Source/UAFEditor/Private/Common/GraphEditorSchemaActions.cpp
+++ b/Engine/Plugins/Experimental/UAF/UAF/Source/UAFEditor/Private/Common/GraphEditorSchemaActions.cpp
@@ -16,7 +16,7 @@
 #include "Editor/RigVMEditorStyle.h"
 #include "Entries/AnimNextVariableEntry.h"
 #include "AnimNextRigVMAssetEditorData.h"
-#include "AnimNextScopedCompileJob.h"
+#include "UAFCompilationScope.h"
 #include "ScopedTransaction.h"
 #include "Variables/AnimNextSharedVariables.h"
 
@@ -263,18 +263,17 @@ UEdGraphNode* FAnimNextSchemaAction_Variable::PerformAction(UEdGraph* ParentGrap
 						{
 							if (VariableExport->GetExportAccessSpecifier() != EAnimNextExportAccessSpecifier::Public)
 							{
-								UE::UAF::UncookedOnly::FScopedCompileJob CompilerResults(LOCTEXT("ForceVariablePublicJobName", "Force Variable Public Access"), { SourceObject });							
+								UE::UAF::UncookedOnly::FCompilationScope CompileScope(LOCTEXT("ForceVariablePublicJobName", "Force Variable Public Access"), VariableOwningRigVMAsset);							
 								VariableExport->SetExportAccessSpecifier(EAnimNextExportAccessSpecifier::Public);
 							}
 						}
 						
 						if (!OwnerEditorData->FindEntry(VariableOwningRigVMAsset->GetFName()))
 						{
-							UE::UAF::UncookedOnly::FScopedCompileJob CompilerResults(LOCTEXT("AddSharedAssetJobName", "Add Shared Asset dependency"), { ParentGraph->GetTypedOuter<UAnimNextRigVMAsset>() });
+							UE::UAF::UncookedOnly::FCompilationScope CompileScope(LOCTEXT("AddSharedAssetJobName", "Add Shared Asset dependency"), ParentGraph->GetTypedOuter<UAnimNextRigVMAsset>());
 							
 							UAnimNextSharedVariables* VariableOwningSharedVariables = UE::UAF::UncookedOnly::FUtils::GetAsset<UAnimNextSharedVariables>(VariableOwnerEditorData);
 							OwnerEditorData->AddSharedVariables(VariableOwningSharedVariables, true, false);
-							OwnerEditorData->RecompileVM();
 						}
 					}
 				}
diff --git a/Engine/Plugins/Experimental/UAF/UAF/Source/UAFEditor/Private/Common/Outliner/CommonOutlinerMode.cpp b/Engine/Plugins/Experimental/UAF/UAF/Source/UAFEditor/Private/Common/Outliner/CommonOutlinerMode.cpp
--- a/Engine/Plugins/Experimental/UAF/UAF/Source/UAFEditor/Private/Common/Outliner/CommonOutlinerMode.cpp
+++ b/Engine/Plugins/Experimental/UAF/UAF/Source/UAFEditor/Private/Common/Outliner/CommonOutlinerMode.cpp
@@ -16,7 +16,7 @@
 #include "Variables/Outliner/VariablesOutlinerCommands.h"
 
 #include "AnimNextRigVMAsset.h"
-#include "AnimNextScopedCompileJob.h"
+#include "UAFCompilationScope.h"
 #include "IWorkspaceEditorModule.h"
 #include "OutlinerCategoryItem.h"
 #include "OutlinerCommands.h"
@@ -627,7 +627,7 @@ void FCommonOutlinerMode::SetEntryExportAccessSpecifiersOnSelection(EAnimNextExp
 	const FSceneOutlinerItemSelection& Selection = SceneOutliner->GetSelection();
 	if (Selection.Num() > 0)
 	{
-		TArray<UObject*> ObjectsToCompile;
+		TArray<UAnimNextRigVMAsset*> ObjectsToCompile;
 		TArray<IAnimNextRigVMExportInterface*> ExportsToChange;
 		Selection.ForEachItem<FOutlinerEntryItem>([&ObjectsToCompile](const FOutlinerEntryItem& SelectedItem)
 		{
@@ -640,7 +640,7 @@ void FCommonOutlinerMode::SetEntryExportAccessSpecifiersOnSelection(EAnimNextExp
 			}
 		});
 
-		UncookedOnly::FScopedCompileJob CompilerResults(LOCTEXT("ModifiedAssets_SetAccessSpecifierOnSelection", "Modified Assets Set Access Specifier on Selection"), ObjectsToCompile);
+		UncookedOnly::FCompilationScope CompileScope(LOCTEXT("ModifiedAssets_SetAccessSpecifierOnSelection", "Set Access Specifier on Selection"), ObjectsToCompile);
 
 		FScopedTransaction Transaction(LOCTEXT("SetAccessSpecifierSelectionTransaction", "Set Access Specifier on Selection"));
 
diff --git a/Engine/Plugins/Experimental/UAF/UAF/Source/UAFEditor/Private/Module/AnimNextModuleFactory.cpp b/Engine/Plugins/Experimental/UAF/UAF/Source/UAFEditor/Private/Module/AnimNextModuleFactory.cpp
--- a/Engine/Plugins/Experimental/UAF/UAF/Source/UAFEditor/Private/Module/AnimNextModuleFactory.cpp
+++ b/Engine/Plugins/Experimental/UAF/UAF/Source/UAFEditor/Private/Module/AnimNextModuleFactory.cpp
@@ -1,6 +1,8 @@
 ﻿// Copyright Epic Games, Inc. All Rights Reserved.
 
 #include "Module/AnimNextModuleFactory.h"
+
+#include "UAFCompilationScope.h"
 #include "Module/AnimNextModule.h"
 #include "Module/AnimNextModule_EditorData.h"
 #include "UncookedOnlyUtils.h"
@@ -43,7 +45,8 @@ UObject* UAnimNextModuleFactory::FactoryCreateNew(UClass* Class, UObject* InPare
 	Controller->AddUnitNode(FRigUnit_AnimNextPrePhysicsEvent::StaticStruct(), FRigVMStruct::ExecuteName, FVector2D(0.0f, 200.0f), FString(), false);
 
 	// Compile the initial skeleton
-	EditorData->RecompileVM();
+	UE::UAF::UncookedOnly::Compilation::RequestAssetCompilation(NewModule);
+	
 	check(!EditorData->bErrorsDuringCompilation);
 
 	return NewModule;
diff --git a/Engine/Plugins/Experimental/UAF/UAF/Source/UAFEditor/Private/Variables/AnimNextSharedVariablesFactory.cpp b/Engine/Plugins/Experimental/UAF/UAF/Source/UAFEditor/Private/Variables/AnimNextSharedVariablesFactory.cpp
--- a/Engine/Plugins/Experimental/UAF/UAF/Source/UAFEditor/Private/Variables/AnimNextSharedVariablesFactory.cpp
+++ b/Engine/Plugins/Experimental/UAF/UAF/Source/UAFEditor/Private/Variables/AnimNextSharedVariablesFactory.cpp
@@ -1,6 +1,8 @@
 // Copyright Epic Games, Inc. All Rights Reserved.
 
 #include "Variables/AnimNextSharedVariablesFactory.h"
+
+#include "UAFCompilationScope.h"
 #include "Variables/AnimNextSharedVariables.h"
 #include "Variables/AnimNextSharedVariables_EditorData.h"
 #include "Module/RigUnit_AnimNextModuleEvents.h"
@@ -36,7 +38,8 @@ UObject* UAnimNextSharedVariablesFactory::FactoryCreateNew(UClass* Class, UObjec
 	EditorData->Initialize(/*bRecompileVM*/false);
 
 	// Compile the initial skeleton
-	EditorData->RecompileVM();
+	UE::UAF::UncookedOnly::Compilation::RequestAssetCompilation(NewDataInterface);
+	
 	check(!EditorData->bErrorsDuringCompilation);
 
 	return NewDataInterface;
diff --git a/Engine/Plugins/Experimental/UAF/UAF/Source/UAFEditor/Private/Variables/Outliner/VariablesOutlinerMode.cpp b/Engine/Plugins/Experimental/UAF/UAF/Source/UAFEditor/Private/Variables/Outliner/VariablesOutlinerMode.cpp
--- a/Engine/Plugins/Experimental/UAF/UAF/Source/UAFEditor/Private/Variables/Outliner/VariablesOutlinerMode.cpp
+++ b/Engine/Plugins/Experimental/UAF/UAF/Source/UAFEditor/Private/Variables/Outliner/VariablesOutlinerMode.cpp
@@ -38,7 +38,7 @@
 #include "Variables/AnimNextSharedVariablesFactory.h"
 #include "Variables/AnimNextVariableSettings.h"
 #include "Widgets/Docking/SDockTab.h"
-#include "AnimNextScopedCompileJob.h"
+#include "UAFCompilationScope.h"
 #include "Framework/Notifications/NotificationManager.h"
 #include "HAL/PlatformApplicationMisc.h"
 #include "Common/Outliner/OutlinerItemMenuContexts.h"
@@ -1634,6 +1634,8 @@ void FVariablesOutlinerMode::PopulateNewVariableToolMenuEntries(UToolMenu* InMen
                     {
                         UAnimNextVariableSettings* Settings = GetMutableDefault<UAnimNextVariableSettings>();
                         FScopedTransaction Transaction(LOCTEXT("AddVariable", "Add variable"));
+
+                        UncookedOnly::FCompilationScope CompileScope(LOCTEXT("AddVariable", "Add Variable"), UncookedOnly::FUtils::GetAsset(EditorData));
                         
                         const FName VariableName = UncookedOnly::FUtils::GetValidVariableName(EditorData, Settings->GetLastVariableName());		
 
diff --git a/Engine/Plugins/Experimental/UAF/UAF/Source/UAFEditor/Private/Workspace/AnimNextWorkspaceEditorMode.cpp b/Engine/Plugins/Experimental/UAF/UAF/Source/UAFEditor/Private/Workspace/AnimNextWorkspaceEditorMode.cpp
--- a/Engine/Plugins/Experimental/UAF/UAF/Source/UAFEditor/Private/Workspace/AnimNextWorkspaceEditorMode.cpp
+++ b/Engine/Plugins/Experimental/UAF/UAF/Source/UAFEditor/Private/Workspace/AnimNextWorkspaceEditorMode.cpp
@@ -5,7 +5,7 @@
 #include "AnimNextEditorModule.h"
 #include "AnimNextRigVMAsset.h"
 #include "AnimNextRigVMAssetEditorData.h"
-#include "AnimNextScopedCompileJob.h"
+#include "UAFCompilationScope.h"
 #include "Component/AnimNextComponent.h"
 #include "Common/UAFBrowserTabSummoner.h"
 #include "ContextObjectStore.h"
@@ -230,69 +230,69 @@ void UAnimNextWorkspaceEditorMode::HandleCompile()
 		return;
 	}
 
-	TArray<UObject*> Assets;
+	TArray<UAnimNextRigVMAsset*> AssetsToCompile;
 
-	UObject* Asset = nullptr;
 	FText JobName = LOCTEXT("DefaultJobName", "Job");
 	if(IsCompileWholeWorkspaceChecked())
 	{
-		WorkspaceEditor->GetOpenedAssets<UObject>(Assets);
-		if(Assets.Num() == 1)
+		TArray<UObject*> OpenObjects;
+		WorkspaceEditor->GetOpenedAssets<UObject>(OpenObjects);
+
+		Algo::TransformIf(OpenObjects, AssetsToCompile, [](UObject* Obj){ return Cast<UAnimNextRigVMAsset>(Obj); }, [](UObject* Obj){ return CastChecked<UAnimNextRigVMAsset>(Obj); } );
+
+		if(AssetsToCompile.Num() == 1)
 		{
-			Asset = Assets[0];
-			JobName = FText::FromName(Asset->GetFName());
+			JobName = FText::FromName(AssetsToCompile[0]->GetFName());
 		}
 		else
 		{
-			Asset = WorkspaceEditor->GetWorkspaceAsset();
-			JobName = FText::FromName(Asset->GetFName());
+			JobName = FText::FromName(WorkspaceEditor->GetWorkspaceAsset()->GetFName());
 		}
 	}
 	else if(UObject* FocusedDocument = WorkspaceEditor->GetFocusedWorkspaceDocument().GetObject())
 	{
 		// Find the outer asset for this document
-		Asset = FocusedDocument;
-		
-		while(Asset && (!Asset->IsAsset() || Asset->IsPackageExternal()))
+		UObject* FocussedObject = FocusedDocument;
+		while(FocussedObject && (!FocussedObject->IsAsset() || FocussedObject->IsPackageExternal()))
 		{
-			Asset = Asset->GetOuter();
+			FocussedObject = FocussedObject->GetOuter();
 		}
 
-		if(Asset)
+		if(UAnimNextRigVMAsset* Asset = Cast<UAnimNextRigVMAsset>(FocussedObject))
 		{
 			JobName = FText::FromName(Asset->GetFName());
-			Assets.Add(Asset);
+			AssetsToCompile.Add(Asset);
 		}
 	}
 
 	if(IsCompileDirtyFilesChecked())
 	{
-		for(int32 AssetIndex = 0; AssetIndex < Assets.Num(); ++AssetIndex)
+		for(int32 AssetIndex = 0; AssetIndex < AssetsToCompile.Num(); ++AssetIndex)
 		{
-			UObject* AssetToCheck = Assets[AssetIndex];
+			UObject* AssetToCheck = AssetsToCompile[AssetIndex];
 			if(TSharedPtr<UE::UAF::Editor::IAssetCompilationHandler> FoundCompiler = GetAssetCompiler(AssetToCheck))
 			{
 				UE::UAF::Editor::ECompileStatus AssetCompileStatus = FoundCompiler->GetCompileStatus(WorkspaceEditor.ToSharedRef(), AssetToCheck);
 				if( AssetCompileStatus != UE::UAF::Editor::ECompileStatus::Dirty &&
 					AssetCompileStatus != UE::UAF::Editor::ECompileStatus::Error)
 				{
-					Assets.RemoveAtSwap(AssetIndex);
+					AssetsToCompile.RemoveAtSwap(AssetIndex);
 					AssetIndex--;
 				}
 			}
 		}
 	}
 
-	if(Assets.Num() == 0)
+	if(AssetsToCompile.Num() == 0)
 	{
 		return;
 	}
 
-	// Start a batch compile scope
-	UE::UAF::UncookedOnly::FScopedCompileJob CompileJob(JobName, Assets);
+	// Open a batch compile scope
+	UE::UAF::UncookedOnly::FCompilationScope CompileScope(JobName, AssetsToCompile);
 
 	UE::UAF::Editor::FAnimNextEditorModule& AnimNextEditorModule = FModuleManager::LoadModuleChecked<UE::UAF::Editor::FAnimNextEditorModule>("UAFEditor");
-	for(UObject* AssetToCompile : Assets)
+	for(UObject* AssetToCompile : AssetsToCompile)
 	{
 		if(TSharedPtr<UE::UAF::Editor::IAssetCompilationHandler> FoundCompiler = GetAssetCompiler(AssetToCompile))
 		{
diff --git a/Engine/Plugins/Experimental/UAF/UAF/Source/UAFEditor/Public/AssetCompilationHandler.h b/Engine/Plugins/Experimental/UAF/UAF/Source/UAFEditor/Public/AssetCompilationHandler.h
--- a/Engine/Plugins/Experimental/UAF/UAF/Source/UAFEditor/Public/AssetCompilationHandler.h
+++ b/Engine/Plugins/Experimental/UAF/UAF/Source/UAFEditor/Public/AssetCompilationHandler.h
@@ -14,7 +14,7 @@ namespace UE::UAF::Editor
 class FAssetCompilationHandler : public IAssetCompilationHandler
 {
 public:
-	UE_API explicit FAssetCompilationHandler(UObject* InAsset);
+	UE_API explicit FAssetCompilationHandler(const UObject* InAsset);
 
 protected:
 	// IAssetCompilationHandler interface
@@ -22,9 +22,6 @@ protected:
 	UE_API virtual void SetAutoCompile(TSharedRef<Workspace::IWorkspaceEditor> InWorkspaceEditor, UObject* InAsset, bool bInAutoCompile) override;
 	UE_API virtual bool GetAutoCompile(TSharedRef<Workspace::IWorkspaceEditor> InWorkspaceEditor, const UObject* InAsset) const override;
 	UE_API virtual ECompileStatus GetCompileStatus(TSharedRef<Workspace::IWorkspaceEditor> InWorkspaceEditor, const UObject* InAsset) const override;
-
-	int32 NumErrors = 0;
-	int32 NumWarnings = 0;
 };
 
 }
diff --git a/Engine/Plugins/Experimental/UAF/UAF/Source/UAFUncookedOnly/Internal/AnimNextRigVMAssetEditorData.h b/Engine/Plugins/Experimental/UAF/UAF/Source/UAFUncookedOnly/Internal/AnimNextRigVMAssetEditorData.h
--- a/Engine/Plugins/Experimental/UAF/UAF/Source/UAFUncookedOnly/Internal/AnimNextRigVMAssetEditorData.h
+++ b/Engine/Plugins/Experimental/UAF/UAF/Source/UAFUncookedOnly/Internal/AnimNextRigVMAssetEditorData.h
@@ -35,7 +35,8 @@ namespace UE::UAF::UncookedOnly
 	struct FPublicVariablesImpl;
 	struct FUtils;
 	struct FUtilsPrivate;
-	class FScopedCompileJob;
+	class FCompilationScope;
+	struct Compilation;
 }
 
 namespace UE::UAF::Editor
@@ -262,7 +263,8 @@ protected:
 	friend UE::UAF::Editor::FAnimNextAnimGraphEditorModule;
 	friend class UE::UAF::Editor::FCallFunctionSharedDataDetails;
 	friend UAnimNextWorkspaceEditorMode;
-	friend UE::UAF::UncookedOnly::FScopedCompileJob;
+	friend UE::UAF::UncookedOnly::FCompilationScope;
+	friend UE::UAF::UncookedOnly::Compilation;
 	friend UE::UAF::Editor::FAssetCompilationHandler;
 	friend class UAnimNextStateTreeTreeEditorData;
 	friend class UE::UAF::Editor::FVariableBindingPropertyCustomization;
diff --git a/Engine/Plugins/Experimental/UAF/UAF/Source/UAFUncookedOnly/Private/AnimNextRigVMAssetEditorData.cpp b/Engine/Plugins/Experimental/UAF/UAF/Source/UAFUncookedOnly/Private/AnimNextRigVMAssetEditorData.cpp
--- a/Engine/Plugins/Experimental/UAF/UAF/Source/UAFUncookedOnly/Private/AnimNextRigVMAssetEditorData.cpp
+++ b/Engine/Plugins/Experimental/UAF/UAF/Source/UAFUncookedOnly/Private/AnimNextRigVMAssetEditorData.cpp
@@ -13,7 +13,7 @@
 #include "AnimNextRigVMAssetSchema.h"
 #include "AnimNextAssetWorkspaceAssetUserData.h"
 #include "AnimNextControllerBase.h"
-#include "AnimNextScopedCompileJob.h"
+#include "UAFCompilationScope.h"
 #include "RigVMPythonUtils.h"
 #include "ExternalPackageHelper.h"
 #include "IAnimNextRigVMGraphInterface.h"
@@ -64,7 +64,7 @@ void UAnimNextRigVMAssetEditorData::BroadcastModified(EAnimNextEditorDataNotifTy
 
 	// Modifications here can trigger compilation, so add a scope to catch compile batches
 	UAnimNextRigVMAsset* Asset = FUtils::GetAsset<UAnimNextRigVMAsset>(this);
-	FScopedCompileJob CompilerResults(LOCTEXT("ModifiedAssetsJobName", "Modified Assets"), { Asset });
+	FCompilationScope CompileScope(LOCTEXT("ModifiedAssetsJobName", "Modified Assets"), Asset);
 
 	if(!bSuspendEditorDataNotifications)
 	{
@@ -1043,7 +1043,7 @@ void UAnimNextRigVMAssetEditorData::RecompileVM()
 	TGuardValue<bool> CompilingGuard(bIsCompiling, true);
 
 	UAnimNextRigVMAsset* Asset = FUtils::GetAsset<UAnimNextRigVMAsset>(this);
-	FScopedCompileJob CompilerResults(Asset);
+	FCompilationScope CompileScope(Asset);
 
 	VMCompileSettings.SetExecuteContextStruct(FAnimNextExecuteContext::StaticStruct());
 	FRigVMCompileSettings Settings = (bCompileInDebugMode) ? FRigVMCompileSettings::Fast(VMCompileSettings.GetExecuteContextStruct()) : VMCompileSettings;
@@ -2307,7 +2307,7 @@ UAnimNextVariableEntry* UAnimNextRigVMAssetEditorData::AddVariable(FName InName,
 
 	// We are about to perform multiple changes, add a scope now to avoid multiple compile broadcasts per change
 	UAnimNextRigVMAsset* Asset = FUtils::GetAsset<UAnimNextRigVMAsset>(this);
-	FScopedCompileJob CompilerResults(LOCTEXT("ModifiedAssetsAddVariable", "Modified Assets Add Variable"), { Asset });
+	FCompilationScope CompileScope(LOCTEXT("ModifiedAssetsAddVariable", "Add Variable"), Asset);
 
 	bool bAlreadyExists = Entries.ContainsByPredicate(DuplicateNamePredicate);
 	int32 NameNumber = InName.GetNumber() + 1;
@@ -2405,7 +2405,7 @@ UAnimNextEventGraphEntry* UAnimNextRigVMAssetEditorData::AddEventGraph(FName InN
 
 	// We are about to perform multiple changes, add a scope now to avoid multiple compile broadcasts per change
 	UAnimNextRigVMAsset* Asset = FUtils::GetAsset<UAnimNextRigVMAsset>(this);
-	FScopedCompileJob CompilerResults(LOCTEXT("ModifiedAssetsAddEventGraph", "Modified Assets Add Event Graph"), { Asset });
+	FCompilationScope CompileScope(LOCTEXT("ModifiedAssetsAddEventGraph", "Add Event Graph"), Asset);
 
 	UAnimNextEventGraphEntry* NewEntry = CreateNewSubEntry<UAnimNextEventGraphEntry>(this);
 	NewEntry->GraphName = NewGraphName;
@@ -2538,7 +2538,7 @@ UAnimNextSharedVariablesEntry* UAnimNextRigVMAssetEditorData::AddSharedVariables
 
 	// We are about to perform multiple changes, add a scope now to avoid multiple compile broadcasts per change
 	UAnimNextRigVMAsset* Asset = FUtils::GetAsset<UAnimNextRigVMAsset>(this);
-	FScopedCompileJob CompilerResults(LOCTEXT("ModifiedAssetsAddSharedVariables", "Modified Assets Add Shared Variables"), { Asset });
+	FCompilationScope CompileScope(LOCTEXT("ModifiedAssetsAddSharedVariables", "Add Shared Variables"), Asset);
 
 	UAnimNextSharedVariablesEntry* NewEntry = CreateNewSubEntry<UAnimNextSharedVariablesEntry>(this);
 	NewEntry->SetAsset(InSharedVariables);
@@ -2606,7 +2606,7 @@ UAnimNextSharedVariablesEntry* UAnimNextRigVMAssetEditorData::AddSharedVariables
 
 	// We are about to perform multiple changes, add a scope now to avoid multiple compile broadcasts per change
 	UAnimNextRigVMAsset* Asset = FUtils::GetAsset<UAnimNextRigVMAsset>(this);
-	FScopedCompileJob CompilerResults(LOCTEXT("ModifiedAssetsAddSharedVariablesStruct", "Modified Assets Add Shared Variables Struct"), { Asset });
+	FCompilationScope CompileScope(LOCTEXT("ModifiedAssetsAddSharedVariablesStruct", "Add Shared Variables Struct"),  Asset);
 
 	UAnimNextSharedVariablesEntry* NewEntry = CreateNewSubEntry<UAnimNextSharedVariablesEntry>(this);
 	NewEntry->SetStruct(InStruct);
@@ -2771,7 +2771,8 @@ void UAnimNextRigVMAssetEditorData::HandleReportFromCompiler(EMessageSeverity::T
 		return;
 	}
 
-	FCompilerResultsLog& Log = UE::UAF::UncookedOnly::FScopedCompileJob::GetLog();
+	UAnimNextRigVMAsset* Asset = UE::UAF::UncookedOnly::FUtils::GetAsset<UAnimNextRigVMAsset>(this);
+	FCompilerResultsLog& Log = UE::UAF::UncookedOnly::FCompilationScope::GetLogForObject(Asset);
 
 	UObject* SubjectForMessage = InSubject;
 	if(URigVMNode* ModelNode = Cast<URigVMNode>(SubjectForMessage))
@@ -2860,7 +2861,8 @@ void UAnimNextRigVMAssetEditorData::HandleMessageFromCompiler(TSharedRef<FTokeni
 		return;
 	}
 
-	FCompilerResultsLog& Log = UE::UAF::UncookedOnly::FScopedCompileJob::GetLog();
+	UAnimNextRigVMAsset* Asset = UE::UAF::UncookedOnly::FUtils::GetAsset<UAnimNextRigVMAsset>(this);
+	FCompilerResultsLog& Log = UE::UAF::UncookedOnly::FCompilationScope::GetLogForObject(Asset);
 	
 	TSharedPtr<FTokenizedMessage> Message;
 	if (InMessage->GetSeverity() == EMessageSeverity::Error)
diff --git a/Engine/Plugins/Experimental/UAF/UAF/Source/UAFUncookedOnly/Private/AnimNextScopedCompileJob.cpp b/Engine/Plugins/Experimental/UAF/UAF/Source/UAFUncookedOnly/Private/AnimNextScopedCompileJob.cpp
deleted file mode 100644
diff --git a/Engine/Plugins/Experimental/UAF/UAF/Source/UAFUncookedOnly/Private/Entries/AnimNextSharedVariablesEntry.cpp b/Engine/Plugins/Experimental/UAF/UAF/Source/UAFUncookedOnly/Private/Entries/AnimNextSharedVariablesEntry.cpp
--- a/Engine/Plugins/Experimental/UAF/UAF/Source/UAFUncookedOnly/Private/Entries/AnimNextSharedVariablesEntry.cpp
+++ b/Engine/Plugins/Experimental/UAF/UAF/Source/UAFUncookedOnly/Private/Entries/AnimNextSharedVariablesEntry.cpp
@@ -5,6 +5,7 @@
 #if WITH_LIVE_CODING
 #include "ILiveCodingModule.h"
 #endif
+#include "UAFCompilationScope.h"
 #include "UncookedOnlyUtils.h"
 #include "Variables/AnimNextSharedVariables_EditorData.h"
 #include "Entries/AnimNextVariableEntry.h"
@@ -133,9 +134,10 @@ void UAnimNextSharedVariablesEntry::HandleAssetModified(UAnimNextRigVMAssetEdito
 	case EAnimNextEditorDataNotifType::VariableTypeChanged:
 	case EAnimNextEditorDataNotifType::VariableDefaultValueChanged:
 	case EAnimNextEditorDataNotifType::VariableBindingChanged:
-		if(UAnimNextRigVMAssetEditorData* EditorData = GetTypedOuter<UAnimNextRigVMAssetEditorData>())
+		if(UAnimNextRigVMAsset* OuterAsset = GetTypedOuter<UAnimNextRigVMAsset>())
 		{
-			EditorData->RequestAutoVMRecompilation();
+			// Recompile the owning asset (which references the modified asset through this entry)
+			UE::UAF::UncookedOnly::Compilation::RequestAssetCompilation(OuterAsset);
 		}
 		break;
 	default:
diff --git a/Engine/Plugins/Experimental/UAF/UAF/Source/UAFUncookedOnly/Private/UAFCompilationScope.cpp b/Engine/Plugins/Experimental/UAF/UAF/Source/UAFUncookedOnly/Private/UAFCompilationScope.cpp
new file mode 100644
--- /dev/null
+++ b/Engine/Plugins/Experimental/UAF/UAF/Source/UAFUncookedOnly/Private/UAFCompilationScope.cpp
@@ -0,0 +1,277 @@
+﻿// Copyright Epic Games, Inc. All Rights Reserved.
+
+#include "UAFCompilationScope.h"
+
+#include "AnimNextRigVMAsset.h"
+#include "AnimNextRigVMAssetEditorData.h"
+#include "Editor.h"
+#include "StatusBarSubsystem.h"
+#include "UncookedOnlyUtils.h"
+#include "Kismet2/CompilerResultsLog.h"
+#include "Logging/MessageLog.h"
+#include "Misc/UObjectToken.h"
+
+#define LOCTEXT_NAMESPACE "UAFCompilationScope"
+
+namespace UE::UAF::UncookedOnly
+{
+
+struct FCompileJobThreadData
+{
+	FText JobName;
+	TMap<TWeakObjectPtr<const UObject>, FCompilerResultsLog> PerObjectCompileLog;
+	TArray<TWeakObjectPtr<UAnimNextRigVMAsset>> CompiledAssets;
+	int32 ScopeDepth = 0;
+	double StartTimeInSeconds = 0;
+};
+
+static thread_local FCompileJobThreadData GCompileJobData;
+
+void Compilation::RequestAssetCompilation(UAnimNextRigVMAsset* InAsset)
+{
+	FCompilationScope Scope(InAsset);
+	UAnimNextRigVMAssetEditorData* EditorData = FUtils::GetEditorData<UAnimNextRigVMAssetEditorData>(InAsset);
+	EditorData->RequestAutoVMRecompilation();
+}
+
+void Compilation::RequestAssetCompilation(TConstArrayView<UAnimNextRigVMAsset*> InAssets)
+{
+	if (InAssets.Num())
+	{
+		FCompilationScope Scope(InAssets);
+
+		for (UAnimNextRigVMAsset* Asset : InAssets)
+		{
+			check(Asset);
+			UAnimNextRigVMAssetEditorData* EditorData = FUtils::GetEditorData<UAnimNextRigVMAssetEditorData>(Asset);
+			EditorData->RequestAutoVMRecompilation();
+		}
+	}
+}
+
+FCompilationScope::FCompilationScope(const FText& InJobName)
+	: FCompilationScope(InJobName, {})
+{
+}
+
+FCompilationScope::FCompilationScope(UAnimNextRigVMAsset* InAsset)
+	: FCompilationScope(FText::Format(LOCTEXT("AssetCompilationFormat", "Compiling {0}"), FText::FromName(InAsset->GetFName())), InAsset)
+{
+}
+
+FCompilationScope::FCompilationScope(TConstArrayView<UAnimNextRigVMAsset*> InAssets)
+{
+	FString AssetsList;	
+	for (const UAnimNextRigVMAsset* Asset : InAssets)
+	{
+		if (AssetsList.Len())
+		{
+			AssetsList.Append(TEXT(", "));
+		}
+		
+		AssetsList.Append(Asset->GetName());
+	}
+	
+	ProcessAssets(FText::Format(LOCTEXT("MultiAssetCompilationFormat", "Compiling ({0})"), FText::FromString(AssetsList)), InAssets);
+}
+
+FCompilationScope::FCompilationScope(const FText& InJobName, UAnimNextRigVMAsset* InAsset)
+	: FCompilationScope(InJobName, {&InAsset, 1})
+{
+}
+
+FCompilationScope::FCompilationScope(const FText& InJobName, TConstArrayView<UAnimNextRigVMAsset*> InAssets)
+{
+	ProcessAssets(InJobName, InAssets);
+}
+
+FCompilationScope::~FCompilationScope()
+{
+	FCompileJobThreadData* ThreadData = &GCompileJobData;
+
+	if(ThreadData->ScopeDepth == 1)
+	{
+		if(ThreadData->CompiledAssets.Num())
+		{
+			TArray<double> PerAssetTimeCompileTimeSpan;
+			for (int32 AssetIndex = 0; AssetIndex < ThreadData->CompiledAssets.Num(); ++AssetIndex)
+			{
+				TWeakObjectPtr<UAnimNextRigVMAsset> WeakAsset = ThreadData->CompiledAssets[AssetIndex];
+				const double StartAssetTime = FPlatformTime::Seconds();
+				if (UAnimNextRigVMAsset* Asset = WeakAsset.Get())
+				{
+					UAnimNextRigVMAssetEditorData* EditorData = FUtils::GetEditorData<UAnimNextRigVMAssetEditorData>(Asset);
+					if(EditorData != nullptr)
+					{
+						EditorData->DecrementVMRecompileBracket();
+					}
+				}
+				const double FinishAssetTime = FPlatformTime::Seconds();
+				PerAssetTimeCompileTimeSpan.Add(FinishAssetTime - StartAssetTime);
+			}
+
+			FMessageLog MessageLog("AnimNextCompilerResults");
+
+			MessageLog.NewPage(FText::Format(LOCTEXT("MessageLogPageFormat", "{0}: {1}"), ThreadData->JobName, FText::AsDateTime(FDateTime::UtcNow())));
+
+			int32 NumErrors = 0;
+			int32 NumWarnings = 0;
+
+			for (const auto& LogPair : ThreadData->PerObjectCompileLog)
+			{
+				NumErrors += LogPair.Value.NumErrors;
+				NumWarnings += LogPair.Value.NumWarnings;
+			}
+
+			// Print summary
+			FNumberFormattingOptions TimeFormat;
+			TimeFormat.MaximumFractionalDigits = 2;
+			TimeFormat.MinimumFractionalDigits = 0;
+			TimeFormat.UseGrouping = false;
+			
+			const double ElapsedSeconds = FPlatformTime::Seconds() - ThreadData->StartTimeInSeconds;
+			FFormatNamedArguments Args;
+			Args.Add(TEXT("JobName"), ThreadData->JobName);
+			Args.Add(TEXT("CompileTimeMs"), ElapsedSeconds * 1000.f);
+			Args.Add(TEXT("NumAssets"), ThreadData->CompiledAssets.Num());
+
+			const int32 AssetMessageIndentationLevel = ThreadData->PerObjectCompileLog.Num() == 1 ? 0 : 1;
+
+			FText JobOverViewMessage;
+			if (NumErrors > 0)
+			{
+				Args.Add(TEXT("NumErrors"), NumErrors);
+				Args.Add(TEXT("NumWarnings"),  NumWarnings);
+				JobOverViewMessage = FText::Format(LOCTEXT("CompileFailed", "[{JobName}] compilation has failed. {NumErrors} {NumErrors}|plural(one=Error,other=Errors), {NumWarnings} {NumWarnings}|plural(one=Warning,other=Warnings) [{NumAssets} {NumAssets}|plural(one=asset,other=assets) in {CompileTimeMs} ms]"), MoveTemp(Args));
+			}
+			else if(NumWarnings > 0)
+			{
+				Args.Add(TEXT("NumWarnings"),  NumWarnings);
+				JobOverViewMessage = FText::Format(LOCTEXT("CompileWarning", "[{JobName}] compilation was successful. {NumWarnings} {NumWarnings}|plural(one=Warning,other=Warnings) [{NumAssets} {NumAssets}|plural(one=asset,other=assets) in {CompileTimeMs} ms]"), MoveTemp(Args));
+			}
+			else
+			{
+				JobOverViewMessage = FText::Format(LOCTEXT("CompileSuccess", "[{JobName}] compilation was successful! [{NumAssets} {NumAssets}|plural(one=asset,other=assets) in {CompileTimeMs} ms]"), MoveTemp(Args));
+			}
+			
+			MessageLog.Info(JobOverViewMessage);
+			
+			for (const auto& LogPair : ThreadData->PerObjectCompileLog)
+			{
+				const int32 AssetIndex = ThreadData->CompiledAssets.IndexOfByKey(LogPair.Key);
+			
+				TSharedRef<FTokenizedMessage> Message = FTokenizedMessage::Create(EMessageSeverity::Info);
+				Message->AddToken(FUObjectToken::Create(LogPair.Key.Get()));
+				Message->AddText(FText::Format(LOCTEXT("AssetCompileTimeFormat", "in {0} ms."), FText::AsNumber(PerAssetTimeCompileTimeSpan[AssetIndex] * 1000.f, &TimeFormat)));
+				Message->SetIndentationLevel(AssetMessageIndentationLevel);
+				MessageLog.AddMessage(Message);
+				
+				if (LogPair.Value.Messages.Num())
+				{
+					for (auto LogMessage : LogPair.Value.Messages)
+					{
+						LogMessage->SetIndentationLevel(AssetMessageIndentationLevel + 1);
+						MessageLog.AddMessage(LogMessage);
+					}
+					
+					NumErrors += LogPair.Value.NumErrors;
+					NumWarnings += LogPair.Value.NumWarnings;
+				}
+			}
+			
+			// Broadcast compilation finished so reallocation can occur
+			for (TWeakObjectPtr<UAnimNextRigVMAsset> WeakAsset : ThreadData->CompiledAssets)
+			{
+				UAnimNextRigVMAsset* Asset = WeakAsset.Get();
+				if (Asset == nullptr)
+				{
+					continue;
+				}
+
+				UAnimNextRigVMAssetEditorData* EditorData = FUtils::GetEditorData<UAnimNextRigVMAssetEditorData>(Asset);
+				if (EditorData == nullptr)
+				{
+					continue;
+				}
+
+				Asset->CompilationState = [&EditorData]()
+				{
+					if (EditorData->bErrorsDuringCompilation)
+					{
+						return EAnimNextRigVMAssetState::CompiledWithErrors;
+					}
+
+					if (EditorData->bWarningsDuringCompilation)
+					{
+						return EAnimNextRigVMAssetState::CompiledWithWarnings;
+					}
+
+					return EAnimNextRigVMAssetState::CompiledWithSuccess;
+				}();
+
+				if (!EditorData->bSuspendCompilationNotifications)
+				{
+					UAnimNextRigVMAsset::OnCompileJobFinished().Broadcast(Asset);
+					EditorData->OnCompileJobFinished();
+				}
+			}
+		}
+
+		ThreadData->CompiledAssets.Empty();
+
+		ThreadData->PerObjectCompileLog.Reset();
+		ThreadData->JobName = FText::GetEmpty();
+		ThreadData->StartTimeInSeconds = 0.0;
+	}
+
+	--ThreadData->ScopeDepth;
+}
+
+FCompilerResultsLog& FCompilationScope::GetLogForObject(const UObject* InObject)
+{
+	check(InObject);
+
+	FCompileJobThreadData* ThreadData = &GCompileJobData;
+	check(ThreadData->CompiledAssets.Contains(InObject));
+	return ThreadData->PerObjectCompileLog.FindOrAdd(InObject);
+}
+
+void FCompilationScope::ProcessAssets(const FText& InJobName, TConstArrayView<UAnimNextRigVMAsset*> InAssets)
+{
+	FCompileJobThreadData* ThreadData = &GCompileJobData;
+	
+	if (ThreadData->ScopeDepth == 0)
+	{
+		ThreadData->JobName = InJobName;
+		ThreadData->StartTimeInSeconds = FPlatformTime::Seconds();
+	}
+		
+	++ThreadData->ScopeDepth;
+
+	for(UAnimNextRigVMAsset* Asset : InAssets)
+	{
+		check(Asset);
+		UAnimNextRigVMAssetEditorData* EditorData = FUtils::GetEditorData<UAnimNextRigVMAssetEditorData>(Asset);
+		if(EditorData == nullptr)
+		{
+			return;
+		}
+
+		if (!ThreadData->CompiledAssets.Contains(Asset))
+		{
+			ThreadData->CompiledAssets.Add(Asset);
+			EditorData->ClearErrorInfoForAllEdGraphs();
+
+			if (!EditorData->bSuspendCompilationNotifications)
+			{
+				EditorData->OnCompileJobStarted();
+				UAnimNextRigVMAsset::OnCompileJobStarted().Broadcast(Asset);
+				EditorData->IncrementVMRecompileBracket();
+			}
+		}
+	}
+}
+
+}
+
+#undef LOCTEXT_NAMESPACE
diff --git a/Engine/Plugins/Experimental/UAF/UAF/Source/UAFUncookedOnly/Private/UncookedOnlyUtils.cpp b/Engine/Plugins/Experimental/UAF/UAF/Source/UAFUncookedOnly/Private/UncookedOnlyUtils.cpp
--- a/Engine/Plugins/Experimental/UAF/UAF/Source/UAFUncookedOnly/Private/UncookedOnlyUtils.cpp
+++ b/Engine/Plugins/Experimental/UAF/UAF/Source/UAFUncookedOnly/Private/UncookedOnlyUtils.cpp
@@ -33,7 +33,7 @@
 #include "UObject/AssetRegistryTagsContext.h"
 #include "Variables/AnimNextSoftVariableReference.h"
 #include "AnimNextExports.h"
-#include "AnimNextScopedCompileJob.h"
+#include "UAFCompilationScope.h"
 #include "AnimNextSharedVariableNode.h"
 #include "ScopedTransaction.h"
 #include "Misc/UObjectToken.h"
@@ -1612,7 +1612,7 @@ void FUtils::DeleteVariable(UAnimNextVariableEntry* VariableEntry, bool bSetupUn
 		UAnimNextRigVMAsset* OuterRigVMAsset = VariableEntry->GetTypedOuter<UAnimNextRigVMAsset>();
 		check(OuterRigVMAsset);
 		
-		FScopedCompileJob CompilerResults(LOCTEXT("ModifiedAssets_DeleteVariable", "Modified Assets Delete Variable"), { OuterRigVMAsset });
+		FCompilationScope CompilerResults(LOCTEXT("ModifiedAssets_DeleteVariable", "Modified Assets Delete Variable"), OuterRigVMAsset);
 		
 		const FAnimNextSoftVariableReference FindReference(VariableEntry->GetEntryName(), OuterRigVMAsset);
 
@@ -1643,7 +1643,7 @@ void FUtils::RenameVariable(UAnimNextVariableEntry* VariableEntry, const FName N
 		UAnimNextRigVMAsset* OuterRigVMAsset = VariableEntry->GetTypedOuter<UAnimNextRigVMAsset>();
 		check(OuterRigVMAsset);
 
-		FScopedCompileJob CompilerResults(LOCTEXT("ModifiedAssets_RenameVariable", "Modified Assets Rename Variable"), { OuterRigVMAsset });
+		FCompilationScope CompilerResults(LOCTEXT("ModifiedAssets_RenameVariable", "Modified Assets Rename Variable"), OuterRigVMAsset);
 		
 		const FAnimNextSoftVariableReference FindReference(VariableEntry->GetEntryName(), OuterRigVMAsset);
 		
@@ -1673,7 +1673,7 @@ UAnimNextVariableEntry* FUtils::MoveVariableToAsset(UAnimNextVariableEntry* Vari
 	{
 		if (UAnimNextRigVMAsset* CurrentOuter = VariableEntry->GetTypedOuter<UAnimNextRigVMAsset>())
 		{
-			FScopedCompileJob CompilerResults(LOCTEXT("ModifiedAssets_MoveVariableToAsset", "Modified Assets Move Variable"), { CurrentOuter, NewOuter });
+			FCompilationScope CompilerResults(LOCTEXT("ModifiedAssets_MoveVariableToAsset", "Modified Assets Move Variable"), MakeConstArrayView({ CurrentOuter, NewOuter }));
 
 			UAnimNextRigVMAssetEditorData* CurrentEditorData = GetEditorData(CurrentOuter);
 			UAnimNextRigVMAssetEditorData* NewEditorData = GetEditorData(NewOuter);
@@ -1716,7 +1716,7 @@ void FUtils::SetVariableType(UAnimNextVariableEntry* VariableEntry, const FAnimN
 		UAnimNextRigVMAsset* OuterRigVMAsset = VariableEntry->GetTypedOuter<UAnimNextRigVMAsset>();
 		check(OuterRigVMAsset);
 
-		FScopedCompileJob CompilerResults(LOCTEXT("ModifiedAssets_SetVariableType", "Modified Assets Set Variable Type"), { OuterRigVMAsset });
+		FCompilationScope CompilerResults(LOCTEXT("ModifiedAssets_SetVariableType", "Modified Assets Set Variable Type"), OuterRigVMAsset);
 		
 		const FAnimNextSoftVariableReference FindReference(VariableEntry->GetEntryName(), OuterRigVMAsset);
 		
@@ -1801,9 +1801,8 @@ void FUtils::ReplaceVariableReferences(UAnimNextRigVMAssetEditorData* InEditorDa
 			FString ErrorMessage;
 			if (CanAddSharedVariablesReference(InEditorData, SharedVariablesReplaceOuter, &ErrorMessage))
 			{
-				UE::UAF::UncookedOnly::FScopedCompileJob CompilerResults(LOCTEXT("AddSharedAssetJobName", "Add Shared Asset dependency"), { ThisOuterAsset });
+				FCompilationScope CompilerResults(LOCTEXT("AddSharedAssetJobName", "Add Shared Asset dependency"), ThisOuterAsset);
 				InEditorData->AddSharedVariables(SharedVariablesReplaceOuter, bSetupUndoRedo, bPrintPythonCommands);
-				InEditorData->RecompileVM();
 			}
 			else
 			{
diff --git a/Engine/Plugins/Experimental/UAF/UAF/Source/UAFUncookedOnly/Public/AnimNextScopedCompileJob.h b/Engine/Plugins/Experimental/UAF/UAF/Source/UAFUncookedOnly/Public/AnimNextScopedCompileJob.h
deleted file mode 100644
diff --git a/Engine/Plugins/Experimental/UAF/UAF/Source/UAFUncookedOnly/Public/UAFCompilationScope.h b/Engine/Plugins/Experimental/UAF/UAF/Source/UAFUncookedOnly/Public/UAFCompilationScope.h
new file mode 100644
--- /dev/null
+++ b/Engine/Plugins/Experimental/UAF/UAF/Source/UAFUncookedOnly/Public/UAFCompilationScope.h
@@ -0,0 +1,51 @@
+﻿// Copyright Epic Games, Inc. All Rights Reserved.
+
+#pragma once
+
+#include "CoreMinimal.h"
+#include "Kismet2/CompilerResultsLog.h"
+
+class UAnimNextRigVMAsset;
+class UAnimNextRigVMAssetEditorData;
+class FCompilerResultsLog;
+
+namespace UE::UAF::UncookedOnly
+{
+
+struct Compilation
+{
+	// Request to synchronously compile the asset
+	UAFUNCOOKEDONLY_API static void RequestAssetCompilation(UAnimNextRigVMAsset* InAsset);
+	// Request to synchronously compile a set of assets
+	UAFUNCOOKEDONLY_API static void RequestAssetCompilation(TConstArrayView<UAnimNextRigVMAsset*> InAssets);
+};
+
+// RAII helper for scoping compilation batches
+// The re-allocation of compiled assets and their dependencies is deferred until the outermost scope exits.
+class FCompilationScope
+{
+public:
+	// Standalone scope with no immediate related assets
+	UAFUNCOOKEDONLY_API explicit FCompilationScope(const FText& InJobName);
+
+	// Scope with a single asset which the caller will compile
+	UAFUNCOOKEDONLY_API explicit FCompilationScope(UAnimNextRigVMAsset* InAsset);
+	
+	// Scope with a set of assets which the caller will compile
+	UAFUNCOOKEDONLY_API explicit FCompilationScope(TConstArrayView<UAnimNextRigVMAsset*> InAssets);
+	
+	// Scope with a explicit name and set of assets which the caller will compile
+	UAFUNCOOKEDONLY_API explicit FCompilationScope(const FText& InJobName, TConstArrayView<UAnimNextRigVMAsset*> InAssets);
+
+	// Scope with explicit name and a single asset which the caller will compile
+	UAFUNCOOKEDONLY_API explicit FCompilationScope(const FText& InJobName, UAnimNextRigVMAsset* InAsset);
+
+	UAFUNCOOKEDONLY_API ~FCompilationScope();
+
+	// Returns a CompilerResultsLog instance specifically for the provided UObject
+	UAFUNCOOKEDONLY_API static FCompilerResultsLog& GetLogForObject(const UObject* InObject);
+private:
+	static void ProcessAssets(const FText& InJobName, TConstArrayView<UAnimNextRigVMAsset*, int> InAssets);
+};
+
+}
diff --git a/Engine/Plugins/Experimental/UAF/UAFAnimGraph/Source/UAFAnimGraphEditor/Private/Graph/AnimNextAnimationGraphFactory.cpp b/Engine/Plugins/Experimental/UAF/UAFAnimGraph/Source/UAFAnimGraphEditor/Private/Graph/AnimNextAnimationGraphFactory.cpp
--- a/Engine/Plugins/Experimental/UAF/UAFAnimGraph/Source/UAFAnimGraphEditor/Private/Graph/AnimNextAnimationGraphFactory.cpp
+++ b/Engine/Plugins/Experimental/UAF/UAFAnimGraph/Source/UAFAnimGraphEditor/Private/Graph/AnimNextAnimationGraphFactory.cpp
@@ -1,6 +1,8 @@
 // Copyright Epic Games, Inc. All Rights Reserved.
 
 #include "Graph/AnimNextAnimationGraphFactory.h"
+
+#include "UAFCompilationScope.h"
 #include "Graph/AnimNextAnimationGraph.h"
 #include "Graph/AnimNextAnimationGraph_EditorData.h"
 #include "UncookedOnlyUtils.h"
@@ -44,7 +46,8 @@ UObject* UAnimNextAnimationGraphFactory::FactoryCreateNew(UClass* Class, UObject
 	NewGraph->SetHiddenInOutliner(true);
 
 	// Compile the initial skeleton
-	EditorData->RecompileVM();
+	UE::UAF::UncookedOnly::Compilation::RequestAssetCompilation(NewAnimationGraph);
+
 	check(!EditorData->bErrorsDuringCompilation);
 
 	return NewAnimationGraph;
diff --git a/Engine/Plugins/Experimental/UAF/UAFAnimGraph/Source/UAFAnimGraphUncookedOnly/Private/Graph/AnimNextAnimationGraph_EditorData.cpp b/Engine/Plugins/Experimental/UAF/UAFAnimGraph/Source/UAFAnimGraphUncookedOnly/Private/Graph/AnimNextAnimationGraph_EditorData.cpp
--- a/Engine/Plugins/Experimental/UAF/UAFAnimGraph/Source/UAFAnimGraphUncookedOnly/Private/Graph/AnimNextAnimationGraph_EditorData.cpp
+++ b/Engine/Plugins/Experimental/UAF/UAFAnimGraph/Source/UAFAnimGraphUncookedOnly/Private/Graph/AnimNextAnimationGraph_EditorData.cpp
@@ -3,7 +3,7 @@
 #include "Graph/AnimNextAnimationGraph_EditorData.h"
 #include "AnimGraphUncookedOnlyUtils.h"
 #include "AnimNextAnimGraphWorkspaceAssetUserData.h"
-#include "AnimNextScopedCompileJob.h"
+#include "UAFCompilationScope.h"
 #include "AnimNextTraitStackUnitNode.h"
 #include "RigVMPythonUtils.h"
 #include "UAFAnimGraphEdGraphSchema.h"
@@ -726,7 +726,8 @@ void UAnimNextAnimationGraph_EditorData::OnPreCompileProcessGraphs(const FRigVMC
 			TempGraph->SetExternalPackage(nullptr);
 			for(URigVMNode* SourceNode : SourceGraph->GetNodes())
 			{
-				FScopedCompileJob::GetLog().NotifyIntermediateObjectCreation(CreatedObjects.FindChecked(SourceNode), SourceNode);
+				FCompilerResultsLog& CompileLog = FCompilationScope::GetLogForObject(AnimationGraph);
+				CompileLog.NotifyIntermediateObjectCreation(CreatedObjects.FindChecked(SourceNode), SourceNode);
 			}
 
 			TempGraph->SetFlags(RF_Transient);
diff --git a/Engine/Plugins/Experimental/UAF/UAFStateTree/Source/UAFStateTreeEditor/Private/AnimNextStateTreeFactory.cpp b/Engine/Plugins/Experimental/UAF/UAFStateTree/Source/UAFStateTreeEditor/Private/AnimNextStateTreeFactory.cpp
--- a/Engine/Plugins/Experimental/UAF/UAFStateTree/Source/UAFStateTreeEditor/Private/AnimNextStateTreeFactory.cpp
+++ b/Engine/Plugins/Experimental/UAF/UAFStateTree/Source/UAFStateTreeEditor/Private/AnimNextStateTreeFactory.cpp
@@ -1,6 +1,8 @@
 // Copyright Epic Games, Inc. All Rights Reserved.
 
 #include "AnimNextStateTreeFactory.h"
+
+#include "UAFCompilationScope.h"
 #include "AnimNextStateTree.h"
 #include "AnimNextStateTreeSchema.h"
 #include "AnimNextStateTree_EditorData.h"
@@ -43,9 +45,9 @@ UObject* UAnimNextStateTreeFactory::FactoryCreateNew(UClass* Class, UObject* InP
 	EditorData->Initialize(/*bRecompileVM*/false);
 
 	// Compile the initial skeleton
-	EditorData->RecompileVM();
+	UE::UAF::UncookedOnly::Compilation::RequestAssetCompilation(NewStateTree);
 	
 	check(!EditorData->bErrorsDuringCompilation);
 
 	return NewStateTree;
-}
\ No newline at end of file
+}
diff --git a/Engine/Plugins/Experimental/UAF/UAFStateTree/Source/UAFStateTreeEditor/Private/StateTreeAssetCompilationHandler.cpp b/Engine/Plugins/Experimental/UAF/UAFStateTree/Source/UAFStateTreeEditor/Private/StateTreeAssetCompilationHandler.cpp
--- a/Engine/Plugins/Experimental/UAF/UAFStateTree/Source/UAFStateTreeEditor/Private/StateTreeAssetCompilationHandler.cpp
+++ b/Engine/Plugins/Experimental/UAF/UAFStateTree/Source/UAFStateTreeEditor/Private/StateTreeAssetCompilationHandler.cpp
@@ -3,7 +3,7 @@
 #include "StateTreeAssetCompilationHandler.h"
 
 #include "AnimNextRigVMAssetEditorData.h"
-#include "AnimNextScopedCompileJob.h"
+#include "UAFCompilationScope.h"
 #include "AnimNextStateTree.h"
 #include "AnimNextStateTree_EditorData.h"
 #include "Editor.h"
@@ -51,7 +51,8 @@ void FStateTreeAssetCompilationHandler::Compile(TSharedRef<Workspace::IWorkspace
 		return;
 	}
 
-	UE::UAF::UncookedOnly::FScopedCompileJob CompileJob(AnimNextStateTree);
+	UE::UAF::UncookedOnly::FCompilationScope CompileScope(AnimNextStateTree);
+	FCompilerResultsLog& CompileLog = CompileScope.GetLogForObject(AnimNextStateTree);
 
 	// We can't recompile state tree while in PIE
 	if (!GEditor->IsPlaySessionInProgress())
@@ -64,12 +65,12 @@ void FStateTreeAssetCompilationHandler::Compile(TSharedRef<Workspace::IWorkspace
 		// Add to messages
 		for(const TSharedRef<FTokenizedMessage>& Message : Log.ToTokenizedMessages())
 		{
-			CompileJob.GetLog().AddTokenizedMessage(Message);
+			CompileLog.AddTokenizedMessage(Message);
 		}
 	}
 	else
 	{
-		CompileJob.GetLog().Note(*FString::Printf(TEXT("StateTree '%s' compilation was skipped because PIE was running"), *InAsset->GetName()));
+		CompileLog.Note(*FString::Printf(TEXT("StateTree '%s' compilation was skipped because PIE was running"), *InAsset->GetName()));
 	}
 
 	// Call super to compile the RigVM Asset 
diff --git a/Engine/Plugins/Experimental/UAF/UAFStateTree/Source/UAFStateTreeUncookedOnly/Private/AnimNextStateTree_EditorData.cpp b/Engine/Plugins/Experimental/UAF/UAFStateTree/Source/UAFStateTreeUncookedOnly/Private/AnimNextStateTree_EditorData.cpp
--- a/Engine/Plugins/Experimental/UAF/UAFStateTree/Source/UAFStateTreeUncookedOnly/Private/AnimNextStateTree_EditorData.cpp
+++ b/Engine/Plugins/Experimental/UAF/UAFStateTree/Source/UAFStateTreeUncookedOnly/Private/AnimNextStateTree_EditorData.cpp
@@ -3,6 +3,7 @@
 #include "AnimNextStateTree_EditorData.h"
 
 #include "AnimGraphUncookedOnlyUtils.h"
+#include "UAFCompilationScope.h"
 #include "AnimNextStateTree.h"
 #include "AnimNextStateTreeEditorData.h"
 #include "AnimNextStateTreeEditorOnlyTypes.h"
@@ -47,13 +48,20 @@ void UAnimNextStateTree_EditorData::RecompileVM()
 	}
 
 	UAnimNextStateTree* AnimationStateTree = UE::UAF::UncookedOnly::FUtils::GetAsset<UAnimNextStateTree>(this);
+	UE::UAF::UncookedOnly::FCompilationScope CompileScope(AnimationStateTree);
 
 	UStateTree* InnerStateTree = AnimationStateTree->StateTree;
 
 	// Recompile ST as we just updated our variables propertybag. It may have changed from an external shared variables update.
 	FStateTreeCompilerLog Log;
-	bool bCompileSucceeded = UStateTreeEditingSubsystem::CompileStateTree(InnerStateTree, Log);
-	ensureMsgf(bCompileSucceeded, TEXT("Failed to compile state tree after data update: %s"), *InnerStateTree->GetFName().ToString());
+	UStateTreeEditingSubsystem::CompileStateTree(InnerStateTree, Log);
+
+	FCompilerResultsLog& CompileLog = UE::UAF::UncookedOnly::FCompilationScope::GetLogForObject(AnimationStateTree);
+	TArray<TSharedRef<FTokenizedMessage>> Messages = Log.ToTokenizedMessages();
+	for (const TSharedRef<FTokenizedMessage>& Message : Messages)
+	{
+		CompileLog.AddTokenizedMessage(Message);
+	}
 }
 
 TConstArrayView<TSubclassOf<UAnimNextRigVMAssetEntry>> UAnimNextStateTree_EditorData::GetEntryClasses() const
diff --git a/Engine/Plugins/Experimental/UAF/UAFTests/Source/UAFAnimGraphTestSuite/Private/GCInstanceDataGCTest.cpp b/Engine/Plugins/Experimental/UAF/UAFTests/Source/UAFAnimGraphTestSuite/Private/GCInstanceDataGCTest.cpp
--- a/Engine/Plugins/Experimental/UAF/UAFTests/Source/UAFAnimGraphTestSuite/Private/GCInstanceDataGCTest.cpp
+++ b/Engine/Plugins/Experimental/UAF/UAFTests/Source/UAFAnimGraphTestSuite/Private/GCInstanceDataGCTest.cpp
@@ -3,6 +3,7 @@
 #include "GCTestsUtil.h"
 
 #include "AnimNextRuntimeTest.h"
+#include "UAFCompilationScope.h"
 #include "AnimNextTest.h"
 #include "UncookedOnlyUtils.h"
 #include "Graph/AnimNextAnimationGraph.h"
@@ -160,7 +161,7 @@ bool FAnimationUAFRuntimeTest_InstanceDataGC::RunTest(const FString& InParameter
 		UE_RETURN_ON_ERROR(DecoratorPin != nullptr, TEXT("FAnimationUAFRuntimeTest_InstanceDataGC -> Failed to find trait pin"));
 
 		// For some reason a manual recompile is needed, otherwise the trait interface implementations are never run
-		EditorData->RecompileVM();
+		UE::UAF::UncookedOnly::Compilation::RequestAssetCompilation(AnimationGraph);
 
 		TSharedPtr<FAnimNextGraphInstance> GraphInstance = AnimationGraph->AllocateInstance();
 
@@ -201,4 +202,4 @@ bool FAnimationUAFRuntimeTest_InstanceDataGC::RunTest(const FString& InParameter
 	return true;
 }
 
-#endif // WITH_DEV_AUTOMATION_TESTS
\ No newline at end of file
+#endif // WITH_DEV_AUTOMATION_TESTS
diff --git a/Engine/Plugins/Experimental/UAF/UAFTests/Source/UAFAnimGraphTestSuite/Private/GCInstanceDataMissingGCTest.cpp b/Engine/Plugins/Experimental/UAF/UAFTests/Source/UAFAnimGraphTestSuite/Private/GCInstanceDataMissingGCTest.cpp
--- a/Engine/Plugins/Experimental/UAF/UAFTests/Source/UAFAnimGraphTestSuite/Private/GCInstanceDataMissingGCTest.cpp
+++ b/Engine/Plugins/Experimental/UAF/UAFTests/Source/UAFAnimGraphTestSuite/Private/GCInstanceDataMissingGCTest.cpp
@@ -3,6 +3,7 @@
 #include "GCTestsUtil.h"
 
 #include "AnimNextRuntimeTest.h"
+#include "UAFCompilationScope.h"
 #include "AnimNextTest.h"
 #include "UncookedOnlyUtils.h"
 #include "Graph/AnimNextAnimationGraph.h"
@@ -145,7 +146,7 @@ bool FAnimationUAFRuntimeTest_InstanceDataGCMissing::RunTest(const FString& InPa
 		UE_RETURN_ON_ERROR(DecoratorPin != nullptr, TEXT("FAnimationUAFRuntimeTest_InstanceDataGCMissing -> Failed to find trait pin"));
 
 		// For some reason a manual recompile is needed, otherwise the trait interface implementations are never run
-		EditorData->RecompileVM();
+		UE::UAF::UncookedOnly::Compilation::RequestAssetCompilation(AnimationGraph);
 
 		TSharedPtr<FAnimNextGraphInstance> GraphInstance = AnimationGraph->AllocateInstance();
 
@@ -186,4 +187,4 @@ bool FAnimationUAFRuntimeTest_InstanceDataGCMissing::RunTest(const FString& InPa
 	return true;
 }
 
-#endif // WITH_DEV_AUTOMATION_TESTS
\ No newline at end of file
+#endif // WITH_DEV_AUTOMATION_TESTS
diff --git a/Engine/Plugins/Experimental/UAF/UAFTests/Source/UAFTests/Private/UAFTestBlueprintLibrary.h b/Engine/Plugins/Experimental/UAF/UAFTests/Source/UAFTests/Private/UAFTestBlueprintLibrary.h
--- a/Engine/Plugins/Experimental/UAF/UAFTests/Source/UAFTests/Private/UAFTestBlueprintLibrary.h
+++ b/Engine/Plugins/Experimental/UAF/UAFTests/Source/UAFTests/Private/UAFTestBlueprintLibrary.h
@@ -5,6 +5,7 @@
 #include "CoreMinimal.h"
 #include "Kismet/BlueprintFunctionLibrary.h"
 #include "AnimNextRigVMAssetEditorData.h"
+#include "UAFCompilationScope.h"
 #include "UAFTestBlueprintLibrary.generated.h"
 
 UCLASS()
@@ -19,7 +20,7 @@ public:
 	{	
 		if (InAsset)
 		{
-			UE::UAF::UncookedOnly::FUtils::GetEditorData(InAsset)->RecompileVM(); 
+			UE::UAF::UncookedOnly::Compilation::RequestAssetCompilation(InAsset);
 		}
 		else
 		{
