> Remove hidden meta-data from propertybage type enum to allow for script usage of the entries
> Add variable type validation to cover invalid scripting API usage

diff --git a/Engine/Plugins/Experimental/UAF/UAF/Source/UAFUncookedOnly/Private/AnimNextRigVMAssetEditorData.cpp b/Engine/Plugins/Experimental/UAF/UAF/Source/UAFUncookedOnly/Private/AnimNextRigVMAssetEditorData.cpp
--- a/Engine/Plugins/Experimental/UAF/UAF/Source/UAFUncookedOnly/Private/AnimNextRigVMAssetEditorData.cpp
+++ b/Engine/Plugins/Experimental/UAF/UAF/Source/UAFUncookedOnly/Private/AnimNextRigVMAssetEditorData.cpp
@@ -2292,6 +2292,12 @@ UAnimNextVariableEntry* UAnimNextRigVMAssetEditorData::AddVariable(FName InName,
 		ReportError(TEXT("UAnimNextRigVMAssetEditorData::AddVariable: Cannot add a variable to this asset - entry is not allowed."));
 		return nullptr;
 	}
+	
+	if (!InType.IsValid())
+	{
+		ReportError(*FString::Printf(TEXT("UAnimNextRigVMAssetEditorData::AddVariable: Cannot add a variable to this asset - type (%s) is not valid."), *InType.ToString()));
+		return nullptr;
+	}
 
 	// Check for duplicate name
 	FName NewParameterName = InName;
diff --git a/Engine/Source/Runtime/CoreUObject/Private/StructUtils/PropertyBag.cpp b/Engine/Source/Runtime/CoreUObject/Private/StructUtils/PropertyBag.cpp
--- a/Engine/Source/Runtime/CoreUObject/Private/StructUtils/PropertyBag.cpp
+++ b/Engine/Source/Runtime/CoreUObject/Private/StructUtils/PropertyBag.cpp
@@ -4250,6 +4250,10 @@ TArray<FPropertyBagPropertyDesc> UPropertyBag::SanitizePropertyDescs(const TArra
 				InOutTypeObject = UPropertyBagMissingObject::StaticClass();
 			}
 		}
+		else 
+		{
+			InOutTypeObject = nullptr;
+		}	
 	};
 
 	for (FPropertyBagPropertyDesc& Desc : Result)
diff --git a/Engine/Source/Runtime/CoreUObject/Public/StructUtils/PropertyBag.h b/Engine/Source/Runtime/CoreUObject/Public/StructUtils/PropertyBag.h
--- a/Engine/Source/Runtime/CoreUObject/Public/StructUtils/PropertyBag.h
+++ b/Engine/Source/Runtime/CoreUObject/Public/StructUtils/PropertyBag.h
@@ -27,12 +27,12 @@ enum class EPropertyBagPropertyType : uint8
 	Name,
 	String,
 	Text,
-	Enum UMETA(Hidden),
-	Struct UMETA(Hidden),
-	Object UMETA(Hidden),
-	SoftObject UMETA(Hidden),
-	Class UMETA(Hidden),
-	SoftClass UMETA(Hidden),
+	Enum,
+	Struct,
+	Object,
+	SoftObject,
+	Class,
+	SoftClass,
 	UInt32,	// Type not fully supported at UI, will work with restrictions to type editing
 	UInt64, // Type not fully supported at UI, will work with restrictions to type editing
 
