Fixup code analyis issues

diff --git a/Engine/Plugins/Experimental/UAF/UAF/Source/UAFEditor/Private/Variables/Outliner/VariablesOutlinerMode.cpp b/Engine/Plugins/Experimental/UAF/UAF/Source/UAFEditor/Private/Variables/Outliner/VariablesOutlinerMode.cpp
--- a/Engine/Plugins/Experimental/UAF/UAF/Source/UAFEditor/Private/Variables/Outliner/VariablesOutlinerMode.cpp
+++ b/Engine/Plugins/Experimental/UAF/UAF/Source/UAFEditor/Private/Variables/Outliner/VariablesOutlinerMode.cpp
@@ -1633,7 +1633,7 @@ void FVariablesOutlinerMode::PopulateNewVariableToolMenuEntries(UToolMenu* InMen
 					if (UAnimNextRigVMAssetEditorData* EditorData = Pair.Key.Get())
                     {
                         UAnimNextVariableSettings* Settings = GetMutableDefault<UAnimNextVariableSettings>();
-                        FScopedTransaction Transaction(LOCTEXT("AddVariable", "Add variable"));
+                        FScopedTransaction Transaction(LOCTEXT("AddVariable", "Add Variable"));
 
                         UncookedOnly::FCompilationScope CompileScope(LOCTEXT("AddVariable", "Add Variable"), UncookedOnly::FUtils::GetAsset(EditorData));
                         
diff --git a/Engine/Plugins/Experimental/UAF/UAF/Source/UAFUncookedOnly/Private/UAFCompilationScope.cpp b/Engine/Plugins/Experimental/UAF/UAF/Source/UAFUncookedOnly/Private/UAFCompilationScope.cpp
--- a/Engine/Plugins/Experimental/UAF/UAF/Source/UAFUncookedOnly/Private/UAFCompilationScope.cpp
+++ b/Engine/Plugins/Experimental/UAF/UAF/Source/UAFUncookedOnly/Private/UAFCompilationScope.cpp
@@ -101,7 +101,7 @@ FCompilationScope::~FCompilationScope()
 				if (UAnimNextRigVMAsset* Asset = WeakAsset.Get())
 				{
 					UAnimNextRigVMAssetEditorData* EditorData = FUtils::GetEditorData<UAnimNextRigVMAssetEditorData>(Asset);
-					if(EditorData != nullptr)
+					if(EditorData != nullptr && !EditorData->bSuspendCompilationNotifications)
 					{
 						EditorData->DecrementVMRecompileBracket();
 					}
@@ -254,7 +254,7 @@ void FCompilationScope::ProcessAssets(const FText& InJobName, TConstArrayView<UA
 		UAnimNextRigVMAssetEditorData* EditorData = FUtils::GetEditorData<UAnimNextRigVMAssetEditorData>(Asset);
 		if(EditorData == nullptr)
 		{
-			return;
+			continue;
 		}
 
 		if (!ThreadData->CompiledAssets.Contains(Asset))
