Fix UAF graph nodes not prioritizing node annotations by their severity

diff --git a/Engine/Plugins/Experimental/UAF/UAF/Source/UAFUncookedOnly/Private/AnimNextRigVMAssetEditorData.cpp b/Engine/Plugins/Experimental/UAF/UAF/Source/UAFUncookedOnly/Private/AnimNextRigVMAssetEditorData.cpp
--- a/Engine/Plugins/Experimental/UAF/UAF/Source/UAFUncookedOnly/Private/AnimNextRigVMAssetEditorData.cpp
+++ b/Engine/Plugins/Experimental/UAF/UAF/Source/UAFUncookedOnly/Private/AnimNextRigVMAssetEditorData.cpp
@@ -2840,20 +2840,6 @@ void UAnimNextRigVMAssetEditorData::HandleReportFromCompiler(EMessageSeverity::T
 		}
 	}
 
-	if (URigVMEdGraphNode* EdGraphNode = Cast<URigVMEdGraphNode>(SubjectForMessage))
-	{
-		if(Message.IsValid())
-		{
-			EdGraphNode->SetErrorInfo(InSeverity, Message->ToText().ToString());
-		}
-		else
-		{
-			EdGraphNode->SetErrorInfo(InSeverity, InMessage);
-		}
-
-		EdGraphNode->bHasCompilerMessage = EdGraphNode->ErrorType <= int32(EMessageSeverity::Info);
-	}
-
 	if (Message.IsValid())
 	{
 		HandleMessageFromCompiler(Message.ToSharedRef(), false);
