Submit missing file from CL 47490231

diff --git a/Engine/Plugins/Experimental/UAF/UAF/Source/UAFEditor/Private/AnimNextEditorModule.cpp b/Engine/Plugins/Experimental/UAF/UAF/Source/UAFEditor/Private/AnimNextEditorModule.cpp
--- a/Engine/Plugins/Experimental/UAF/UAF/Source/UAFEditor/Private/AnimNextEditorModule.cpp
+++ b/Engine/Plugins/Experimental/UAF/UAF/Source/UAFEditor/Private/AnimNextEditorModule.cpp
@@ -186,10 +186,6 @@ void FAnimNextEditorModule::StartupModule()
 	PropertyModule.RegisterCustomPropertyTypeLayout(
 		"AnimNextVariableReference",
 		FOnGetPropertyTypeCustomizationInstance::CreateLambda([] { return MakeShared<FVariableReferencePropertyCustomization>(); }));
-
-	PropertyModule.RegisterCustomPropertyTypeLayout(
-		"AnimNextSoftVariableReference",
-		FOnGetPropertyTypeCustomizationInstance::CreateLambda([] { return MakeShared<FSoftVariableReferencePropertyCustomization>(); }));
 	
 	ModuleEventPropertyTypeIdentifier = MakeShared<FModuleEventPropertyTypeIdentifier>();
 	PropertyModule.RegisterCustomPropertyTypeLayout(
@@ -319,12 +315,12 @@ void FAnimNextEditorModule::StartupModule()
 								check(SharedVariableNode->Type == EAnimNextSharedVariablesType::Asset && Asset || SharedVariableNode->Type == EAnimNextSharedVariablesType::Struct && Struct);
 
 								const UObject* Object = Asset ? Cast<UObject>(Asset) : Cast<UObject>(Struct);
-								SearchReference = FAnimNextSoftVariableReference(SharedVariableNode->GetVariableName(), Object);
+								SearchReference = FAnimNextSoftVariableReference::FromName(SharedVariableNode->GetVariableName(), Object);
 								if (SharedVariableNode->Type == EAnimNextSharedVariablesType::Asset)
 								{
 									const UAnimNextRigVMAssetEditorData* EditorData = UE::UAF::UncookedOnly::FUtils::GetEditorData<const UAnimNextRigVMAssetEditorData, const UAnimNextSharedVariables>(Asset);						
 									const UAnimNextVariableEntry* Entry = Cast<UAnimNextVariableEntry>(EditorData->FindEntry(SharedVariableNode->GetVariableName()));
-									SearchType = Entry->GetType();
+									SearchType = Entry ? Entry->GetType() : FAnimNextParamType();
 								}
 								else if (SharedVariableNode->Type == EAnimNextSharedVariablesType::Struct)
 								{
@@ -339,9 +335,16 @@ void FAnimNextEditorModule::StartupModule()
 								if (const UAnimNextRigVMAsset* Asset = Cast<UAnimNextRigVMAsset>(VariableNode->GetTypedOuter(UAnimNextRigVMAsset::StaticClass())))
 								{
 									const UAnimNextRigVMAssetEditorData* EditorData = UncookedOnly::FUtils::GetEditorData<const UAnimNextRigVMAssetEditorData, const UAnimNextRigVMAsset>(Asset); 
-									const UAnimNextVariableEntry* Entry = Cast<UAnimNextVariableEntry>(EditorData->FindEntry(VariableNode->GetVariableName()));
-									SearchReference = FAnimNextSoftVariableReference(VariableNode->GetVariableName(), Asset);
-									SearchType = Entry->GetType();
+									if (const UAnimNextVariableEntry* Entry = Cast<UAnimNextVariableEntry>(EditorData->FindEntry(VariableNode->GetVariableName())))
+									{
+										SearchReference = FAnimNextSoftVariableReference::FromName(VariableNode->GetVariableName(), Asset);
+										SearchType = Entry->GetType();
+									}
+									else
+									{
+										SearchReference.Reset();
+										SearchType = FAnimNextParamType();
+									}
 								}
 							}
 
@@ -612,7 +615,6 @@ void FAnimNextEditorModule::ShutdownModule()
 		PropertyModule.UnregisterCustomPropertyTypeLayout("AnimNextParamType");
 		PropertyModule.UnregisterCustomPropertyTypeLayout("AnimNextVariableBinding");
 		PropertyModule.UnregisterCustomPropertyTypeLayout("AnimNextVariableReference");
-		PropertyModule.UnregisterCustomPropertyTypeLayout("AnimNextSoftVariableReference");
 		PropertyModule.UnregisterCustomPropertyTypeLayout("NameProperty", ModuleEventPropertyTypeIdentifier);
 		PropertyModule.UnregisterCustomPropertyTypeLayout("UAFSystemReference");
 		PropertyModule.UnregisterCustomClassLayout("AnimNextVariableEntry");
