Remove dll export from header implemented functions (CIS fix)

diff --git a/Engine/Plugins/Experimental/UAF/UAF/Source/UAF/Public/Variables/AnimNextSoftVariableReference.h b/Engine/Plugins/Experimental/UAF/UAF/Source/UAF/Public/Variables/AnimNextSoftVariableReference.h
--- a/Engine/Plugins/Experimental/UAF/UAF/Source/UAF/Public/Variables/AnimNextSoftVariableReference.h
+++ b/Engine/Plugins/Experimental/UAF/UAF/Source/UAF/Public/Variables/AnimNextSoftVariableReference.h
@@ -14,12 +14,12 @@ struct FAnimNextSoftVariableReference
 	FAnimNextSoftVariableReference() = default;
 
 	// Construct a soft reference from a FName
-	UAF_API static FAnimNextSoftVariableReference FromName(FName InName, const FSoftObjectPath& InSoftObjectPath)
+	static FAnimNextSoftVariableReference FromName(FName InName, const FSoftObjectPath& InSoftObjectPath)
 	{
 		return FAnimNextSoftVariableReference(InName, InSoftObjectPath);
 	}
 
-	UAF_API explicit FAnimNextSoftVariableReference(const FAnimNextVariableReference& VariableReference) : Name(VariableReference.GetName()), SoftObjectPath(VariableReference.GetObject()) {}
+	explicit FAnimNextSoftVariableReference(const FAnimNextVariableReference& VariableReference) : Name(VariableReference.GetName()), SoftObjectPath(VariableReference.GetObject()) {}
 
 	// Get the name of the variable
 	FName GetName() const
