Fix crash on engine shutdown when attempting to use UObjects in UAF

diff --git a/Engine/Plugins/Experimental/UAF/UAF/Source/UAF/Private/Module/SystemReference.cpp b/Engine/Plugins/Experimental/UAF/UAF/Source/UAF/Private/Module/SystemReference.cpp
--- a/Engine/Plugins/Experimental/UAF/UAF/Source/UAF/Private/Module/SystemReference.cpp
+++ b/Engine/Plugins/Experimental/UAF/UAF/Source/UAF/Private/Module/SystemReference.cpp
@@ -1,4 +1,4 @@
-﻿// Copyright Epic Games, Inc. All Rights Reserved.
+// Copyright Epic Games, Inc. All Rights Reserved.
 
 #include "Module/SystemReference.h"
 
@@ -10,6 +10,7 @@
 #include "Factory/SystemFactory.h"
 #include "Factory/UAFSystemFactoryParams.h"
 #include "Logging/StructuredLog.h"
+#include "Misc/CoreDelegates.h"
 #include "Module/AnimNextModuleInstance.h"
 #include "Module/ModuleTaskContext.h"
 #include "Module/SystemDependency.h"
@@ -67,6 +68,15 @@ namespace UE::UAF::Private
 		checkf(GPool, TEXT("UAF system pool is not initialized"));
 		return *GPool;
 	}
+
+	static FDelegateHandle GSystemPoolPreExitHandle;
+
+	static void SystemPoolOnPreExit()
+	{
+		GetPool().FlushPendingActions();
+
+		FCoreDelegates::OnEnginePreExit.Remove(GSystemPoolPreExitHandle);
+	}
 }
 
 namespace UE::UAF
@@ -143,14 +153,16 @@ void FSystemReference::Init()
 {
 	checkf(Private::GPool == nullptr, TEXT("UAF: System pool has already been created"));
 	Private::GPool = new Private::FSystemPool();
+
+	Private::GSystemPoolPreExitHandle = FCoreDelegates::OnEnginePreExit.AddStatic(&Private::SystemPoolOnPreExit);
 }
 
 void FSystemReference::Destroy()
 {
-	UE::UAF::Private::GetPool().FlushPendingActions();
-
 	checkf(Private::GPool != nullptr, TEXT("UAF: System pool has already been destroyed"));
+	checkf(Private::GPool->PendingActions.IsEmpty(), TEXT("UAF: System pool has pending actions that were added after Engine PreExit"));
 	checkf(Private::GPool->SystemInstances.Num() == 0, TEXT("UAF: System pool is being destroyed with %d allocated entries"), Private::GPool->SystemInstances.Num());
+
 	delete Private::GPool;
 	Private::GPool = nullptr;
 }
