Fix UAF variables not being settable from BP

diff --git a/Engine/Plugins/Experimental/UAF/UAF/Source/UAF/Private/Variables/AnimNextVariableReference.cpp b/Engine/Plugins/Experimental/UAF/UAF/Source/UAF/Private/Variables/AnimNextVariableReference.cpp
--- a/Engine/Plugins/Experimental/UAF/UAF/Source/UAF/Private/Variables/AnimNextVariableReference.cpp
+++ b/Engine/Plugins/Experimental/UAF/UAF/Source/UAF/Private/Variables/AnimNextVariableReference.cpp
@@ -64,6 +64,17 @@ FAnimNextVariableReference::FAnimNextVariableReference(const FAnimNextSoftVariab
 	, Object(InSoftReference.GetSoftObjectPath().TryLoad())
 {
 	check(Object == nullptr || Object->IsA<UAnimNextRigVMAsset>() || Object->IsA<UScriptStruct>());
+
+#if WITH_EDITOR
+	if (Object != nullptr)
+	{
+		CachedGuid = UE::UAF::UncookedOnly::IAnimNextUncookedOnlyModule::Get().GetVariableGuidByName(Name, Object);
+		if (!CachedGuid.IsValid())
+		{
+			UE_LOG(LogAnimation, Warning, TEXT("Unable to retrieve Variable Guid while attempting to construct FAnimNextVariableReference for %s in %s, variable might not exist on object."), *Name.ToString(), *Object->GetPathName());
+		}
+	}
+#endif // WITH_EDITOR
 }
 
 FAnimNextVariableReference::FAnimNextVariableReference(FName InName, const UAnimNextRigVMAsset* InAsset)
diff --git a/Engine/Plugins/Experimental/UAF/UAF/Source/UAF/Public/Variables/UAFInstanceVariableData.h b/Engine/Plugins/Experimental/UAF/UAF/Source/UAF/Public/Variables/UAFInstanceVariableData.h
--- a/Engine/Plugins/Experimental/UAF/UAF/Source/UAF/Public/Variables/UAFInstanceVariableData.h
+++ b/Engine/Plugins/Experimental/UAF/UAF/Source/UAF/Public/Variables/UAFInstanceVariableData.h
@@ -137,8 +137,22 @@ private:
 	// Legacy name-based mappings of variables
 	TMap<FName, FVariableMapping> LegacyVariableNameMap;
 
+	// We override the default equality/hash to ignore the GUID since they aren't needed for runtime purposes (editor only)
+	struct FVariableMapKeyFuncs : TDefaultMapHashableKeyFuncs<FAnimNextVariableReference, FVariableMapping, false>
+	{
+		static UE_FORCEINLINE_HINT bool Matches(KeyInitType A, KeyInitType B)
+		{
+			return A.GetName() == B.GetName() && A.GetObject() == B.GetObject();
+		}
+
+		static UE_FORCEINLINE_HINT uint32 GetKeyHash(KeyInitType Key)
+		{
+			return HashCombineFast(GetTypeHash(Key.GetName()), GetTypeHash(Key.GetObject()));
+		}
+	};
+
 	// Map of variable ref -> FVariableMapping
-	TMap<FAnimNextVariableReference, FVariableMapping> VariablesMap;
+	TMap<FAnimNextVariableReference, FVariableMapping, FDefaultSetAllocator, FVariableMapKeyFuncs> VariablesMap;
 
 	// Total count of all variables
 	int32 NumAllVariables = 0;
