UAF: Fixed debug draw smearing

diff --git a/Engine/Plugins/Experimental/UAF/UAF/Source/UAF/Private/AnimNextDebugDraw.cpp b/Engine/Plugins/Experimental/UAF/UAF/Source/UAF/Private/AnimNextDebugDraw.cpp
--- a/Engine/Plugins/Experimental/UAF/UAF/Source/UAF/Private/AnimNextDebugDraw.cpp
+++ b/Engine/Plugins/Experimental/UAF/UAF/Source/UAF/Private/AnimNextDebugDraw.cpp
@@ -40,7 +40,8 @@ void FAnimNextDebugSceneProxy::GetDynamicMeshElements(const TArray<const FSceneV
 		if (VisibilityMap & (1 << ViewIndex))
 		{
 			const FSceneView* View = Views[ViewIndex];
-			FPrimitiveDrawInterface* PDI = Collector.GetPDI(ViewIndex);
+
+			FPrimitiveDrawInterface* PDI = Collector.GetDebugPDI(ViewIndex);
 
 			if (bIsEnabled)
 			{ 
