Various asset wizard fixes
* Fix system asset not using the anim graph asset.
* Fix occasional crash when assigning module to a UAF component.
* Fix not displaying an error when user assets have clashing names.
* Now display error messages whenever the wizard config is invalid.
* Now show each asset class type in the wizard.

#virtualized

diff --git a/Engine/Plugins/Experimental/UAF/UAF/Source/UAFEditor/Private/Wizard/AssetWizard.cpp b/Engine/Plugins/Experimental/UAF/UAF/Source/UAFEditor/Private/Wizard/AssetWizard.cpp
--- a/Engine/Plugins/Experimental/UAF/UAF/Source/UAFEditor/Private/Wizard/AssetWizard.cpp
+++ b/Engine/Plugins/Experimental/UAF/UAF/Source/UAFEditor/Private/Wizard/AssetWizard.cpp
@@ -26,102 +26,94 @@
 #include "TemplateConfig.h"
 #include "ContentBrowserModule.h"
 #include "IContentBrowserSingleton.h"
+#include "UObject/SavePackage.h"
+#include "PackageTools.h"
 
 #define LOCTEXT_NAMESPACE "UE::UAF::Editor::Wizard"
 
 namespace UE::UAF::Editor
 {
-	void FAssetWizard::HandleTemplateSelected(const TObjectPtr<const UUAFTemplateDataAsset> Template, const TObjectPtr<const UUAFTemplateConfig> TemplateConfig)
+	void FAssetWizard::CreateAssets(const TObjectPtr<const UUAFTemplateDataAsset> Template, const TObjectPtr<const UUAFTemplateConfig> TemplateConfig)
 	{
-		struct FAssetPair
-		{
-			TObjectPtr<UObject> TemplateAsset;
-			TObjectPtr<UObject> UserAsset;
-		};
-		TArray<FAssetPair> Assets;
+		FScopedSlowTask SlowTask(4.0f, LOCTEXT("SlowTask", "Creating Assets from Template"));
+		SlowTask.MakeDialog(false);
 
-		bool bFailedToInstantiateTemplateAssets = false;
+		TMap<UObject*, UObject*> TemplateToUserAssetMap;
 
-		for (int32 TemplateAssetIndex = 0; TemplateAssetIndex < Template->Assets.Num(); ++TemplateAssetIndex)
+		// 1. Start by duplicating all the template assets into the user's chosen directory.
 		{
-			TObjectPtr<UObject> TemplateAsset  = Template->Assets[TemplateAssetIndex];
-			
-			FAssetData AssetData(TemplateAsset);
+			SlowTask.EnterProgressFrame(1.0f, LOCTEXT("CreateAssets", "Duplicating template assets"));
 
-			const FString AssetName = TemplateConfig->AssetNaming[TemplateAssetIndex];
-			TObjectPtr<UObject> NewAsset = IAssetTools::Get().DuplicateAsset(AssetName, TemplateConfig->OutputPath.Path, TemplateAsset);
+			bool bFailedToInstantiateTemplateAssets = false;
 			
-			if (NewAsset != nullptr)
+			for (int32 TemplateAssetIndex = 0; TemplateAssetIndex < Template->Assets.Num(); ++TemplateAssetIndex)
 			{
-				FAssetPair& AssetPair = Assets.AddDefaulted_GetRef();
-				AssetPair.TemplateAsset = TemplateAsset;
-				AssetPair.UserAsset = NewAsset;
-			}
-			else
-			{
-				// Failed to instantiate template
-				bFailedToInstantiateTemplateAssets = true;
-				break;
+				TObjectPtr<UObject> TemplateAsset  = Template->Assets[TemplateAssetIndex];
+				
+				const FString UserAssetName = TemplateConfig->AssetNaming[TemplateAssetIndex];
+				
+				if (TObjectPtr<UObject> UserAsset = IAssetTools::Get().DuplicateAsset(UserAssetName, TemplateConfig->OutputPath.Path, TemplateAsset))
+				{
+					TemplateToUserAssetMap.Add(TemplateAsset, UserAsset);
+				}
+				else
+				{
+					UE_LOG(LogAnimation, Warning, TEXT("Asset Wizard: Failed to duplicate asset: %s"), *TemplateAsset->GetFullName());
+					bFailedToInstantiateTemplateAssets = true;
+					break;
+				}
 			}
-		}
 
-		if (bFailedToInstantiateTemplateAssets)
-		{
-			for (FAssetPair& Asset : Assets)
+			if (bFailedToInstantiateTemplateAssets)
 			{
-				ObjectTools::DeleteSingleObject(Asset.UserAsset);
-			}
-
-			UE_LOG(LogAnimation, Warning, TEXT("Failed to instantiate template assets"));
-			return;
-		}
-		
-		for (int32 AssetIndex = 0; AssetIndex < Assets.Num(); ++AssetIndex)
-		{
-			// This asset will have all of its references to any of the template assets updates to point to the new user asset.
-			// Any references inside this asset to other assets not part of the template asset pack will remain the same.
-			TObjectPtr<UObject> AssetToRepair = Assets[AssetIndex].UserAsset;
-
-			TMap<UObject*, UObject*> ReplaceMap;
-			
-			for (int32 OtherAssetIndex = 0; OtherAssetIndex < Assets.Num(); ++OtherAssetIndex)
-			{
-				if (AssetIndex != OtherAssetIndex)
+				for (TPair<UObject*, UObject*> AssetPair : TemplateToUserAssetMap)
 				{
-					FAssetPair OtherAssetPair = Assets[OtherAssetIndex];
-					ReplaceMap.Add(OtherAssetPair.TemplateAsset, OtherAssetPair.UserAsset);
+					UObject* UserAsset = AssetPair.Value;
+					ObjectTools::DeleteSingleObject(UserAsset);
 				}
+				return;
 			}
-			
-			FArchiveReplaceObjectRef<UObject> ReplaceAr(AssetToRepair, ReplaceMap);
-			AssetToRepair->PostEditChange();
 		}
-		
-		if (Template->AssetToOpen)
-		{
-			const int32 AssetIndex = Assets.IndexOfByPredicate([Template](const FAssetPair& InAssetPair){ return InAssetPair.TemplateAsset == Template->AssetToOpen; });
 
-			if (Assets.IsValidIndex(AssetIndex))
+		// 2. Scan for all references within our new user assets that refer to the template assets and replace them with their user asset counterpart.
+		{		
+			SlowTask.EnterProgressFrame(1.0f, LOCTEXT("FixingReferences", "Fixing asset references"));
+
+			for (TPair<UObject*, UObject*> AssetPair : TemplateToUserAssetMap)
 			{
-				TObjectPtr<UObject> AssetToOpen = Assets[AssetIndex].UserAsset;
-				
-				Workspace::IWorkspaceEditorModule& WorkspaceEditorModule = FModuleManager::Get().LoadModuleChecked<Workspace::IWorkspaceEditorModule>("WorkspaceEditor");
-				Workspace::IWorkspaceEditor* WorkspaceEditor = WorkspaceEditorModule.OpenWorkspaceForObject(AssetToOpen, Workspace::EOpenWorkspaceMethod::Default, UAnimNextWorkspaceFactory::StaticClass());
+				UObject* UserAsset = AssetPair.Value;
+
+				FArchiveReplaceObjectRef<UObject> ReplaceAr(UserAsset, TemplateToUserAssetMap);
+				UserAsset->Serialize(ReplaceAr);
+
+				UserAsset->MarkPackageDirty();
+
+				FString PackageFileName = FPackageName::LongPackageNameToFilename(
+					UserAsset->GetPackage()->GetName(),
+					FPackageName::GetAssetPackageExtension()
+				);
 
-				UAnimNextWorkspaceEditorMode* EditorMode = Cast<UAnimNextWorkspaceEditorMode>(WorkspaceEditor->GetEditorModeManager().GetActiveScriptableMode(UAnimNextWorkspaceEditorMode::EM_AnimNextWorkspace));
-				EditorMode->HandleCompile();
+				FSavePackageArgs Args;
+				Args.TopLevelFlags = RF_Public | RF_Standalone;
+				Args.Error = GError;
+
+				UPackage::SavePackage(UserAsset->GetPackage(), UserAsset, *PackageFileName, Args);
 			}
 		}
 
-		if (TemplateConfig->BlueprintMode == ETemplateBlueprintMode::CreateNewBlueprint || TemplateConfig->BlueprintMode == ETemplateBlueprintMode::ModifyExistingBlueprint)
+		// 3. Create a blueprint asset depending on the template config
 		{
-			TObjectPtr<UBlueprint> Blueprint = [TemplateConfig]()
+			SlowTask.EnterProgressFrame(1.0f, LOCTEXT("ConstructingBlueprint", "Constructing blueprint"));
+
+			if (TemplateConfig->BlueprintMode == ETemplateBlueprintMode::CreateNewBlueprint || TemplateConfig->BlueprintMode == ETemplateBlueprintMode::ModifyExistingBlueprint)
 			{
+				UBlueprint* Blueprint = nullptr;
+
 				if (TemplateConfig->BlueprintMode == ETemplateBlueprintMode::CreateNewBlueprint)
 				{
 					TObjectPtr<UBlueprintFactory> BlueprintFactory = NewObject<UBlueprintFactory>();
 					BlueprintFactory->AddToRoot();
-			
+
 					BlueprintFactory->ParentClass = TemplateConfig->BlueprintClass.Get()
 						? TemplateConfig->BlueprintClass.Get()
 						: AActor::StaticClass();
@@ -130,87 +122,114 @@ namespace UE::UAF::Editor
 					IAssetTools& AssetTools = AssetToolsModule.Get();
 
 					TObjectPtr<UObject> NewAsset = AssetTools.CreateAsset(TemplateConfig->BlueprintAssetName, TemplateConfig->OutputPath.Path, UBlueprint::StaticClass(), BlueprintFactory);
+					Blueprint = Cast<UBlueprint>(NewAsset);
 
-					return Cast<UBlueprint>(NewAsset);
+					if (!Blueprint)
+					{
+						UE_LOG(LogAnimation, Error, TEXT("Asset Wizard: Failed to create blueprint"));
+					}
 				}
 				else
 				{
-					return TemplateConfig->BlueprintToModify.LoadSynchronous();
-				}
-			}();
+					Blueprint = TemplateConfig->BlueprintToModify.LoadSynchronous();
 
-			if (!Blueprint)
-			{
-				UE_LOG(LogAnimation, Error, TEXT("UAF asset wizard failed to create/update a blueprint"));
-				return;
-			}
-			
-			TObjectPtr<AActor> Actor = CastChecked<AActor>(Blueprint->GeneratedClass->GetDefaultObject());
+					if (!Blueprint)
+					{
+						UE_LOG(LogAnimation, Error, TEXT("Asset Wizard: Failed to load blueprint to modify"));
+					}
+				}
 
-			USkeletalMeshComponent* SkelMeshComp = Cast<USkeletalMeshComponent>(Actor->GetComponentByClass(USkeletalMeshComponent::StaticClass()));
-			UAnimNextComponent* UAFComp = Cast<UAnimNextComponent>(Actor->GetComponentByClass(UAnimNextComponent::StaticClass()));
+				if (Blueprint)
+				{
+					AActor* Actor = CastChecked<AActor>(Blueprint->GeneratedClass->GetDefaultObject());
 
-			const FAssetPair* ModuleAssetPair = Assets.FindByPredicate([](const FAssetPair& InAssetPair)
-			{
-				return InAssetPair.UserAsset->GetClass() == UAnimNextModule::StaticClass();
-			});
+					USkeletalMeshComponent* SkelMeshComp = Actor->GetComponentByClass<USkeletalMeshComponent>();
+					UAnimNextComponent* UAFComp = Actor->GetComponentByClass<UAnimNextComponent>();
 
-			// Create BP components
-			{
-				USimpleConstructionScript* ConstructionScript = Blueprint->SimpleConstructionScript;
-				USCS_Node* DefaultRootNode = ConstructionScript->GetDefaultSceneRootNode();
-				
-				if (!SkelMeshComp)
-				{
-					// No SkelMesh component exists in this blueprint so create one
-					USCS_Node* SkelMeshCompNode = ConstructionScript->CreateNode(USkeletalMeshComponent::StaticClass(), "Mesh");
-					if (DefaultRootNode)
+					// Create a USkeletalMeshComponent and/or a UAnimNextComponent if either are missing
 					{
-						DefaultRootNode->AddChildNode(SkelMeshCompNode);
+						USimpleConstructionScript* ConstructionScript = Blueprint->SimpleConstructionScript;
+						USCS_Node* DefaultRootNode = ConstructionScript->GetDefaultSceneRootNode();
+
+						if (!SkelMeshComp)
+						{
+							USCS_Node* SkelMeshCompNode = ConstructionScript->CreateNode(USkeletalMeshComponent::StaticClass(), "Mesh");
+							if (DefaultRootNode)
+							{
+								DefaultRootNode->AddChildNode(SkelMeshCompNode);
+							}
+							else
+							{
+								ConstructionScript->AddNode(SkelMeshCompNode);
+							}
+
+							SkelMeshComp = CastChecked<USkeletalMeshComponent>(SkelMeshCompNode->ComponentTemplate);
+						}
+
+						if (!UAFComp)
+						{
+							USCS_Node* UAFCompNode = ConstructionScript->CreateNode(UAnimNextComponent::StaticClass(), "AnimationFramework");
+							ConstructionScript->AddNode(UAFCompNode);
+
+							UAFComp = CastChecked<UAnimNextComponent>(UAFCompNode->ComponentTemplate);
+						}
 					}
-					else
+
+					// Configure components
 					{
-						ConstructionScript->AddNode(SkelMeshCompNode);
+						check(SkelMeshComp && UAFComp);
+
+						SkelMeshComp->SetSkeletalMesh(TemplateConfig->SkeletalMesh.LoadSynchronous());
+						SkelMeshComp->SetEnableAnimation(false);
+
+						FTransform RelativeTransform;
+						RelativeTransform.SetTranslation(FVector(0, 0, -90.0));
+						RelativeTransform.SetRotation(FQuat::MakeFromEuler(FVector(0, 0, -90.0)));
+						SkelMeshComp->SetRelativeTransform(RelativeTransform);
+
+						if (Template->SystemToAssignToComponent)
+						{
+							UObject** UserAssetToAssignPtr = TemplateToUserAssetMap.Find(Template->SystemToAssignToComponent);
+							if (ensureAlwaysMsgf(UserAssetToAssignPtr, TEXT("Asset Wizard: Could not find user system asset counterpart for '%s' to assign to the UAF component. Is the template itself misconfigured?"), *Template->AssetToOpen->GetFullName()))
+							{
+								UAFComp->SetAsset(*UserAssetToAssignPtr);
+							}
+						}
 					}
 
-					SkelMeshComp = CastChecked<USkeletalMeshComponent>(SkelMeshCompNode->ComponentTemplate);
-				}
-
-				if (!UAFComp)
-				{
-					// No AnimNext component exists in this blueprint so create one
-					USCS_Node* UAFCompNode = ConstructionScript->CreateNode(UAnimNextComponent::StaticClass(), "AnimationFramework");
-					ConstructionScript->AddNode(UAFCompNode);
-
-					UAFComp = CastChecked<UAnimNextComponent>(UAFCompNode->ComponentTemplate);
+					FKismetEditorUtilities::CompileBlueprint(Blueprint);
 				}
 			}
+		}
+
+		// 4. Open up one of the newly created assets as determined by the template.
+		{
+			SlowTask.EnterProgressFrame(1.0f, LOCTEXT("Finalising", "Finalising"));
 
-			// Configure BP components
+			if (Template->AssetToOpen)
 			{
-				if (SkelMeshComp)
+				UObject** UserAssetToOpenPtr = TemplateToUserAssetMap.Find(Template->AssetToOpen);
+				if (ensureAlwaysMsgf(UserAssetToOpenPtr, TEXT("Asset Wizard: Could not open user asset counterpart for '%s' to open. Is the template itself misconfigured?"), *Template->AssetToOpen->GetFullName()))
 				{
-					SkelMeshComp->SetSkeletalMesh(TemplateConfig->SkeletalMesh.LoadSynchronous());
-					SkelMeshComp->SetEnableAnimation(false);
-					
-					FTransform RelativeTransform;
-					RelativeTransform.SetTranslation(FVector(0, 0, -90.0));
-					RelativeTransform.SetRotation(FQuat::MakeFromEuler(FVector(0, 0, -90.0)));
-					SkelMeshComp->SetRelativeTransform(RelativeTransform);
-				}
+					Workspace::IWorkspaceEditorModule& WorkspaceEditorModule = FModuleManager::Get().LoadModuleChecked<Workspace::IWorkspaceEditorModule>("WorkspaceEditor");
+					Workspace::IWorkspaceEditor* WorkspaceEditor = WorkspaceEditorModule.OpenWorkspaceForObject(*UserAssetToOpenPtr, Workspace::EOpenWorkspaceMethod::Default, UAnimNextWorkspaceFactory::StaticClass());
 
-				if (UAFComp && ModuleAssetPair && SkelMeshComp && SkelMeshComp->GetSkeletalMeshAsset())
-				{
-					UAFComp->SetAsset(CastChecked<UAnimNextModule>(ModuleAssetPair->UserAsset.Get()));
+					UAnimNextWorkspaceEditorMode* EditorMode = Cast<UAnimNextWorkspaceEditorMode>(WorkspaceEditor->GetEditorModeManager().GetActiveScriptableMode(UAnimNextWorkspaceEditorMode::EM_AnimNextWorkspace));
+					EditorMode->HandleCompile();
 				}
 			}
-
-			FKismetEditorUtilities::CompileBlueprint(Blueprint);
 		}
 	}
 
 	void FAssetWizard::ShowTemplatePicker()
 	{
+		FContentBrowserModule& ContentBrowserModule = FModuleManager::LoadModuleChecked<FContentBrowserModule>("ContentBrowser");
+
+		TArray<FString> SelectedPaths;
+		const FContentBrowserItemPath CurrentPath = ContentBrowserModule.Get().GetCurrentPath();
+
+		GetMutableDefault<UUAFTemplateConfig>()->OutputPath.Path = CurrentPath.HasInternalPath() ? CurrentPath.GetInternalPathString() : TEXT("/Game");
+
 		TSharedRef<SWindow> TemplateWindow = SNew(SWindow)
 			.Title(LOCTEXT("AssetWizard", "UAF Asset Wizard"))
 			.ClientSize(FVector2D(1200, 800))
@@ -221,7 +240,7 @@ namespace UE::UAF::Editor
 			SNew(STemplatePicker)
 			.OnTemplateSelected_Lambda([TemplateWindow](const TObjectPtr<const UUAFTemplateDataAsset> InTemplate, const TObjectPtr<const UUAFTemplateConfig> InTemplateConfig) mutable
 			{
-				HandleTemplateSelected(InTemplate, InTemplateConfig);
+				FAssetWizard::CreateAssets(InTemplate, InTemplateConfig);
 				FSlateApplication::Get().RequestDestroyWindow(TemplateWindow);
 			})
 			.OnCancel_Lambda([TemplateWindow]()
@@ -230,13 +249,6 @@ namespace UE::UAF::Editor
 			})
 		);
 
-		FContentBrowserModule& ContentBrowserModule = FModuleManager::LoadModuleChecked<FContentBrowserModule>("ContentBrowser");
-
-		TArray<FString> SelectedPaths;
-		const FContentBrowserItemPath CurrentPath = ContentBrowserModule.Get().GetCurrentPath();
-
-		GetMutableDefault<UUAFTemplateConfig>()->OutputPath.Path = CurrentPath.HasInternalPath() ? CurrentPath.GetInternalPathString() : TEXT("/Game");
-
 		FSlateApplication::Get().AddWindow(TemplateWindow);
 	}
 	
diff --git a/Engine/Plugins/Experimental/UAF/UAF/Source/UAFEditor/Private/Wizard/AssetWizard.h b/Engine/Plugins/Experimental/UAF/UAF/Source/UAFEditor/Private/Wizard/AssetWizard.h
--- a/Engine/Plugins/Experimental/UAF/UAF/Source/UAFEditor/Private/Wizard/AssetWizard.h
+++ b/Engine/Plugins/Experimental/UAF/UAF/Source/UAFEditor/Private/Wizard/AssetWizard.h
@@ -17,7 +17,7 @@ namespace UE::UAF::Editor
 		
 	private:
 		static void ShowTemplatePicker();
-		
-		static void HandleTemplateSelected(const TObjectPtr<const UUAFTemplateDataAsset> Template, const TObjectPtr<const UUAFTemplateConfig> TemplateConfig);
+
+		static void CreateAssets(const TObjectPtr<const UUAFTemplateDataAsset> Template, const TObjectPtr<const UUAFTemplateConfig> TemplateConfig);
 	};
 }
\ No newline at end of file
diff --git a/Engine/Plugins/Experimental/UAF/UAF/Source/UAFEditor/Private/Wizard/STemplatePicker.cpp b/Engine/Plugins/Experimental/UAF/UAF/Source/UAFEditor/Private/Wizard/STemplatePicker.cpp
--- a/Engine/Plugins/Experimental/UAF/UAF/Source/UAFEditor/Private/Wizard/STemplatePicker.cpp
+++ b/Engine/Plugins/Experimental/UAF/UAF/Source/UAFEditor/Private/Wizard/STemplatePicker.cpp
@@ -292,6 +292,7 @@ namespace UE::UAF::Editor
 							.BodyContent()
 							[
 								SAssignNew(IncludedAssetsListView, SListView<TSharedPtr<FAssetMetadata>>)
+								.SelectionMode(ESelectionMode::None)
 								.ListItemsSource(&IncludedAssets)
 								.OnGenerateRow_Lambda([this](TSharedPtr<FAssetMetadata> InAsset, const TSharedRef<STableViewBase>& InTableView)
 								{
@@ -321,8 +322,20 @@ namespace UE::UAF::Editor
 											.VAlign(VAlign_Center)
 											.Padding(4.0f)
 											[
-												SNew(STextBlock)
-												.Text(FText::FromName(InAsset->AssetData.AssetName))
+												SNew(SVerticalBox)
+												+ SVerticalBox::Slot()
+												.AutoHeight()
+												[
+													SNew(STextBlock)
+													.Text(FText::FromName(InAsset->AssetData.AssetName))
+												]
+												+ SVerticalBox::Slot()
+												.AutoHeight()
+												[
+													SNew(STextBlock)
+													.TextStyle(&FAppStyle::Get(), "SmallText.Subdued")
+													.Text(InAsset->AssetData.GetClass()->GetDisplayNameText())
+												]
 											]
 										]
 										+ SHorizontalBox::Slot()
@@ -339,15 +352,30 @@ namespace UE::UAF::Editor
 											{
 												InAsset->AssetRename = InText.ToString();
 											})
-											.OnVerifyTextChanged_Lambda([this](const FText& InNewText, FText& OutErrorMessage)
+											.OnVerifyTextChanged_Lambda([this, InAsset](const FText& InNewText, FText& OutErrorMessage)
 											{
 												UEditorAssetSubsystem* EditorAssetSubsystem = GEditor->GetEditorSubsystem<UEditorAssetSubsystem>();
 
-												const FString AssetPath = FPaths::Combine( TemplateConfig->OutputPath.Path, InNewText.ToString());
+												const FString AssetPath = FPaths::Combine(TemplateConfig->OutputPath.Path, InNewText.ToString());
 
 												if (EditorAssetSubsystem->DoesAssetExist(AssetPath))
 												{
-													OutErrorMessage = LOCTEXT("AssetAlreadyExists", "An asset with that name already exists");
+													OutErrorMessage = LOCTEXT("AssetAlreadyExists", "Name clashes with an existing asset");
+													return false;
+												}
+
+												for (TSharedPtr<FAssetMetadata> OtherAsset : IncludedAssets)
+												{
+													if (InAsset != OtherAsset && OtherAsset->AssetRename.Equals(InNewText.ToString()))
+													{
+														OutErrorMessage = LOCTEXT("AssetNameClash", "Name clashes with another template asset");
+														return false;
+													}
+												}
+
+												if (TemplateConfig->BlueprintMode == ETemplateBlueprintMode::CreateNewBlueprint && TemplateConfig->BlueprintAssetName.Equals(InNewText.ToString()))
+												{
+													OutErrorMessage = LOCTEXT("AssetNameClashWithBlueprint", "Name clashes with blueprint asset name");
 													return false;
 												}
 												
@@ -387,6 +415,32 @@ namespace UE::UAF::Editor
 					SNew(SHorizontalBox)
 					+ SHorizontalBox::Slot()
 					.HAlign(HAlign_Right)
+					.VAlign(VAlign_Center)
+					.AutoWidth()
+					[
+						SNew(SImage)
+						.Visibility_Lambda([this]()
+						{
+							return ErrorText.IsEmpty() ? EVisibility::Collapsed : EVisibility::Visible;
+						})
+						.ColorAndOpacity(FStyleColors::Error)
+						.Image(FAppStyle::GetBrush("Icons.Error"))
+					]
+					+ SHorizontalBox::Slot()
+					.HAlign(HAlign_Right)
+					.VAlign(VAlign_Center)
+					.AutoWidth()
+					.Padding(FMargin(4.0f, 0.0f, 8.0f, 0.0f))
+					[
+						SNew(STextBlock)
+						.ColorAndOpacity(FStyleColors::Error)
+						.Text_Lambda([this]()
+						{
+							return ErrorText;
+						})
+					]
+					+ SHorizontalBox::Slot()
+					.HAlign(HAlign_Right)
 					.AutoWidth()
 					.Padding(FMargin(16.0f, 16.0f, 8.0f, 16.0f))
 					[
@@ -485,7 +539,7 @@ namespace UE::UAF::Editor
 
 	void STemplatePicker::HandleOutputPathChanged()
 	{
-		RebuildIncludedAssetsListView();
+		UpdateUserAssetNames();
 	}
 
 	void STemplatePicker::HandleBlueprintToModifyChanged()
@@ -503,9 +557,28 @@ namespace UE::UAF::Editor
 		}
 	}
 
-	void STemplatePicker::RebuildIncludedAssetsListView()
+	void STemplatePicker::UpdateUserAssetNames()
 	{
 		UEditorAssetSubsystem* EditorAssetSubsystem = GEditor->GetEditorSubsystem<UEditorAssetSubsystem>();
+
+		const auto GetNewUserAssetName = [&, this](const FString BaseAssetName)
+			{
+				if (!EditorAssetSubsystem->DoesAssetExist(FPaths::Combine(TemplateConfig->OutputPath.Path, BaseAssetName)))
+				{
+					return BaseAssetName;
+				}
+
+				FString OutAssetName = BaseAssetName;
+				int32 Suffix = 1;
+
+				while (EditorAssetSubsystem->DoesAssetExist(FPaths::Combine(TemplateConfig->OutputPath.Path, OutAssetName)))
+				{
+					OutAssetName = BaseAssetName + FString::FromInt(Suffix++);
+				}
+
+				return OutAssetName;
+			};
+
 		
 		IncludedAssets.Empty();
 
@@ -513,30 +586,15 @@ namespace UE::UAF::Editor
 		{
 			check(SelectedTemplate->Template);
 			
+			TemplateConfig->BlueprintAssetName = GetNewUserAssetName(SelectedTemplate->Template->DefaultBlueprintAssetName);
+
 			for (const TObjectPtr<UObject>& Asset : SelectedTemplate->Template->Assets)
 			{
 				TSharedPtr<FAssetMetadata> AssetMetadata = MakeShared<FAssetMetadata>();
 				AssetMetadata->AssetData = FAssetData(Asset);
-				
-				const FString BaseAssetName = AssetMetadata->AssetData.AssetName.ToString();
-				AssetMetadata->AssetRename = BaseAssetName;
-				
-				FString NewAssetPath = FPaths::Combine(TemplateConfig->OutputPath.Path, AssetMetadata->AssetRename);
-				if (EditorAssetSubsystem->DoesAssetExist(NewAssetPath))
-				{
-					int32 Suffix = 1;
-					
-					while (true)
-					{
-						AssetMetadata->AssetRename = BaseAssetName + "_" + FString::FromInt(Suffix++);
-						NewAssetPath = FPaths::Combine(TemplateConfig->OutputPath.Path, AssetMetadata->AssetRename);
 
-						if (!EditorAssetSubsystem->DoesAssetExist(NewAssetPath))
-						{
-							break;
-						}
-					}
-				}
+				const FString BaseAssetName = AssetMetadata->AssetData.AssetName.ToString();
+				AssetMetadata->AssetRename = GetNewUserAssetName(BaseAssetName);
 				
 				IncludedAssets.Add(AssetMetadata);
 			}
@@ -576,7 +634,7 @@ namespace UE::UAF::Editor
 			}
 		}
 		
-		RebuildIncludedAssetsListView();
+		UpdateUserAssetNames();
 	}
 	
 	void STemplatePicker::RefreshModuleTemplates()
@@ -600,6 +658,8 @@ namespace UE::UAF::Editor
 
 	bool STemplatePicker::CanCreateAssets() const
 	{
+		ErrorText = FText::GetEmpty();
+
 		if (!SelectedTemplate.IsValid())
 		{
 			return false;
@@ -608,16 +668,44 @@ namespace UE::UAF::Editor
 		switch (TemplateConfig->BlueprintMode)
 		{
 		case ETemplateBlueprintMode::CreateNewBlueprint:
-			return !TemplateConfig->BlueprintAssetName.IsEmpty()
-				&& TemplateConfig->BlueprintClass.Get()
-				&& !TemplateConfig->SkeletalMesh.IsNull();
+			if (TemplateConfig->BlueprintAssetName.IsEmpty())
+			{
+				ErrorText = LOCTEXT("InvalidBlueprintAssetName", "Invalid blueprint asset name");
+				return false;
+			}
+			if (TemplateConfig->BlueprintClass.Get() == nullptr)
+			{
+				ErrorText = LOCTEXT("InvalidBlueprintClass", "Invalid blueprint class");
+				return false;
+			}
+			if (TemplateConfig->SkeletalMesh.IsNull())
+			{
+				ErrorText = LOCTEXT("InvalidSkeletalMesh", "Invalid skeletal mesh");
+				return false;
+			}
+			break;
 		case ETemplateBlueprintMode::ModifyExistingBlueprint:
-			return !TemplateConfig->BlueprintToModify.IsNull();
-		case ETemplateBlueprintMode::DoNothing:
-			return true;
+			if (TemplateConfig->BlueprintToModify.IsNull())
+			{
+				ErrorText = LOCTEXT("InvalidBlueprintToModify", "Invalid blueprint to modify");
+				return false;
+			}
+			break;
+		}
+
+		if (TemplateConfig->BlueprintMode == ETemplateBlueprintMode::CreateNewBlueprint)
+		{
+			for (TSharedPtr<FAssetMetadata> Asset : IncludedAssets)
+			{
+				if (Asset->AssetRename.Equals(TemplateConfig->BlueprintAssetName))
+				{
+					ErrorText = LOCTEXT("InvalidAssetNames", "Blueprint asset name clashes with template asset name");
+					return false;
+				}
+			}
 		}
 
-		return false;
+		return true;
 	}
 	
 	void SModuleTemplate::Construct(const FArguments& InArgs, const TSharedRef<STableViewBase>& OwnerTable)
diff --git a/Engine/Plugins/Experimental/UAF/UAF/Source/UAFEditor/Private/Wizard/STemplatePicker.h b/Engine/Plugins/Experimental/UAF/UAF/Source/UAFEditor/Private/Wizard/STemplatePicker.h
--- a/Engine/Plugins/Experimental/UAF/UAF/Source/UAFEditor/Private/Wizard/STemplatePicker.h
+++ b/Engine/Plugins/Experimental/UAF/UAF/Source/UAFEditor/Private/Wizard/STemplatePicker.h
@@ -45,7 +45,7 @@ namespace UE::UAF::Editor
 
 		void HandleBlueprintToModifyChanged();
 
-		void RebuildIncludedAssetsListView();
+		void UpdateUserAssetNames();
 
 		bool CanCreateAssets() const;
 		
@@ -85,6 +85,8 @@ namespace UE::UAF::Editor
 		TObjectPtr<UUAFTemplateConfig> TemplateConfig;
 		
 		TSharedPtr<SListView<TSharedPtr<FAssetMetadata>>> IncludedAssetsListView;
+
+		mutable FText ErrorText;
 	};
 
 	class SModuleTemplate : public STableRow<TSharedPtr<STemplatePicker::FTemplate>>
diff --git a/Engine/Plugins/Experimental/UAF/UAF/Source/UAFEditor/Private/Wizard/TemplateDataAsset.h b/Engine/Plugins/Experimental/UAF/UAF/Source/UAFEditor/Private/Wizard/TemplateDataAsset.h
--- a/Engine/Plugins/Experimental/UAF/UAF/Source/UAFEditor/Private/Wizard/TemplateDataAsset.h
+++ b/Engine/Plugins/Experimental/UAF/UAF/Source/UAFEditor/Private/Wizard/TemplateDataAsset.h
@@ -8,6 +8,8 @@
 
 #include "TemplateDataAsset.generated.h"
 
+class UAnimNextModule;
+
 UCLASS()
 class UUAFTemplateDataAsset : public UPrimaryDataAsset
 {
@@ -32,12 +34,18 @@ public:
 	UPROPERTY(EditAnywhere, Category = Default)
 	FSlateBrush DetailsImage;
 
+	// These are the assets that form this template. Each of these will be duplicated into the user's chosen directory (with references patched).
 	UPROPERTY(EditAnywhere, Category = Default)
 	TArray<TObjectPtr<UObject>> Assets;
 
+	// This is the asset that will be opened automatically once the wizard has finished creating the assets for this template.
 	UPROPERTY(EditAnywhere, Category = Default)
 	TObjectPtr<UObject> AssetToOpen;
 
+	// This is the UAF system that we will assign to the UAF component if the user is configuring a blueprint.
+	UPROPERTY(EditAnywhere, Category = Default)
+	TObjectPtr<UAnimNextModule> SystemToAssignToComponent;
+
 	UPROPERTY(EditAnywhere, Category = Default)
 	FString DefaultBlueprintAssetName;
 };
\ No newline at end of file
diff --git a/Engine/Plugins/Experimental/UAF/UAF/Source/UAFUncookedOnly/Internal/AnimNextRigVMAssetEditorData.h b/Engine/Plugins/Experimental/UAF/UAF/Source/UAFUncookedOnly/Internal/AnimNextRigVMAssetEditorData.h
--- a/Engine/Plugins/Experimental/UAF/UAF/Source/UAFUncookedOnly/Internal/AnimNextRigVMAssetEditorData.h
+++ b/Engine/Plugins/Experimental/UAF/UAF/Source/UAFUncookedOnly/Internal/AnimNextRigVMAssetEditorData.h
@@ -291,6 +291,7 @@ protected:
 	virtual bool IsEditorOnly() const override { return true; }
 	UE_API virtual bool Rename(const TCHAR* NewName = nullptr, UObject* NewOuter = nullptr, ERenameFlags Flags = REN_None) override;	
 	UE_API virtual void PreDuplicate(FObjectDuplicationParameters& DupParams) override;
+	UE_API virtual void PostDuplicate(bool bDuplicateForPIE) override;
 
 	UE_API void HandlePackageDone(const FEndLoadPackageContext& Context);
 	UE_API void HandlePackageDone();
diff --git a/Engine/Plugins/Experimental/UAF/UAF/Source/UAFUncookedOnly/Private/AnimNextRigVMAssetEditorData.cpp b/Engine/Plugins/Experimental/UAF/UAF/Source/UAFUncookedOnly/Private/AnimNextRigVMAssetEditorData.cpp
--- a/Engine/Plugins/Experimental/UAF/UAF/Source/UAFUncookedOnly/Private/AnimNextRigVMAssetEditorData.cpp
+++ b/Engine/Plugins/Experimental/UAF/UAF/Source/UAFUncookedOnly/Private/AnimNextRigVMAssetEditorData.cpp
@@ -359,6 +359,14 @@ void UAnimNextRigVMAssetEditorData::PreDuplicate(FObjectDuplicationParameters& D
 	FExternalPackageHelper::DuplicateExternalPackages(this, DupParams);
 }
 
+void UAnimNextRigVMAssetEditorData::PostDuplicate(bool bDuplicateForPIE)
+{
+	Super::PostDuplicate(bDuplicateForPIE);
+
+	ReconstructAllNodes();
+	//RecompileVM();
+}
+
 void UAnimNextRigVMAssetEditorData::HandlePackageDone(const FEndLoadPackageContext& Context)
 {
 	if (!Context.LoadedPackages.Contains(GetPackage()))
