Add CharacterMovementComponent to allowed class types in RigVM

diff --git a/Engine/Plugins/Experimental/UAF/UAF/Source/UAF/Private/AnimNextModuleImpl.cpp b/Engine/Plugins/Experimental/UAF/UAF/Source/UAF/Private/AnimNextModuleImpl.cpp
--- a/Engine/Plugins/Experimental/UAF/UAF/Source/UAF/Private/AnimNextModuleImpl.cpp
+++ b/Engine/Plugins/Experimental/UAF/UAF/Source/UAF/Private/AnimNextModuleImpl.cpp
@@ -38,6 +38,7 @@
 #include "UAF/Attributes/Transformers/Sanitize.h"
 
 #include "Animation/BuiltInAttributeTypes.h"
+#include "GameFramework/CharacterMovementComponent.h"
 #include "UAF/Attributes/EngineAttributes.h"
 
 #define LOCTEXT_NAMESPACE "AnimNextModule"
@@ -64,6 +65,7 @@ namespace UE::UAF
 			{ UAnimNextModule::StaticClass(), FRigVMRegistry::ERegisterObjectOperation::Class },
 			{ UAnimNextRigVMAsset::StaticClass(), FRigVMRegistry::ERegisterObjectOperation::Class },
 			{ UHierarchyTable::StaticClass(), FRigVMRegistry::ERegisterObjectOperation::Class },
+			{ UCharacterMovementComponent::StaticClass(), FRigVMRegistry::ERegisterObjectOperation::Class },
 		};
 
 		RigVMRegistry.RegisterObjectTypes(AllowedObjectTypes);
