Follow-up from CL 47531829, remove temporarily TMap lookup fix - and remove FGuid property from GetTypeHash/=operators as it should only be used for load-time validation and be excluded from any runtime behaviour/code.

diff --git a/Engine/Plugins/Experimental/UAF/UAF/Source/UAF/Public/Variables/AnimNextVariableReference.h b/Engine/Plugins/Experimental/UAF/UAF/Source/UAF/Public/Variables/AnimNextVariableReference.h
--- a/Engine/Plugins/Experimental/UAF/UAF/Source/UAF/Public/Variables/AnimNextVariableReference.h
+++ b/Engine/Plugins/Experimental/UAF/UAF/Source/UAF/Public/Variables/AnimNextVariableReference.h
@@ -89,11 +89,8 @@ struct FAnimNextVariableReference
 
 	bool operator==(const FAnimNextVariableReference& InOther) const
 	{
-		return Name == InOther.Name && Object == InOther.Object
-#if WITH_EDITORONLY_DATA 
-			 && CachedGuid == InOther.CachedGuid
-#endif // WITH_EDITORONLY_DATA
-		;
+		// Exclude FGuid value as its purpose is editor-time validation, and will be filtered out of non-editor serialization
+		return Name == InOther.Name && Object == InOther.Object;
 	}
 
 	bool operator!=(const FAnimNextVariableReference& InOther) const
@@ -103,12 +100,8 @@ struct FAnimNextVariableReference
 
 	friend uint32 GetTypeHash(const FAnimNextVariableReference& InKey)
 	{
-		return HashCombineFast(GetTypeHash(InKey.Name)
-			, GetTypeHash(InKey.Object)
-#if WITH_EDITORONLY_DATA
-			, GetTypeHash(InKey.CachedGuid)
-#endif // WITH_EDITORONLY_DATA
-			);
+		// Exclude FGuid value as its purpose is editor-time validation, and will be filtered out of non-editor serialization
+		return HashCombineFast(GetTypeHash(InKey.Name), GetTypeHash(InKey.Object));
 	}
 
 	// Get the property associated with this variable. Returns null for invalid variables.
diff --git a/Engine/Plugins/Experimental/UAF/UAF/Source/UAF/Public/Variables/UAFInstanceVariableData.h b/Engine/Plugins/Experimental/UAF/UAF/Source/UAF/Public/Variables/UAFInstanceVariableData.h
--- a/Engine/Plugins/Experimental/UAF/UAF/Source/UAF/Public/Variables/UAFInstanceVariableData.h
+++ b/Engine/Plugins/Experimental/UAF/UAF/Source/UAF/Public/Variables/UAFInstanceVariableData.h
@@ -137,22 +137,8 @@ private:
 	// Legacy name-based mappings of variables
 	TMap<FName, FVariableMapping> LegacyVariableNameMap;
 
-	// We override the default equality/hash to ignore the GUID since they aren't needed for runtime purposes (editor only)
-	struct FVariableMapKeyFuncs : TDefaultMapHashableKeyFuncs<FAnimNextVariableReference, FVariableMapping, false>
-	{
-		static UE_FORCEINLINE_HINT bool Matches(KeyInitType A, KeyInitType B)
-		{
-			return A.GetName() == B.GetName() && A.GetObject() == B.GetObject();
-		}
-
-		static UE_FORCEINLINE_HINT uint32 GetKeyHash(KeyInitType Key)
-		{
-			return HashCombineFast(GetTypeHash(Key.GetName()), GetTypeHash(Key.GetObject()));
-		}
-	};
-
 	// Map of variable ref -> FVariableMapping
-	TMap<FAnimNextVariableReference, FVariableMapping, FDefaultSetAllocator, FVariableMapKeyFuncs> VariablesMap;
+	TMap<FAnimNextVariableReference, FVariableMapping> VariablesMap;
 
 	// Total count of all variables
 	int32 NumAllVariables = 0;
@@ -173,4 +159,4 @@ struct TStructOpsTypeTraits<FUAFInstanceVariableData> : public TStructOpsTypeTra
 	{
 		WithAddStructReferencedObjects = true,
 	};
-};
\ No newline at end of file
+};
