Add support for dynamic attributes in abstract skeleton runtime
Improve attribute set map unit tests

diff --git a/Engine/Plugins/Experimental/UAF/UAF/Source/UAF/Private/AnimNextModuleImpl.cpp b/Engine/Plugins/Experimental/UAF/UAF/Source/UAF/Private/AnimNextModuleImpl.cpp
--- a/Engine/Plugins/Experimental/UAF/UAF/Source/UAF/Private/AnimNextModuleImpl.cpp
+++ b/Engine/Plugins/Experimental/UAF/UAF/Source/UAF/Private/AnimNextModuleImpl.cpp
@@ -32,12 +32,6 @@
 #include "RewindDebuggerRuntimeInterface/IRewindDebuggerRuntimeExtension.h"
 #include "RewindDebugger/RewindDebuggerUAFRuntime.h"
 #include "UAF/Attributes/AttributeRuntimeRegistry.h"
-#include "UAF/Attributes/Transformers/Accumulate.h"
-#include "UAF/Attributes/Transformers/Interpolate.h"
-#include "UAF/Attributes/Transformers/Overwrite.h"
-#include "UAF/Attributes/Transformers/Sanitize.h"
-
-#include "Animation/BuiltInAttributeTypes.h"
 #include "GameFramework/CharacterMovementComponent.h"
 #include "UAF/Attributes/EngineAttributes.h"
 
@@ -96,20 +90,7 @@ namespace UE::UAF
 		FSystemFactory::Init();
 
 		FAttributeRuntimeRegistry::Init();
-
-		{
-			FAttributeRuntimeRegistry& RuntimeRegistry = FAttributeRuntimeRegistry::Get();
-
-			RuntimeRegistry.RegisterSetMapInitializer<FFloatAnimationAttribute>();
-			RuntimeRegistry.RegisterSetMapInitializer<FIntegerAnimationAttribute>();
-			RuntimeRegistry.RegisterSetMapInitializer<FStringAnimationAttribute>();
-			RuntimeRegistry.RegisterSetMapInitializer<FBoneTransformAnimationAttribute>();
-
-			RuntimeRegistry.RegisterAttributeValueTransformer<Transformers::FAccumulate_BoneTransformAttribute>();
-			RuntimeRegistry.RegisterAttributeValueTransformer<Transformers::FInterpolate_BoneTransformAttribute>();
-			RuntimeRegistry.RegisterAttributeValueTransformer<Transformers::FOverwrite_BoneTransformAttribute>();
-			RuntimeRegistry.RegisterAttributeValueTransformer<Transformers::FSanitize_BoneTransformAttribute>();
-		}
+		RegisterEngineAttributes();
 
 #if UAF_TRACE_ENABLED
 		IModularFeatures::Get().RegisterModularFeature(IRewindDebuggerRuntimeExtension::ModularFeatureName, &GRewindDebuggerUAFRuntime);
diff --git a/Engine/Plugins/Experimental/UAF/UAF/Source/UAF/Private/DecompressionTools.cpp b/Engine/Plugins/Experimental/UAF/UAF/Source/UAF/Private/DecompressionTools.cpp
--- a/Engine/Plugins/Experimental/UAF/UAF/Source/UAF/Private/DecompressionTools.cpp
+++ b/Engine/Plugins/Experimental/UAF/UAF/Source/UAF/Private/DecompressionTools.cpp
@@ -550,6 +550,11 @@ namespace Private
 		{
 			return Elements[Index];
 		}
+
+		bool IsSorted() const
+		{
+			return bSorted;
+		}
 	};
 }
 
@@ -800,43 +805,62 @@ void FDecompressionTools::GetAnimationCurves(const UAnimSequence* AnimSequence,
 			Private::FBlendedCurveWrapper Curves;
 			CompressedData.CurveCompressionCodec->DecompressCurves(CompressedData, Curves, ExtractionContext.CurrentTime);
 
-			const FAttributeMappingKey MappingKey = FAttributeMappingKey::MakeFromTo<FFloatAnimationAttribute>();
-			TAttributeSetMap<FFloatAnimationAttribute>* FloatAttributes = OutCollection.FindSetMap<FFloatAnimationAttribute>(MappingKey);
-			if (FloatAttributes != nullptr && !FloatAttributes->IsEmpty())
+			const int32 NumDecompressedCurves = Curves.Num();
+			if (NumDecompressedCurves != 0)
 			{
-				// Extract our set defined floats out into our set map
-				const FAttributeTypedSetPtr& FloatSet = FloatAttributes->GetTypedSet();
+				ensureMsgf(Curves.IsSorted(), TEXT("Decompressed curves must be sorted"));
+
+				const FAttributeMappingKey MappingKey = FAttributeMappingKey::MakeFromTo<FFloatAnimationAttribute>();
+				TAttributeSetMap<FFloatAnimationAttribute>* FloatAttributes = OutCollection.FindSetMap<FFloatAnimationAttribute>(MappingKey);
+				TAttributeMap<FFloatAnimationAttribute>* FloatMap = nullptr;
 
-				// TODO: We can perform a tape merge since both inputs are in sorted order (in reverse order since we remove values from our temp curve array)
-				
-				const int32 NumDecompressedCurves = Curves.Num();
-				for (int32 DecompressedCurveIndex = NumDecompressedCurves - 1; DecompressedCurveIndex >= 0; --DecompressedCurveIndex)
+				if (FloatAttributes != nullptr && !FloatAttributes->IsEmpty())
 				{
-					const FName CurveName = Curves[DecompressedCurveIndex].Name;
+					// Extract our set defined floats out into our set map
+					const FAttributeTypedSetPtr& FloatSet = FloatAttributes->GetTypedSet();
 
-					if (const FAttributeBindingIndex BindingIndex = FloatSet->FindBindingIndex(CurveName))
+					// TODO: We can perform a tape merge since both inputs are in sorted order (in reverse order since we remove values from our temp curve array)
+
+					for (int32 DecompressedCurveIndex = 0; DecompressedCurveIndex < NumDecompressedCurves; ++DecompressedCurveIndex)
 					{
-						// This curve is mapped within our binding, it is statically defined
-						if (const FAttributeSetIndex AttributeIndex = FloatSet->GetIndex(BindingIndex))
+						const UE::Anim::FCurveElement& CurveElement = Curves[DecompressedCurveIndex];
+
+						if (const FAttributeBindingIndex BindingIndex = FloatSet->FindBindingIndex(CurveElement.Name))
 						{
-							// This curve is within our current named set, keep it
-							(*FloatAttributes)[AttributeIndex].Value = Curves[DecompressedCurveIndex].Value;
+							// This curve is mapped within our binding, it is statically defined
+							if (const FAttributeSetIndex AttributeIndex = FloatSet->GetIndex(BindingIndex))
+							{
+								// This curve is within our current named set, keep it
+								(*FloatAttributes)[AttributeIndex].Value = CurveElement.Value;
+							}
+							else
+							{
+								// This curve is stripped and isn't within our current named set
+							}
 						}
 						else
 						{
-							// This curve is stripped and isn't within our current named set
+							// This curve isn't mapped within our binding, it is dynamically defined
+							if (FloatMap == nullptr)
+							{
+								FloatMap = OutCollection.AddMap<FFloatAnimationAttribute>();
+							}
+
+							FloatMap->Append(CurveElement.Name, FFloatAnimationAttribute{ CurveElement.Value });
 						}
 					}
-					else
-					{
-						// This curve isn't mapped within our binding, it is dynamically defined
+				}
+				else
+				{
+					// We have no float attributes within our named set, everything is dynamically defined
+					FloatMap = OutCollection.AddMap<FFloatAnimationAttribute>();
 
-						// TODO: Append dynamic attribute
+					for (int32 DecompressedCurveIndex = 0; DecompressedCurveIndex < NumDecompressedCurves; ++DecompressedCurveIndex)
+					{
+						FloatMap->Append(Curves[DecompressedCurveIndex].Name, FFloatAnimationAttribute{ Curves[DecompressedCurveIndex].Value });
 					}
 				}
 			}
-
-			// TODO: Append our dynamic float values (could do it in loop above, one value at a time)
 		}
 	}
 }
@@ -909,8 +933,8 @@ void FDecompressionTools::GetAnimationAttributes(const UAnimSequence* AnimSequen
 			return;
 		}
 
-		FAttributeNamedSetPtr NamedSet = VMContext.GetNamedSet();
-		FAttributeTypedSetPtr BoneTypedSet = NamedSet->FindTypedSet<FBoneTransformAnimationAttribute>();
+		const FAttributeNamedSetPtr& NamedSet = VMContext.GetNamedSet();
+		const FAttributeTypedSetPtr BoneTypedSet = NamedSet->FindTypedSet<FBoneTransformAnimationAttribute>();
 
 		for (const TPair<FAnimationAttributeIdentifier, FAttributeCurve>& BakedAttribute : AnimSequence->AttributeCurves)
 		{
@@ -924,21 +948,20 @@ void FDecompressionTools::GetAnimationAttributes(const UAnimSequence* AnimSequen
 				if (BoneTypedSet && BoneTypedSet->GetIndex(BoneBindingIndex).IsValid())
 				{
 					// Our parent bone is present within the current named set
-					if (FAttributeTypedSetPtr AttributeTypedSet = NamedSet->FindTypedSet(AttributeType))
-					{
-						const FName AttributeName = BakedAttribute.Key.GetName();
+					const FName AttributeName = BakedAttribute.Key.GetName();
 
+					if (const FAttributeTypedSetPtr AttributeTypedSet = NamedSet->FindTypedSet(AttributeType))
+					{
 						if (const FAttributeSetIndex SetIndex = AttributeTypedSet->FindIndex(AttributeName))
 						{
 							// This attribute is present within the current named set
 
 							FAttributeSetMap* SetMap = OutCollection.FindSetMap(FAttributeMappingKey::MakeFromTo(AttributeType));
 							checkf(SetMap != nullptr, TEXT("This attribute exists within our named set but isn't present in the output collection"));
-							checkf(SetMap->GetValueType() == AttributeType, TEXT("Attribute type does not match the set map"));
 
-							SetMap->SetValueWithSetter(SetIndex, [&ExtractionContext, &BakedAttribute](UScriptStruct* SetType, uint8* OutValuePtr)
+							SetMap->SetValueWithSetter(SetIndex, [&ExtractionContext, &BakedAttribute](UScriptStruct* ValueType, uint8* OutValuePtr)
 								{
-									BakedAttribute.Value.EvaluateToPtr(SetType, ExtractionContext.CurrentTime, OutValuePtr);
+									BakedAttribute.Value.EvaluateToPtr(ValueType, ExtractionContext.CurrentTime, OutValuePtr);
 								});
 						}
 						else
@@ -949,8 +972,12 @@ void FDecompressionTools::GetAnimationAttributes(const UAnimSequence* AnimSequen
 					else
 					{
 						// This attribute type is unknown within our binding, it is dynamic
+						FAttributeMap* AttributeMap = OutCollection.FindOrAddMap(AttributeType);
 
-						// TODO: Append dynamic attribute
+						AttributeMap->AddWithSetter(AttributeName, [&ExtractionContext, &BakedAttribute](UScriptStruct* ValueType, uint8* OutValuePtr)
+							{
+								BakedAttribute.Value.EvaluateToPtr(ValueType, ExtractionContext.CurrentTime, OutValuePtr);
+							});
 					}
 				}
 				else
@@ -961,22 +988,22 @@ void FDecompressionTools::GetAnimationAttributes(const UAnimSequence* AnimSequen
 			else
 			{
 				// This attribute isn't attached to a bone
+				// These are not currently supported by AnimBP nor UAF
+#if 0
+				const FName AttributeName = BakedAttribute.Key.GetName();
 
-				if (FAttributeTypedSetPtr AttributeTypedSet = NamedSet->FindTypedSet(AttributeType))
+				if (const FAttributeTypedSetPtr AttributeTypedSet = NamedSet->FindTypedSet(AttributeType))
 				{
-					const FName AttributeName = BakedAttribute.Key.GetName();
-
 					if (const FAttributeSetIndex SetIndex = AttributeTypedSet->FindIndex(AttributeName))
 					{
 						// This attribute is present within the current named set
 
 						FAttributeSetMap* SetMap = OutCollection.FindSetMap(FAttributeMappingKey::MakeFromTo(AttributeType));
 						checkf(SetMap != nullptr, TEXT("This attribute exists within our named set but isn't present in the output collection"));
-						checkf(SetMap->GetValueType() == AttributeType, TEXT("Attribute type does not match the set map"));
 
-						SetMap->SetValueWithSetter(SetIndex, [&ExtractionContext, &BakedAttribute](UScriptStruct* SetType, uint8* OutValuePtr)
+						SetMap->SetValueWithSetter(SetIndex, [&ExtractionContext, &BakedAttribute](UScriptStruct* ValueType, uint8* OutValuePtr)
 							{
-								BakedAttribute.Value.EvaluateToPtr(SetType, ExtractionContext.CurrentTime, OutValuePtr);
+								BakedAttribute.Value.EvaluateToPtr(ValueType, ExtractionContext.CurrentTime, OutValuePtr);
 							});
 					}
 					else
@@ -987,9 +1014,14 @@ void FDecompressionTools::GetAnimationAttributes(const UAnimSequence* AnimSequen
 				else
 				{
 					// This attribute type is unknown within our binding, it is dynamic
+					FAttributeMap* AttributeMap = OutCollection.FindOrAddMap(AttributeType);
 
-					// TODO: Append dynamic attribute
+					AttributeMap->AddWithSetter(AttributeName, [&ExtractionContext, &BakedAttribute](UScriptStruct* ValueType, uint8* OutValuePtr)
+						{
+							BakedAttribute.Value.EvaluateToPtr(ValueType, ExtractionContext.CurrentTime, OutValuePtr);
+						});
 				}
+#endif
 			}
 		}
 	}
diff --git a/Engine/Plugins/Experimental/UAF/UAF/Source/UAF/Private/UAF/Attributes/AttributeCollection.cpp b/Engine/Plugins/Experimental/UAF/UAF/Source/UAF/Private/UAF/Attributes/AttributeCollection.cpp
--- a/Engine/Plugins/Experimental/UAF/UAF/Source/UAF/Private/UAF/Attributes/AttributeCollection.cpp
+++ b/Engine/Plugins/Experimental/UAF/UAF/Source/UAF/Private/UAF/Attributes/AttributeCollection.cpp
@@ -10,7 +10,9 @@ namespace UE::UAF
 	FAttributeCollection::FAttributeCollection(const FAttributeCollection& Other)
 		: NamedSet(Other.NamedSet)
 		, SetMapCount(Other.SetMapCount)
-		, SetMapCapacity(Other.SetMapCapacity)
+		, SetMapCapacity(Other.SetMapCount)
+		, MapCount(Other.MapCount)
+		, MapCapacity(Other.MapCount)
 		, ReallocFun(Other.ReallocFun)
 	{
 		if (Other.SetMaps != nullptr)
@@ -28,6 +30,22 @@ namespace UE::UAF
 		{
 			SetMaps = nullptr;
 		}
+
+		if (Other.Maps != nullptr)
+		{
+			const int32 BufferSize = sizeof(FMapEntry) * Other.MapCount;
+			Maps = reinterpret_cast<FMapEntry*>((*Other.ReallocFun)(nullptr, 0, BufferSize));
+			FMemory::Memcpy(Maps, Other.Maps, BufferSize);
+
+			for (int32 MapIndex = 0; MapIndex < Other.MapCount; ++MapIndex)
+			{
+				Maps[MapIndex].Map = Other.Maps[MapIndex].Map->Duplicate(Other.ReallocFun);
+			}
+		}
+		else
+		{
+			Maps = nullptr;
+		}
 	}
 
 	FAttributeCollection::FAttributeCollection(FAttributeCollection&& Other)
@@ -35,18 +53,21 @@ namespace UE::UAF
 		, SetMaps(Other.SetMaps)
 		, SetMapCount(Other.SetMapCount)
 		, SetMapCapacity(Other.SetMapCapacity)
+		, Maps(Other.Maps)
+		, MapCount(Other.MapCount)
+		, MapCapacity(Other.MapCapacity)
 		, ReallocFun(Other.ReallocFun)
 	{
 		Other.SetMaps = nullptr;
 		Other.SetMapCount = 0;
 		Other.SetMapCapacity = 0;
+		Other.Maps = nullptr;
+		Other.MapCount = 0;
+		Other.MapCapacity = 0;
 	}
 
 	FAttributeCollection::FAttributeCollection(const FAttributeNamedSetPtr& InNamedSet, FReallocFun InReallocFun)
 		: NamedSet(InNamedSet)
-		, SetMaps(nullptr)
-		, SetMapCount(0)
-		, SetMapCapacity(0)
 		, ReallocFun(InReallocFun)
 	{
 		check(InReallocFun != nullptr);
@@ -63,6 +84,19 @@ namespace UE::UAF
 
 			(*ReallocFun)(reinterpret_cast<uint8*>(SetMaps), sizeof(FSetMapEntry) * SetMapCapacity, 0);
 			SetMaps = nullptr;
+			SetMapCount = 0;
+		}
+
+		if (Maps != nullptr)
+		{
+			for (int32 MapIndex = 0; MapIndex < MapCount; ++MapIndex)
+			{
+				ReleaseAttributeMap(Maps[MapIndex].Map);
+			}
+
+			(*ReallocFun)(reinterpret_cast<uint8*>(Maps), sizeof(FMapEntry) * MapCapacity, 0);
+			Maps = nullptr;
+			MapCount = 0;
 		}
 	}
 
@@ -77,7 +111,7 @@ namespace UE::UAF
 
 			NamedSet = Other.NamedSet;
 			SetMapCount = Other.SetMapCount;
-			SetMapCapacity = Other.SetMapCapacity;
+			SetMapCapacity = Other.SetMapCount;
 
 			if (Other.SetMaps != nullptr)
 			{
@@ -94,6 +128,25 @@ namespace UE::UAF
 			{
 				SetMaps = nullptr;
 			}
+
+			MapCount = Other.MapCount;
+			MapCapacity = Other.MapCount;
+
+			if (Other.Maps != nullptr)
+			{
+				const int32 BufferSize = sizeof(FMapEntry) * Other.MapCount;
+				Maps = reinterpret_cast<FMapEntry*>((*ReallocFun)(nullptr, 0, BufferSize));
+				FMemory::Memcpy(Maps, Other.Maps, BufferSize);
+
+				for (int32 MapIndex = 0; MapIndex < Other.MapCount; ++MapIndex)
+				{
+					Maps[MapIndex].Map = Other.Maps[MapIndex].Map->Duplicate(ReallocFun);
+				}
+			}
+			else
+			{
+				Maps = nullptr;
+			}
 		}
 
 		return *this;
@@ -112,10 +165,16 @@ namespace UE::UAF
 			SetMaps = Other.SetMaps;
 			SetMapCount = Other.SetMapCount;
 			SetMapCapacity = Other.SetMapCapacity;
+			Maps = Other.Maps;
+			MapCount = Other.MapCount;
+			MapCapacity = Other.MapCapacity;
 
 			Other.SetMaps = nullptr;
 			Other.SetMapCount = 0;
 			Other.SetMapCapacity = 0;
+			Other.Maps = nullptr;
+			Other.MapCount = 0;
+			Other.MapCapacity = 0;
 		}
 
 		return *this;
@@ -123,7 +182,7 @@ namespace UE::UAF
 
 	bool FAttributeCollection::IsEmpty() const
 	{
-		return SetMapCount == 0;
+		return SetMapCount == 0 && MapCount == 0;
 	}
 
 	const FAttributeNamedSetPtr& FAttributeCollection::GetNamedSet() const
@@ -131,6 +190,29 @@ namespace UE::UAF
 		return NamedSet;
 	}
 
+	void FAttributeCollection::Reset(const FAttributeNamedSetPtr& InNamedSet)
+	{
+		// Free our maps but retain the arrays that contains them
+		for (int32 MapIndex = 0; MapIndex < SetMapCount; ++MapIndex)
+		{
+			ReleaseAttributeSetMap(SetMaps[MapIndex].Map);
+		}
+		SetMapCount = 0;
+
+		for (int32 MapIndex = 0; MapIndex < MapCount; ++MapIndex)
+		{
+			ReleaseAttributeMap(Maps[MapIndex].Map);
+		}
+		MapCount = 0;
+
+		NamedSet = InNamedSet;
+	}
+
+	FReallocFun FAttributeCollection::GetAllocator() const
+	{
+		return ReallocFun;
+	}
+
 	int32 FAttributeCollection::NumSetMaps() const
 	{
 		return SetMapCount;
@@ -166,7 +248,7 @@ namespace UE::UAF
 
 		const FAttributeMappingKey Key(AttributeType, ValueType);
 
-		const int32 MapIndex = Algo::LowerBoundBy(TArrayView<FSetMapEntry>(SetMaps, SetMapCount), Key, [](const FSetMapEntry& Entry) { return Entry.Key; });
+		const int32 MapIndex = Algo::LowerBoundBy(TArrayView<const FSetMapEntry>(SetMaps, SetMapCount), Key, [](const FSetMapEntry& Entry) { return Entry.Key; });
 		if (MapIndex < SetMapCount && SetMaps[MapIndex].Key == Key)
 		{
 			// A map already exists that matches the specified types
@@ -189,8 +271,8 @@ namespace UE::UAF
 		{
 			// The first time we add a set, we reserve space for an entry per typed set plus some extra room
 			const bool bAllowQuantize = true;
-			const int32 DesiredCapasity = SetMapCapacity == 0 ? (NamedSet->NumTypedSets() + 4) : (SetMapCapacity + 1);
-			const int32 NewCapacity = DefaultCalculateSlackGrow(DesiredCapasity, SetMapCapacity, sizeof(FAttributeSetMap*), bAllowQuantize);
+			const int32 DesiredCapacity = SetMapCapacity == 0 ? (NamedSet->NumTypedSets() + 4) : (SetMapCapacity + 1);
+			const int32 NewCapacity = DefaultCalculateSlackGrow(DesiredCapacity, SetMapCapacity, sizeof(FSetMapEntry), bAllowQuantize);
 
 			SetMaps = reinterpret_cast<FSetMapEntry*>((*ReallocFun)(
 				reinterpret_cast<uint8*>(SetMaps),
@@ -209,6 +291,95 @@ namespace UE::UAF
 		return &SetMaps[MapIndex].Map;
 	}
 
+	bool FAttributeCollection::AddSetMap(FAttributeSetMap* Map)
+	{
+		if (Map == nullptr)
+		{
+			// Invalid set map
+			return false;
+		}
+		else if (Map->GetAllocator() != ReallocFun)
+		{
+			// Allocator mismatch
+			return false;
+		}
+
+		const FAttributeMappingKey Key = Map->GetMappingKey();
+
+		const int32 MapIndex = Algo::LowerBoundBy(TArrayView<const FSetMapEntry>(SetMaps, SetMapCount), Key, [](const FSetMapEntry& Entry) { return Entry.Key; });
+		if (MapIndex < SetMapCount && SetMaps[MapIndex].Key == Key)
+		{
+			// A map already exists that matches the specified types
+			return false;
+		}
+
+		if (SetMapCount >= SetMapCapacity)
+		{
+			// The first time we add a set, we reserve space for an entry per typed set plus some extra room
+			const bool bAllowQuantize = true;
+			const int32 DesiredCapacity = SetMapCapacity == 0 ? (NamedSet->NumTypedSets() + 4) : (SetMapCapacity + 1);
+			const int32 NewCapacity = DefaultCalculateSlackGrow(DesiredCapacity, SetMapCapacity, sizeof(FSetMapEntry), bAllowQuantize);
+
+			SetMaps = reinterpret_cast<FSetMapEntry*>((*ReallocFun)(
+				reinterpret_cast<uint8*>(SetMaps),
+				sizeof(FSetMapEntry) * SetMapCapacity,
+				sizeof(FSetMapEntry) * NewCapacity));
+
+			SetMapCapacity = NewCapacity;
+		}
+
+		const int32 NumToMove = SetMapCount - MapIndex;
+		FMemory::Memmove(SetMaps + MapIndex + 1, SetMaps + MapIndex, sizeof(FSetMapEntry) * NumToMove);
+		SetMaps[MapIndex].Key = Key;
+		SetMaps[MapIndex].Map = Map;
+		SetMapCount++;
+
+		return true;
+	}
+
+	bool FAttributeCollection::AppendSetMap(FAttributeSetMap* Map)
+	{
+		if (Map == nullptr)
+		{
+			// Invalid set map
+			return false;
+		}
+		else if (Map->GetAllocator() != ReallocFun)
+		{
+			// Allocator mismatch
+			return false;
+		}
+
+		const FAttributeMappingKey Key = Map->GetMappingKey();
+
+		if (SetMapCount != 0 && !(SetMaps[SetMapCount - 1].Key < Key))
+		{
+			// We don't belong at the end of the set maps
+			return false;
+		}
+
+		if (SetMapCount >= SetMapCapacity)
+		{
+			// The first time we add a set, we reserve space for an entry per typed set plus some extra room
+			const bool bAllowQuantize = true;
+			const int32 DesiredCapacity = SetMapCapacity == 0 ? (NamedSet->NumTypedSets() + 4) : (SetMapCapacity + 1);
+			const int32 NewCapacity = DefaultCalculateSlackGrow(DesiredCapacity, SetMapCapacity, sizeof(FSetMapEntry), bAllowQuantize);
+
+			SetMaps = reinterpret_cast<FSetMapEntry*>((*ReallocFun)(
+				reinterpret_cast<uint8*>(SetMaps),
+				sizeof(FSetMapEntry) * SetMapCapacity,
+				sizeof(FSetMapEntry) * NewCapacity));
+
+			SetMapCapacity = NewCapacity;
+		}
+
+		SetMaps[SetMapCount].Key = Key;
+		SetMaps[SetMapCount].Map = Map;
+		SetMapCount++;
+
+		return true;
+	}
+
 	bool FAttributeCollection::RemoveSetMap(FAttributeSetMap* Map)
 	{
 		if (Map == nullptr)
@@ -222,7 +393,7 @@ namespace UE::UAF
 
 	bool FAttributeCollection::RemoveSetMap(const FAttributeMappingKey& MappingKey)
 	{
-		const int32 MapIndex = Algo::LowerBoundBy(TArrayView<FSetMapEntry>(SetMaps, SetMapCount), MappingKey, [](const FSetMapEntry& Entry) { return Entry.Key; });
+		const int32 MapIndex = Algo::LowerBoundBy(TArrayView<const FSetMapEntry>(SetMaps, SetMapCount), MappingKey, [](const FSetMapEntry& Entry) { return Entry.Key; });
 		if (MapIndex >= SetMapCount || SetMaps[MapIndex].Key != MappingKey)
 		{
 			// No map found that matches the desired types
@@ -242,7 +413,7 @@ namespace UE::UAF
 
 	FAttributeSetMap* FAttributeCollection::FindSetMap(const FAttributeMappingKey& MappingKey)
 	{
-		const int32 MapIndex = Algo::LowerBoundBy(TArrayView<FSetMapEntry>(SetMaps, SetMapCount), MappingKey, [](const FSetMapEntry& Entry) { return Entry.Key; });
+		const int32 MapIndex = Algo::LowerBoundBy(TArrayView<const FSetMapEntry>(SetMaps, SetMapCount), MappingKey, [](const FSetMapEntry& Entry) { return Entry.Key; });
 		if (MapIndex >= SetMapCount || SetMaps[MapIndex].Key != MappingKey)
 		{
 			// No map found that matches the desired types
@@ -254,7 +425,7 @@ namespace UE::UAF
 
 	const FAttributeSetMap* FAttributeCollection::FindSetMap(const FAttributeMappingKey& MappingKey) const
 	{
-		const int32 MapIndex = Algo::LowerBoundBy(TArrayView<FSetMapEntry>(SetMaps, SetMapCount), MappingKey, [](const FSetMapEntry& Entry) { return Entry.Key; });
+		const int32 MapIndex = Algo::LowerBoundBy(TArrayView<const FSetMapEntry>(SetMaps, SetMapCount), MappingKey, [](const FSetMapEntry& Entry) { return Entry.Key; });
 		if (MapIndex >= SetMapCount || SetMaps[MapIndex].Key != MappingKey)
 		{
 			// No map found that matches the desired types
@@ -273,4 +444,218 @@ namespace UE::UAF
 
 		return AddSetMap(MappingKey);
 	}
+
+	int32 FAttributeCollection::NumMaps() const
+	{
+		return MapCount;
+	}
+
+	FAttributeMap* FAttributeCollection::AddMap(UScriptStruct* ValueType)
+	{
+		if (!NamedSet)
+		{
+			// We are not initialized
+			return nullptr;
+		}
+
+		// We don't know the concrete type of the attribute map, use indirection to construct it
+		constexpr bool bConstructMap = true;
+		return *AddMapImpl(ValueType, bConstructMap);
+	}
+
+	FAttributeMap** FAttributeCollection::AddMapImpl(UScriptStruct* ValueType, bool bConstructMap)
+	{
+		const int32 MapIndex = Algo::LowerBoundBy(TArrayView<const FMapEntry>(Maps, MapCount), ValueType, [](const FMapEntry& Entry) { return Entry.ValueType; });
+		if (MapIndex < MapCount && Maps[MapIndex].ValueType == ValueType)
+		{
+			// A map already exists that matches the specified type
+			return nullptr;
+		}
+
+		FAttributeMap* Map = nullptr;
+		if (bConstructMap)
+		{
+			Map = FAttributeRuntimeRegistry::Get().ConstructMap(ValueType, ReallocFun);
+
+			if (Map == nullptr)
+			{
+				// Failed to construct a map, the type has probably not been registered
+				return nullptr;
+			}
+		}
+
+		if (MapCount >= MapCapacity)
+		{
+			const bool bAllowQuantize = true;
+			const int32 DesiredCapacity = MapCapacity + 1;
+			const int32 NewCapacity = DefaultCalculateSlackGrow(DesiredCapacity, MapCapacity, sizeof(FMapEntry), bAllowQuantize);
+
+			Maps = reinterpret_cast<FMapEntry*>((*ReallocFun)(
+				reinterpret_cast<uint8*>(Maps),
+				sizeof(FMapEntry) * MapCapacity,
+				sizeof(FMapEntry) * NewCapacity));
+
+			MapCapacity = NewCapacity;
+		}
+
+		const int32 NumToMove = MapCount - MapIndex;
+		FMemory::Memmove(Maps + MapIndex + 1, Maps + MapIndex, sizeof(FMapEntry) * NumToMove);
+		Maps[MapIndex].ValueType = ValueType;
+		Maps[MapIndex].Map = Map;
+		MapCount++;
+
+		return &Maps[MapIndex].Map;
+	}
+
+	bool FAttributeCollection::AddMap(FAttributeMap* Map)
+	{
+		if (Map == nullptr)
+		{
+			// Invalid set map
+			return false;
+		}
+		else if (Map->GetAllocator() != ReallocFun)
+		{
+			// Allocator mismatch
+			return false;
+		}
+
+		UScriptStruct* ValueType = Map->GetValueType();
+
+		const int32 MapIndex = Algo::LowerBoundBy(TArrayView<const FMapEntry>(Maps, MapCount), ValueType, [](const FMapEntry& Entry) { return Entry.ValueType; });
+		if (MapIndex < MapCount && Maps[MapIndex].ValueType == ValueType)
+		{
+			// A map already exists that matches the specified types
+			return false;
+		}
+
+		if (MapCount >= MapCapacity)
+		{
+			const bool bAllowQuantize = true;
+			const int32 DesiredCapacity = MapCapacity + 1;
+			const int32 NewCapacity = DefaultCalculateSlackGrow(DesiredCapacity, MapCapacity, sizeof(FMapEntry), bAllowQuantize);
+
+			Maps = reinterpret_cast<FMapEntry*>((*ReallocFun)(
+				reinterpret_cast<uint8*>(Maps),
+				sizeof(FMapEntry) * MapCapacity,
+				sizeof(FMapEntry) * NewCapacity));
+
+			MapCapacity = NewCapacity;
+		}
+
+		const int32 NumToMove = MapCount - MapIndex;
+		FMemory::Memmove(Maps + MapIndex + 1, Maps + MapIndex, sizeof(FMapEntry) * NumToMove);
+		Maps[MapIndex].ValueType = ValueType;
+		Maps[MapIndex].Map = Map;
+		MapCount++;
+
+		return true;
+	}
+
+	bool FAttributeCollection::AppendMap(FAttributeMap* Map)
+	{
+		if (Map == nullptr)
+		{
+			// Invalid set map
+			return false;
+		}
+		else if (Map->GetAllocator() != ReallocFun)
+		{
+			// Allocator mismatch
+			return false;
+		}
+
+		UScriptStruct* ValueType = Map->GetValueType();
+
+		if (MapCount != 0 && !(Maps[MapCount - 1].ValueType < ValueType))
+		{
+			// We don't belong at the end of the name/value maps
+			return false;
+		}
+
+		if (MapCount >= MapCapacity)
+		{
+			const bool bAllowQuantize = true;
+			const int32 DesiredCapacity = MapCapacity + 1;
+			const int32 NewCapacity = DefaultCalculateSlackGrow(DesiredCapacity, MapCapacity, sizeof(FMapEntry), bAllowQuantize);
+
+			Maps = reinterpret_cast<FMapEntry*>((*ReallocFun)(
+				reinterpret_cast<uint8*>(Maps),
+				sizeof(FMapEntry) * MapCapacity,
+				sizeof(FMapEntry) * NewCapacity));
+
+			MapCapacity = NewCapacity;
+		}
+
+		Maps[MapCount].ValueType = ValueType;
+		Maps[MapCount].Map = Map;
+		MapCount++;
+
+		return true;
+	}
+
+	bool FAttributeCollection::RemoveMap(FAttributeMap* Map)
+	{
+		if (Map == nullptr)
+		{
+			// Invalid map
+			return false;
+		}
+
+		return RemoveMap(Map->GetValueType());
+	}
+
+	bool FAttributeCollection::RemoveMap(UScriptStruct* ValueType)
+	{
+		const int32 MapIndex = Algo::LowerBoundBy(TArrayView<const FMapEntry>(Maps, MapCount), ValueType, [](const FMapEntry& Entry) { return Entry.ValueType; });
+		if (MapIndex >= MapCount || Maps[MapIndex].ValueType != ValueType)
+		{
+			// No map found that matches the desired types
+			return false;
+		}
+
+		// Destroy and free our map
+		ReleaseAttributeMap(Maps[MapIndex].Map);
+
+		// Remove our entry
+		const int32 NumToMove = MapCount - MapIndex - 1;
+		FMemory::Memmove(Maps + MapIndex, Maps + MapIndex + 1, sizeof(FMapEntry) * NumToMove);
+		MapCount--;
+
+		return true;
+	}
+
+	FAttributeMap* FAttributeCollection::FindMap(UScriptStruct* ValueType)
+	{
+		const int32 MapIndex = Algo::LowerBoundBy(TArrayView<const FMapEntry>(Maps, MapCount), ValueType, [](const FMapEntry& Entry) { return Entry.ValueType; });
+		if (MapIndex >= MapCount || Maps[MapIndex].ValueType != ValueType)
+		{
+			// No map found that matches the desired types
+			return nullptr;
+		}
+
+		return Maps[MapIndex].Map;
+	}
+
+	const FAttributeMap* FAttributeCollection::FindMap(UScriptStruct* ValueType) const
+	{
+		const int32 MapIndex = Algo::LowerBoundBy(TArrayView<const FMapEntry>(Maps, MapCount), ValueType, [](const FMapEntry& Entry) { return Entry.ValueType; });
+		if (MapIndex >= MapCount || Maps[MapIndex].ValueType != ValueType)
+		{
+			// No map found that matches the desired types
+			return nullptr;
+		}
+
+		return Maps[MapIndex].Map;
+	}
+
+	FAttributeMap* FAttributeCollection::FindOrAddMap(UScriptStruct* ValueType)
+	{
+		if (FAttributeMap* Map = FindMap(ValueType))
+		{
+			return Map;
+		}
+
+		return AddMap(ValueType);
+	}
 }
diff --git a/Engine/Plugins/Experimental/UAF/UAF/Source/UAF/Private/UAF/Attributes/AttributeMap.cpp b/Engine/Plugins/Experimental/UAF/UAF/Source/UAF/Private/UAF/Attributes/AttributeMap.cpp
new file mode 100644
--- /dev/null
+++ b/Engine/Plugins/Experimental/UAF/UAF/Source/UAF/Private/UAF/Attributes/AttributeMap.cpp
@@ -0,0 +1,72 @@
+// Copyright Epic Games, Inc. All Rights Reserved.
+
+#include "UAF/Attributes/AttributeMap.h"
+
+#include "UObject/Class.h"
+
+namespace UE::UAF
+{
+	FAttributeMap::FAttributeMap() = default;
+
+	FAttributeMap::FAttributeMap(UScriptStruct* InValueType, FReallocFun InReallocFun, int32 InContainerSize)
+		: ValueTypeStruct(InValueType)
+		, NumValues(0)
+		, MaxValues(0)
+		, DataPtr(nullptr)
+		, ReallocFun(InReallocFun)
+		, ContainerSize(InContainerSize)
+	{
+		check(InReallocFun != nullptr);
+	}
+
+	FAttributeMap::~FAttributeMap()
+	{
+		checkf(DataPtr == nullptr, TEXT("Derived types should handle destruction"));
+	}
+
+	bool FAttributeMap::IsEmpty() const
+	{
+		return NumValues == 0;
+	}
+
+	int32 FAttributeMap::Num() const
+	{
+		return NumValues;
+	}
+
+	int32 FAttributeMap::Max() const
+	{
+		return MaxValues;
+	}
+
+	UScriptStruct* FAttributeMap::GetValueType() const
+	{
+		return ValueTypeStruct;
+	}
+
+	int32 FAttributeMap::GetAllocatedSize() const
+	{
+		const int32 ValueTypeSize = ValueTypeStruct != nullptr ? ValueTypeStruct->GetStructureSize() : 0;
+		const int32 ValueTypeAlignment = ValueTypeSize >= 16 ? 16 : 8;
+		return ContainerSize + Align(sizeof(FName) * MaxValues, ValueTypeAlignment) + (ValueTypeSize * MaxValues);
+	}
+
+	FReallocFun FAttributeMap::GetAllocator() const
+	{
+		return ReallocFun;
+	}
+
+	void FAttributeMap::MoveTo(FAttributeMap& Other)
+	{
+		if (!ensureMsgf(ValueTypeStruct == Other.GetValueType(), TEXT("Value types must match to allow moving")))
+		{
+			return;
+		}
+
+		Swap(NumValues, Other.NumValues);
+		Swap(MaxValues, Other.MaxValues);
+		Swap(DataPtr, Other.DataPtr);
+		Swap(ReallocFun, Other.ReallocFun);
+		Swap(ContainerSize, Other.ContainerSize);
+	}
+}
diff --git a/Engine/Plugins/Experimental/UAF/UAF/Source/UAF/Private/UAF/Attributes/AttributeRuntimeRegistry.cpp b/Engine/Plugins/Experimental/UAF/UAF/Source/UAF/Private/UAF/Attributes/AttributeRuntimeRegistry.cpp
--- a/Engine/Plugins/Experimental/UAF/UAF/Source/UAF/Private/UAF/Attributes/AttributeRuntimeRegistry.cpp
+++ b/Engine/Plugins/Experimental/UAF/UAF/Source/UAF/Private/UAF/Attributes/AttributeRuntimeRegistry.cpp
@@ -48,6 +48,18 @@ namespace UE::UAF
 		return nullptr;
 	}
 
+	FAttributeMap* FAttributeRuntimeRegistry::ConstructMap(UScriptStruct* ValueType, FReallocFun ReallocFun) const
+	{
+		FReadScopeLock Lock(SetMapInitializerLock);
+		if (const MapConstructorFun* ConstructorFun = MapInitializers.Find(ValueType))
+		{
+			return (*ConstructorFun)(ReallocFun);
+		}
+
+		// Value type has not been registered
+		return nullptr;
+	}
+
 	FAttributeValueTransformerMapPtr FAttributeRuntimeRegistry::GetTransformerMap() const
 	{
 		FReadScopeLock Lock(TransformerMapLock);
@@ -60,7 +72,46 @@ namespace UE::UAF
 		SetMapInitializers.Add(ValueType, Initializer);
 	}
 
-	bool FAttributeRuntimeRegistry::RegisterAttributeValueTransformer(FName TransformerName, UScriptStruct* ValueType, FRawTransformerFunc TransformerFunc)
+	void FAttributeRuntimeRegistry::RegisterMapInitializer(UScriptStruct* ValueType, MapConstructorFun Initializer)
+	{
+		FWriteScopeLock Lock(MapInitializerLock);
+		MapInitializers.Add(ValueType, Initializer);
+	}
+
+	bool FAttributeRuntimeRegistry::RegisterAttributeValueSetMapTransformer(FName TransformerName, UScriptStruct* ValueType, FRawTransformerFunc TransformerFunc)
+	{
+		FWriteScopeLock Lock(TransformerMapLock);
+
+		// If someone is already holding a reference to our named operators, we cannot mutate
+		// When this occurs, we make a copy and swap it, keeping the old one alive until it is no longer needed
+		FAttributeValueTransformerMapPtr MutableTransformerMap = TransformerMap;
+		if (!MutableTransformerMap || MutableTransformerMap.GetSharedReferenceCount() != 2)
+		{
+			MutableTransformerMap = MakeShared<typename FAttributeValueTransformerMapPtr::ElementType>();
+		}
+
+		const uint32 TransformerNameHash = GetTypeHash(TransformerName);
+		FAttributeValueTransformerList* TransformerList = MutableTransformerMap->FindByHash(TransformerNameHash, TransformerName);
+		if (TransformerList == nullptr)
+		{
+			TransformerList = &MutableTransformerMap->AddByHash(TransformerNameHash, TransformerName, FAttributeValueTransformerList(TransformerName));
+		}
+
+		if (!TransformerList->AddSetMapTransformer(ValueType, TransformerFunc))
+		{
+			// We failed to add this transformer
+			return false;
+		}
+
+		if (TransformerMap != MutableTransformerMap)
+		{
+			TransformerMap = MutableTransformerMap;
+		}
+
+		return true;
+	}
+
+	bool FAttributeRuntimeRegistry::RegisterAttributeValueMapTransformer(FName TransformerName, UScriptStruct* ValueType, FRawTransformerFunc TransformerFunc)
 	{
 		FWriteScopeLock Lock(TransformerMapLock);
 
@@ -79,7 +130,7 @@ namespace UE::UAF
 			TransformerList = &MutableTransformerMap->AddByHash(TransformerNameHash, TransformerName, FAttributeValueTransformerList(TransformerName));
 		}
 
-		if (!TransformerList->AddTransformer(ValueType, TransformerFunc))
+		if (!TransformerList->AddMapTransformer(ValueType, TransformerFunc))
 		{
 			// We failed to add this transformer
 			return false;
diff --git a/Engine/Plugins/Experimental/UAF/UAF/Source/UAF/Private/UAF/Attributes/AttributeSetMap.cpp b/Engine/Plugins/Experimental/UAF/UAF/Source/UAF/Private/UAF/Attributes/AttributeSetMap.cpp
--- a/Engine/Plugins/Experimental/UAF/UAF/Source/UAF/Private/UAF/Attributes/AttributeSetMap.cpp
+++ b/Engine/Plugins/Experimental/UAF/UAF/Source/UAF/Private/UAF/Attributes/AttributeSetMap.cpp
@@ -57,4 +57,9 @@ namespace UE::UAF
 	{
 		return AllocatedSize;
 	}
+
+	FReallocFun FAttributeSetMap::GetAllocator() const
+	{
+		return ReallocFun;
+	}
 }
diff --git a/Engine/Plugins/Experimental/UAF/UAF/Source/UAF/Private/UAF/Attributes/EngineAttributes.cpp b/Engine/Plugins/Experimental/UAF/UAF/Source/UAF/Private/UAF/Attributes/EngineAttributes.cpp
new file mode 100644
--- /dev/null
+++ b/Engine/Plugins/Experimental/UAF/UAF/Source/UAF/Private/UAF/Attributes/EngineAttributes.cpp
@@ -0,0 +1,39 @@
+// Copyright Epic Games, Inc. All Rights Reserved.
+
+#include "UAF/Attributes/EngineAttributes.h"
+
+#include "UAF/Attributes/AttributeRuntimeRegistry.h"
+#include "UAF/Attributes/Transformers/Accumulate.h"
+#include "UAF/Attributes/Transformers/Interpolate.h"
+#include "UAF/Attributes/Transformers/Overwrite.h"
+#include "UAF/Attributes/Transformers/Sanitize.h"
+
+namespace UE::UAF
+{
+	void RegisterEngineAttributes()
+	{
+		FAttributeRuntimeRegistry& RuntimeRegistry = FAttributeRuntimeRegistry::Get();
+
+		RuntimeRegistry.RegisterSetMapInitializer<FFloatAnimationAttribute>();
+		RuntimeRegistry.RegisterSetMapInitializer<FIntegerAnimationAttribute>();
+		RuntimeRegistry.RegisterSetMapInitializer<FStringAnimationAttribute>();
+		RuntimeRegistry.RegisterSetMapInitializer<FBoneTransformAnimationAttribute>();
+
+		RuntimeRegistry.RegisterMapInitializer<FFloatAnimationAttribute>();
+		RuntimeRegistry.RegisterMapInitializer<FIntegerAnimationAttribute>();
+		RuntimeRegistry.RegisterMapInitializer<FStringAnimationAttribute>();
+		RuntimeRegistry.RegisterMapInitializer<FBoneTransformAnimationAttribute>();
+
+		RuntimeRegistry.RegisterAttributeValueTransformer<Transformers::FAccumulate_BoneTransformAttribute>();
+		RuntimeRegistry.RegisterAttributeValueTransformer<Transformers::FAccumulate_FloatAttribute>();
+		RuntimeRegistry.RegisterAttributeValueTransformer<Transformers::FInterpolate_BoneTransformAttribute>();
+		RuntimeRegistry.RegisterAttributeValueTransformer<Transformers::FInterpolate_FloatAttribute>();
+		RuntimeRegistry.RegisterAttributeValueTransformer<Transformers::FOverwrite_BoneTransformAttribute>();
+		RuntimeRegistry.RegisterAttributeValueTransformer<Transformers::FOverwrite_FloatAttribute>();
+		RuntimeRegistry.RegisterAttributeValueTransformer<Transformers::FSanitize_BoneTransformAttribute>();
+
+#if UE_UAF_VALIDATE_MAPPED_VALUES
+		RuntimeRegistry.RegisterAttributeValueTransformer<Transformers::FSanitize_FloatAttribute>();
+#endif
+	}
+}
diff --git a/Engine/Plugins/Experimental/UAF/UAF/Source/UAF/Private/UAF/Attributes/Transformers/Accumulate.cpp b/Engine/Plugins/Experimental/UAF/UAF/Source/UAF/Private/UAF/Attributes/Transformers/Accumulate.cpp
--- a/Engine/Plugins/Experimental/UAF/UAF/Source/UAF/Private/UAF/Attributes/Transformers/Accumulate.cpp
+++ b/Engine/Plugins/Experimental/UAF/UAF/Source/UAF/Private/UAF/Attributes/Transformers/Accumulate.cpp
@@ -2,28 +2,259 @@
 
 #include "UAF/Attributes/Transformers/Accumulate.h"
 
-#include "UAF/Attributes/AttributeCollection.h"
-#include "UAF/Attributes/AttributeSetMap.h"
 #include "UAF/Attributes/IteratorUtils.h"
 
 namespace UE::UAF::Transformers
 {
 	FName FAccumulate::TransformerName = TEXT("UAF::Transformers::Accumulate");
 
-	void FAccumulate::Apply(const FAttributeValueTransformerMapPtr& TransformerMap, FAttributeCollection& CollectionA, const FAttributeCollection& CollectionB, float Alpha)
+	void FAccumulate::Apply(const FAttributeValueTransformerMapPtr& TransformerMap, const FAttributeCollection& CollectionA, const FAttributeCollection& CollectionB, float Alpha, FAttributeCollection& OutputCollection)
 	{
 		// Named sets must match as it ensures that our inputs have the same sizes/shapes
 		check(CollectionA.GetNamedSet() == CollectionB.GetNamedSet());
 
-		if (const FAttributeValueTransformerList* TransformerList = TransformerMap->Find(TransformerName))
+		const bool bIsOutputA = &CollectionA == &OutputCollection;
+		if (bIsOutputA && !FAnimWeight::IsRelevant(Alpha))
 		{
-			// Keep a list of set maps that don't have a matching pair
-			// These are duplicated afterwards to avoid mutating the collection while iterating
-			TArray<const FAttributeSetMap*, TInlineAllocator<8>> UnmatchedSetMapsB;
+			// Our input A is our output and B is not relevant, nothing to do
+			return;
+		}
+
+		const FAttributeValueTransformerList* TransformerList = TransformerMap->Find(TransformerName);
+		if (TransformerList == nullptr)
+		{
+			// No transformer implementations have been registered
+			return;
+		}
+
+		const bool bIsOutputB = &CollectionB == &OutputCollection;
+		if (bIsOutputA || bIsOutputB)
+		{
+			// Our output collection matches one of our inputs, we'll modify it in place
+			FReallocFun OutputAllocator = OutputCollection.GetAllocator();
+
+			// Set maps
+			{
+				// Keep a list of set maps that needs to be appended to our output
+				TArray<FAttributeSetMap*, TInlineAllocator<8>> PendingSetMaps;
+
+				OuterJoinBy(
+					FAttributeValueTransformerListWithAttributeCollectionSetMapJoinOp(),
+					[Alpha, bIsOutputA, OutputAllocator, &PendingSetMaps](FRawTransformerFunc Transformer, FAttributeSetMap* SetMapA, FAttributeSetMap* SetMapB)
+					{
+						if (SetMapA == nullptr)
+						{
+							if (SetMapB == nullptr)
+							{
+								// No maps found for the transformer type
+							}
+							else
+							{
+								// A doesn't have a matching pair
+								if (bIsOutputA)
+								{
+									// A is our output, duplicate it
+									FAttributeSetMap* SetMapCopy = SetMapB->Duplicate(OutputAllocator);
+									PendingSetMaps.Add(SetMapCopy);
+								}
+								else
+								{
+									// B is our output, nothing to do
+								}
+							}
+						}
+						else
+						{
+							if (SetMapB == nullptr)
+							{
+								// B doesn't have a matching pair
+								if (bIsOutputA)
+								{
+									// A is our output, nothing to do
+								}
+								else
+								{
+									// B is our output, duplicate it
+									FAttributeSetMap* SetMapCopy = SetMapA->Duplicate(OutputAllocator);
+									PendingSetMaps.Add(SetMapCopy);
+								}
+							}
+							else
+							{
+								if (Transformer == nullptr)
+								{
+									// No transformer was found for this value type, retain the one with the largest weight
+									if (Alpha >= 0.5f)
+									{
+										// B has more weight
+										if (bIsOutputA)
+										{
+											// A is our output, copy it over
+											SetMapB->CopyTo(*SetMapA);
+										}
+										else
+										{
+											// B is our output, nothing to do
+										}
+									}
+									else
+									{
+										// A has more weight
+										if (bIsOutputA)
+										{
+											// A is our output, nothing to do
+										}
+										else
+										{
+											// B is our output, copy it over
+											SetMapA->CopyTo(*SetMapB);
+										}
+									}
+								}
+								else
+								{
+									// Transform our matching pair in place
+									FAttributeSetMap* OutputSetMap = bIsOutputA ? SetMapA : SetMapB;
+
+									FTransformSetMapFunc TypedTransformer = reinterpret_cast<FTransformSetMapFunc>(Transformer);
+									(*TypedTransformer)(SetMapA, SetMapB, Alpha, OutputSetMap, nullptr);
+								}
+							}
+						}
+					},
+					TransformerList->CreateSetMapTransformerIterator(),
+					// Remove the const qualifier since we'll mutate one in place as our output
+					const_cast<FAttributeCollection&>(CollectionA).CreateSetMapIterator(),
+					const_cast<FAttributeCollection&>(CollectionB).CreateSetMapIterator());
+
+				// Append our pending entries
+				for (FAttributeSetMap* SetMap : PendingSetMaps)
+				{
+					OutputCollection.AddSetMap(SetMap);
+				}
+			}
+
+			// Name/value maps
+			{
+				// Keep a list of maps that needs to be appended to our output
+				TArray<FAttributeMap*, TInlineAllocator<8>> PendingMaps;
+
+				OuterJoinBy(
+					FAttributeValueTransformerListWithAttributeCollectionMapJoinOp(),
+					[Alpha, bIsOutputA, OutputAllocator, &PendingMaps](FRawTransformerFunc Transformer, FAttributeMap* MapA, FAttributeMap* MapB)
+					{
+						if (MapA == nullptr)
+						{
+							if (MapB == nullptr)
+							{
+								// No maps found for the transformer type
+							}
+							else
+							{
+								// A doesn't have a matching pair
+								if (bIsOutputA)
+								{
+									// A is our output, duplicate it
+									FAttributeMap* MapCopy = MapB->Duplicate(OutputAllocator);
+									PendingMaps.Add(MapCopy);
+								}
+								else
+								{
+									// B is our output, nothing to do
+								}
+							}
+						}
+						else
+						{
+							if (MapB == nullptr)
+							{
+								// B doesn't have a matching pair
+								if (bIsOutputA)
+								{
+									// A is our output, nothing to do
+								}
+								else
+								{
+									// B is our output, duplicate it
+									FAttributeMap* MapCopy = MapA->Duplicate(OutputAllocator);
+									PendingMaps.Add(MapCopy);
+								}
+							}
+							else
+							{
+								if (Transformer == nullptr)
+								{
+									// No transformer was found for this value type, retain the one with the largest weight
+									if (Alpha >= 0.5f)
+									{
+										// B has more weight
+										if (bIsOutputA)
+										{
+											// A is our output, copy it over
+											MapB->CopyTo(*MapA);
+										}
+										else
+										{
+											// B is our output, nothing to do
+										}
+									}
+									else
+									{
+										// A has more weight
+										if (bIsOutputA)
+										{
+											// A is our output, nothing to do
+										}
+										else
+										{
+											// B is our output, copy it over
+											MapA->CopyTo(*MapB);
+										}
+									}
+								}
+								else
+								{
+									// Transform our matching pair in a new map
+									FAttributeMap* OutputMap = nullptr;
+
+									FTransformMapFunc TypedTransformer = reinterpret_cast<FTransformMapFunc>(Transformer);
+									(*TypedTransformer)(MapA, MapB, Alpha, OutputMap, OutputAllocator);
+
+									checkf(OutputMap != nullptr, TEXT("Transformer function is responsible for allocating an output buffer"));
+									checkf(OutputMap->GetAllocator() == OutputAllocator, TEXT("Transformer function is responsible for allocating using the output allocator"));
 
+									// Overwrite our old output
+									OutputMap->MoveTo(bIsOutputA ? *MapA : *MapB);
+
+									// Destroy our temporary output map
+									ReleaseAttributeMap(OutputMap);
+								}
+							}
+						}
+					},
+					TransformerList->CreateMapTransformerIterator(),
+					// Remove the const qualifier since we'll mutate one in place as our output
+					const_cast<FAttributeCollection&>(CollectionA).CreateMapIterator(),
+					const_cast<FAttributeCollection&>(CollectionB).CreateMapIterator());
+
+				// Append our pending entries
+				for (FAttributeMap* Map : PendingMaps)
+				{
+					OutputCollection.AddMap(Map);
+				}
+			}
+		}
+		else
+		{
+			check(CollectionA.GetNamedSet() == OutputCollection.GetNamedSet());
+
+			// Our output collection isn't either of our inputs, reset it
+			OutputCollection.Reset(OutputCollection.GetNamedSet());
+
+			// Set maps
 			OuterJoinBy(
-				FAttributeValueTransformerListWithAttributeCollectionJoinOp(),
-				[Alpha, &UnmatchedSetMapsB](FRawTransformerFunc Transformer, FAttributeSetMap* SetMapA, const FAttributeSetMap* SetMapB)
+				FAttributeValueTransformerListWithAttributeCollectionSetMapJoinOp(),
+				[Alpha, &OutputCollection](FRawTransformerFunc Transformer, const FAttributeSetMap* SetMapA, const FAttributeSetMap* SetMapB)
 				{
 					if (SetMapA == nullptr)
 					{
@@ -33,66 +264,240 @@ namespace UE::UAF::Transformers
 						}
 						else
 						{
-							// This set map doesn't have a matching pair
-							UnmatchedSetMapsB.Add(SetMapB);
+							// This set map doesn't have a matching pair, duplicate it
+							FAttributeSetMap* SetMapCopy = SetMapB->Duplicate(OutputCollection.GetAllocator());
+							OutputCollection.AppendSetMap(SetMapCopy);
 						}
 					}
 					else
 					{
 						if (SetMapB == nullptr)
 						{
-							// This set map doesn't have a matching pair, keep it
+							// This set map doesn't have a matching pair, duplicate it
+							FAttributeSetMap* SetMapCopy = SetMapA->Duplicate(OutputCollection.GetAllocator());
+							OutputCollection.AppendSetMap(SetMapCopy);
 						}
 						else
 						{
 							if (Transformer == nullptr)
 							{
-								// No transformer was found for this value type
+								// No transformer was found for this value type, retain the one with the largest weight
+								FAttributeSetMap* SetMapCopy;
+								if (Alpha >= 0.5f)
+								{
+									SetMapCopy = SetMapB->Duplicate(OutputCollection.GetAllocator());
+								}
+								else
+								{
+									SetMapCopy = SetMapA->Duplicate(OutputCollection.GetAllocator());
+								}
+
+								OutputCollection.AppendSetMap(SetMapCopy);
 							}
 							else
 							{
 								// Transform our matching pair
-								FTransformFunc TypedTransformer = reinterpret_cast<FTransformFunc>(Transformer);
-								(*TypedTransformer)(SetMapA, SetMapB, Alpha);
+								FAttributeSetMap* OutputSetMap = nullptr;
+
+								FTransformSetMapFunc TypedTransformer = reinterpret_cast<FTransformSetMapFunc>(Transformer);
+								(*TypedTransformer)(SetMapA, SetMapB, Alpha, OutputSetMap, OutputCollection.GetAllocator());
+
+								checkf(OutputSetMap != nullptr, TEXT("Transformer function is responsible for allocating an output buffer"));
+								checkf(OutputSetMap->GetAllocator() == OutputCollection.GetAllocator(), TEXT("Transformer function is responsible for allocating using the output allocator"));
+								OutputCollection.AppendSetMap(OutputSetMap);
 							}
 						}
 					}
 				},
-				TransformerList->CreateTransformerIterator(),
-				CollectionA.CreateSetMapIterator(),
+				TransformerList->CreateSetMapTransformerIterator(),
+				CollectionA.CreateConstSetMapIterator(),
 				CollectionB.CreateConstSetMapIterator());
 
-			// Duplicate any unmatched entries
-			for (const FAttributeSetMap* SetMapB : UnmatchedSetMapsB)
-			{
-				const FAttributeMappingKey MappingKey = SetMapB->GetMappingKey();
-				FAttributeSetMap* SetMapA = CollectionA.AddSetMap(MappingKey);
-				checkf(SetMapA != nullptr, TEXT("Failed to add set mapping"));
+			// Name/value maps
+			OuterJoinBy(
+				FAttributeValueTransformerListWithAttributeCollectionMapJoinOp(),
+				[Alpha, &OutputCollection](FRawTransformerFunc Transformer, const FAttributeMap* MapA, const FAttributeMap* MapB)
+				{
+					if (MapA == nullptr)
+					{
+						if (MapB == nullptr)
+						{
+							// No maps found for the transformer type
+						}
+						else
+						{
+							// This map doesn't have a matching pair, duplicate it
+							FAttributeMap* MapCopy = MapB->Duplicate(OutputCollection.GetAllocator());
+							OutputCollection.AppendMap(MapCopy);
+						}
+					}
+					else
+					{
+						if (MapB == nullptr)
+						{
+							// This map doesn't have a matching pair
+							// Our output is A and thus already contains this
+						}
+						else
+						{
+							if (Transformer == nullptr)
+							{
+								// No transformer was found for this value type, retain the one with the largest weight
+								FAttributeMap* MapCopy;
+								if (Alpha >= 0.5f)
+								{
+									MapCopy = MapB->Duplicate(OutputCollection.GetAllocator());
+								}
+								else
+								{
+									MapCopy = MapA->Duplicate(OutputCollection.GetAllocator());
+								}
 
-				SetMapB->CopyTo(*SetMapA);
-			}
+								OutputCollection.AppendMap(MapCopy);
+							}
+							else
+							{
+								// Transform our matching pair in place
+								FAttributeMap* OutputMap = nullptr;
+
+								FTransformMapFunc TypedTransformer = reinterpret_cast<FTransformMapFunc>(Transformer);
+								(*TypedTransformer)(MapA, MapB, Alpha, OutputMap, OutputCollection.GetAllocator());
+
+								checkf(OutputMap != nullptr, TEXT("Transformer function is responsible for allocating an output buffer"));
+								checkf(OutputMap->GetAllocator() == OutputCollection.GetAllocator(), TEXT("Transformer function is responsible for allocating using the output allocator"));
+								OutputCollection.AppendMap(OutputMap);
+							}
+						}
+					}
+				},
+				TransformerList->CreateMapTransformerIterator(),
+				CollectionA.CreateConstMapIterator(),
+				CollectionB.CreateConstMapIterator());
 		}
 	}
 
-	void FAccumulate_BoneTransformAttribute::Transform(UE::UAF::FAttributeSetMap* SetMapA, const UE::UAF::FAttributeSetMap* SetMapB, float Alpha)
+	void FAccumulate_BoneTransformAttribute::TransformSetMap(const FAttributeSetMap* SetMapA, const FAttributeSetMap* SetMapB, float Alpha, FAttributeSetMap*& OutputSetMap, FReallocFun OutputAllocator)
 	{
 		check(SetMapA->GetTypedSet() == SetMapB->GetTypedSet());
+		check(SetMapA->Num() == SetMapB->Num());
+
+		if (OutputSetMap == nullptr)
+		{
+			const FAttributeSetMap::FConstructArgs Args(SetMapA->GetTypedSet(), SetMapA->GetValueType(), OutputAllocator);
+			OutputSetMap = MakeAttributeSetMap<FBoneTransformAnimationAttribute>(Args);
+		}
+
+		check(SetMapA->GetTypedSet() == OutputSetMap->GetTypedSet());
+		check(SetMapA->Num() == OutputSetMap->Num());
 
-		TAttributeSetMap<FBoneTransformAnimationAttribute>* BoneTransformsA = static_cast<TAttributeSetMap<FBoneTransformAnimationAttribute>*>(SetMapA);
-		const TAttributeSetMap<FBoneTransformAnimationAttribute>* BoneTransformsB = static_cast<const TAttributeSetMap<FBoneTransformAnimationAttribute>*>(SetMapB);
+		const TAttributeSetMap<FBoneTransformAnimationAttribute>* BoneTransformsA = CastChecked<FBoneTransformAnimationAttribute>(SetMapA);
+		const TAttributeSetMap<FBoneTransformAnimationAttribute>* BoneTransformsB = CastChecked<FBoneTransformAnimationAttribute>(SetMapB);
+		TAttributeSetMap<FBoneTransformAnimationAttribute>* OutputBoneTransforms = CastChecked<FBoneTransformAnimationAttribute>(OutputSetMap);
 
-		FBoneTransformAnimationAttribute* TransformPtrA = BoneTransformsA->GetData();
+		const FBoneTransformAnimationAttribute* TransformPtrA = BoneTransformsA->GetData();
 		const FBoneTransformAnimationAttribute* TransformPtrB = BoneTransformsB->GetData();
-		const FBoneTransformAnimationAttribute* EndTransformPtrA = TransformPtrA + BoneTransformsA->Num();
+		FBoneTransformAnimationAttribute* OutputTransformPtr = OutputBoneTransforms->GetData();
+		const FBoneTransformAnimationAttribute* TransformPtrAEnd = TransformPtrA + BoneTransformsA->Num();
 
 		const ScalarRegister VAlpha(Alpha);
 
-		while (TransformPtrA < EndTransformPtrA)
+		while (TransformPtrA < TransformPtrAEnd)
 		{
-			TransformPtrA->Value.AccumulateWithShortestRotation(TransformPtrB->Value, VAlpha);
+			OutputTransformPtr->Value = TransformPtrA->Value;
+			OutputTransformPtr->Value.AccumulateWithShortestRotation(TransformPtrB->Value, VAlpha);
 
 			TransformPtrA++;
 			TransformPtrB++;
+			OutputTransformPtr++;
+		}
+	}
+
+	void FAccumulate_BoneTransformAttribute::TransformMap(const FAttributeMap* MapA, const FAttributeMap* MapB, float Alpha, FAttributeMap*& OutputMap, FReallocFun OutputAllocator)
+	{
+		check(OutputMap == nullptr);
+
+		const TAttributeMap<FBoneTransformAnimationAttribute>* BoneTransformsA = CastChecked<FBoneTransformAnimationAttribute>(MapA);
+		const TAttributeMap<FBoneTransformAnimationAttribute>* BoneTransformsB = CastChecked<FBoneTransformAnimationAttribute>(MapB);
+		TAttributeMap<FBoneTransformAnimationAttribute>* OutputBoneTransforms = MakeAttributeMap<FBoneTransformAnimationAttribute>(OutputAllocator);
+
+		const ScalarRegister VAlpha(Alpha);
+
+		OuterJoinBy(
+			FAttributeMapJoinOp<FBoneTransformAnimationAttribute>(),
+			[&VAlpha, OutputBoneTransforms](FName Name, const FBoneTransformAnimationAttribute* A, const FBoneTransformAnimationAttribute* B)
+			{
+				// Missing values are set to the identity
+				const FTransform ValueA = A != nullptr ? A->Value : FTransform::Identity;
+				const FTransform ValueB = B != nullptr ? B->Value : FTransform::Identity;
+
+				FBoneTransformAnimationAttribute Output{ ValueA };
+				Output.Value.AccumulateWithShortestRotation(ValueB, VAlpha);
+
+				OutputBoneTransforms->Append(Name, Output);
+			},
+			BoneTransformsA->CreateConstIterator(),
+			BoneTransformsB->CreateConstIterator());
+
+		OutputMap = OutputBoneTransforms;
+	}
+
+	void FAccumulate_FloatAttribute::TransformSetMap(const FAttributeSetMap* SetMapA, const FAttributeSetMap* SetMapB, float Alpha, FAttributeSetMap*& OutputSetMap, FReallocFun OutputAllocator)
+	{
+		check(SetMapA->GetTypedSet() == SetMapB->GetTypedSet());
+		check(SetMapA->Num() == SetMapB->Num());
+
+		if (OutputSetMap == nullptr)
+		{
+			const FAttributeSetMap::FConstructArgs Args(SetMapA->GetTypedSet(), SetMapA->GetValueType(), OutputAllocator);
+			OutputSetMap = MakeAttributeSetMap<FFloatAnimationAttribute>(Args);
 		}
+
+		check(SetMapA->GetTypedSet() == OutputSetMap->GetTypedSet());
+		check(SetMapA->Num() == OutputSetMap->Num());
+
+		const TAttributeSetMap<FFloatAnimationAttribute>* FloatsA = CastChecked<FFloatAnimationAttribute>(SetMapA);
+		const TAttributeSetMap<FFloatAnimationAttribute>* FloatsB = CastChecked<FFloatAnimationAttribute>(SetMapB);
+		TAttributeSetMap<FFloatAnimationAttribute>* OutputFloats = CastChecked<FFloatAnimationAttribute>(OutputSetMap);
+
+		const FFloatAnimationAttribute* FloatsPtrA = FloatsA->GetData();
+		const FFloatAnimationAttribute* FloatsPtrB = FloatsB->GetData();
+		FFloatAnimationAttribute* OutputFloatsPtr = OutputFloats->GetData();
+		const FFloatAnimationAttribute* FloatsPtrAEnd = FloatsPtrA + FloatsA->Num();
+
+		while (FloatsPtrA < FloatsPtrAEnd)
+		{
+			OutputFloatsPtr->Value = FloatsPtrA->Value + (FloatsPtrB->Value * Alpha);
+
+			FloatsPtrA++;
+			FloatsPtrB++;
+			OutputFloatsPtr++;
+		}
+	}
+
+	void FAccumulate_FloatAttribute::TransformMap(const FAttributeMap* MapA, const FAttributeMap* MapB, float Alpha, FAttributeMap*& OutputMap, FReallocFun OutputAllocator)
+	{
+		check(OutputMap == nullptr);
+
+		const TAttributeMap<FFloatAnimationAttribute>* FloatsA = CastChecked<FFloatAnimationAttribute>(MapA);
+		const TAttributeMap<FFloatAnimationAttribute>* FloatsB = CastChecked<FFloatAnimationAttribute>(MapB);
+		TAttributeMap<FFloatAnimationAttribute>* OutputFloats = MakeAttributeMap<FFloatAnimationAttribute>(OutputAllocator);
+
+		OuterJoinBy(
+			FAttributeMapJoinOp<FFloatAnimationAttribute>(),
+			[&Alpha, OutputFloats](FName Name, const FFloatAnimationAttribute* A, const FFloatAnimationAttribute* B)
+			{
+				// Missing values are set to 0.0
+				const float ValueA = A != nullptr ? A->Value : 0.0f;
+				const float ValueB = B != nullptr ? B->Value : 0.0f;
+
+				FFloatAnimationAttribute Output;
+				Output.Value = ValueA + (ValueB * Alpha);
+
+				OutputFloats->Append(Name, Output);
+			},
+			FloatsA->CreateConstIterator(),
+			FloatsB->CreateConstIterator());
+
+		OutputMap = OutputFloats;
 	}
 }
diff --git a/Engine/Plugins/Experimental/UAF/UAF/Source/UAF/Private/UAF/Attributes/Transformers/Interpolate.cpp b/Engine/Plugins/Experimental/UAF/UAF/Source/UAF/Private/UAF/Attributes/Transformers/Interpolate.cpp
--- a/Engine/Plugins/Experimental/UAF/UAF/Source/UAF/Private/UAF/Attributes/Transformers/Interpolate.cpp
+++ b/Engine/Plugins/Experimental/UAF/UAF/Source/UAF/Private/UAF/Attributes/Transformers/Interpolate.cpp
@@ -2,28 +2,265 @@
 
 #include "UAF/Attributes/Transformers/Interpolate.h"
 
-#include "UAF/Attributes/AttributeCollection.h"
-#include "UAF/Attributes/AttributeSetMap.h"
 #include "UAF/Attributes/IteratorUtils.h"
 
 namespace UE::UAF::Transformers
 {
 	FName FInterpolate::TransformerName = TEXT("UAF::Transformers::Interpolate");
 
-	void FInterpolate::Apply(const FAttributeValueTransformerMapPtr& TransformerMap, FAttributeCollection& CollectionA, const FAttributeCollection& CollectionB, float Alpha)
+	void FInterpolate::Apply(const FAttributeValueTransformerMapPtr& TransformerMap, const FAttributeCollection& CollectionA, const FAttributeCollection& CollectionB, float Alpha, FAttributeCollection& OutputCollection)
 	{
 		// Named sets must match as it ensures that our inputs have the same sizes/shapes
 		check(CollectionA.GetNamedSet() == CollectionB.GetNamedSet());
 
-		if (const FAttributeValueTransformerList* TransformerList = TransformerMap->Find(TransformerName))
+		const bool bIsOutputA = &CollectionA == &OutputCollection;
+		if (bIsOutputA && !FAnimWeight::IsRelevant(Alpha))
 		{
-			// Keep a list of set maps that don't have a matching pair
-			// These are duplicated afterwards to avoid mutating the collection while iterating
-			TArray<const FAttributeSetMap*, TInlineAllocator<8>> UnmatchedSetMapsB;
+			// Our input A is our output and B is not relevant, nothing to do
+			return;
+		}
+
+		const bool bIsOutputB = &CollectionB == &OutputCollection;
+		if (bIsOutputB && FAnimWeight::IsFullWeight(Alpha))
+		{
+			// Our input B is our output and A is not relevant, nothing to do
+			return;
+		}
+
+		const FAttributeValueTransformerList* TransformerList = TransformerMap->Find(TransformerName);
+		if (TransformerList == nullptr)
+		{
+			// No transformer implementations have been registered
+			return;
+		}
+
+		if (bIsOutputA || bIsOutputB)
+		{
+			// Our output collection matches one of our inputs, we'll modify it in place
+			FReallocFun OutputAllocator = OutputCollection.GetAllocator();
 
+			// Set maps
+			{
+				// Keep a list of set maps that needs to be appended to our output
+				TArray<FAttributeSetMap*, TInlineAllocator<8>> PendingSetMaps;
+
+				OuterJoinBy(
+					FAttributeValueTransformerListWithAttributeCollectionSetMapJoinOp(),
+					[Alpha, bIsOutputA, OutputAllocator, &PendingSetMaps](FRawTransformerFunc Transformer, FAttributeSetMap* SetMapA, FAttributeSetMap* SetMapB)
+					{
+						if (SetMapA == nullptr)
+						{
+							if (SetMapB == nullptr)
+							{
+								// No maps found for the transformer type
+							}
+							else
+							{
+								// A doesn't have a matching pair
+								if (bIsOutputA)
+								{
+									// A is our output, duplicate it
+									FAttributeSetMap* SetMapCopy = SetMapB->Duplicate(OutputAllocator);
+									PendingSetMaps.Add(SetMapCopy);
+								}
+								else
+								{
+									// B is our output, nothing to do
+								}
+							}
+						}
+						else
+						{
+							if (SetMapB == nullptr)
+							{
+								// B doesn't have a matching pair
+								if (bIsOutputA)
+								{
+									// A is our output, nothing to do
+								}
+								else
+								{
+									// B is our output, duplicate it
+									FAttributeSetMap* SetMapCopy = SetMapA->Duplicate(OutputAllocator);
+									PendingSetMaps.Add(SetMapCopy);
+								}
+							}
+							else
+							{
+								if (Transformer == nullptr)
+								{
+									// No transformer was found for this value type, retain the one with the largest weight
+									if (Alpha >= 0.5f)
+									{
+										// B has more weight
+										if (bIsOutputA)
+										{
+											// A is our output, copy it over
+											SetMapB->CopyTo(*SetMapA);
+										}
+										else
+										{
+											// B is our output, nothing to do
+										}
+									}
+									else
+									{
+										// A has more weight
+										if (bIsOutputA)
+										{
+											// A is our output, nothing to do
+										}
+										else
+										{
+											// B is our output, copy it over
+											SetMapA->CopyTo(*SetMapB);
+										}
+									}
+								}
+								else
+								{
+									// Transform our matching pair in place
+									FAttributeSetMap* OutputSetMap = bIsOutputA ? SetMapA : SetMapB;
+
+									FTransformSetMapFunc TypedTransformer = reinterpret_cast<FTransformSetMapFunc>(Transformer);
+									(*TypedTransformer)(SetMapA, SetMapB, Alpha, OutputSetMap, nullptr);
+								}
+							}
+						}
+					},
+					TransformerList->CreateSetMapTransformerIterator(),
+					// Remove the const qualifier since we'll mutate one in place as our output
+					const_cast<FAttributeCollection&>(CollectionA).CreateSetMapIterator(),
+					const_cast<FAttributeCollection&>(CollectionB).CreateSetMapIterator());
+
+				// Append our pending entries
+				for (FAttributeSetMap* SetMap : PendingSetMaps)
+				{
+					OutputCollection.AddSetMap(SetMap);
+				}
+			}
+
+			// Name/value maps
+			{
+				// Keep a list of maps that needs to be appended to our output
+				TArray<FAttributeMap*, TInlineAllocator<8>> PendingMaps;
+
+				OuterJoinBy(
+					FAttributeValueTransformerListWithAttributeCollectionMapJoinOp(),
+					[Alpha, bIsOutputA, OutputAllocator, &PendingMaps](FRawTransformerFunc Transformer, FAttributeMap* MapA, FAttributeMap* MapB)
+					{
+						if (MapA == nullptr)
+						{
+							if (MapB == nullptr)
+							{
+								// No maps found for the transformer type
+							}
+							else
+							{
+								// A doesn't have a matching pair
+								if (bIsOutputA)
+								{
+									// A is our output, duplicate it
+									FAttributeMap* MapCopy = MapB->Duplicate(OutputAllocator);
+									PendingMaps.Add(MapCopy);
+								}
+								else
+								{
+									// B is our output, nothing to do
+								}
+							}
+						}
+						else
+						{
+							if (MapB == nullptr)
+							{
+								// B doesn't have a matching pair
+								if (bIsOutputA)
+								{
+									// A is our output, nothing to do
+								}
+								else
+								{
+									// B is our output, duplicate it
+									FAttributeMap* MapCopy = MapA->Duplicate(OutputAllocator);
+									PendingMaps.Add(MapCopy);
+								}
+							}
+							else
+							{
+								if (Transformer == nullptr)
+								{
+									// No transformer was found for this value type, retain the one with the largest weight
+									if (Alpha >= 0.5f)
+									{
+										// B has more weight
+										if (bIsOutputA)
+										{
+											// A is our output, copy it over
+											MapB->CopyTo(*MapA);
+										}
+										else
+										{
+											// B is our output, nothing to do
+										}
+									}
+									else
+									{
+										// A has more weight
+										if (bIsOutputA)
+										{
+											// A is our output, nothing to do
+										}
+										else
+										{
+											// B is our output, copy it over
+											MapA->CopyTo(*MapB);
+										}
+									}
+								}
+								else
+								{
+									// Transform our matching pair in a new map
+									FAttributeMap* OutputMap = nullptr;
+
+									FTransformMapFunc TypedTransformer = reinterpret_cast<FTransformMapFunc>(Transformer);
+									(*TypedTransformer)(MapA, MapB, Alpha, OutputMap, OutputAllocator);
+
+									checkf(OutputMap != nullptr, TEXT("Transformer function is responsible for allocating an output buffer"));
+									checkf(OutputMap->GetAllocator() == OutputAllocator, TEXT("Transformer function is responsible for allocating using the output allocator"));
+
+									// Overwrite our old output
+									OutputMap->MoveTo(bIsOutputA ? *MapA : *MapB);
+
+									// Destroy our temporary output map
+									ReleaseAttributeMap(OutputMap);
+								}
+							}
+						}
+					},
+					TransformerList->CreateMapTransformerIterator(),
+					// Remove the const qualifier since we'll mutate one in place as our output
+					const_cast<FAttributeCollection&>(CollectionA).CreateMapIterator(),
+					const_cast<FAttributeCollection&>(CollectionB).CreateMapIterator());
+
+				// Append our pending entries
+				for (FAttributeMap* Map : PendingMaps)
+				{
+					OutputCollection.AddMap(Map);
+				}
+			}
+		}
+		else
+		{
+			check(CollectionA.GetNamedSet() == OutputCollection.GetNamedSet());
+
+			// Our output collection isn't either of our inputs, reset it
+			OutputCollection.Reset(OutputCollection.GetNamedSet());
+
+			// Set maps
 			OuterJoinBy(
-				FAttributeValueTransformerListWithAttributeCollectionJoinOp(),
-				[Alpha, &UnmatchedSetMapsB](FRawTransformerFunc Transformer, FAttributeSetMap* SetMapA, const FAttributeSetMap* SetMapB)
+				FAttributeValueTransformerListWithAttributeCollectionSetMapJoinOp(),
+				[Alpha, &OutputCollection](FRawTransformerFunc Transformer, const FAttributeSetMap* SetMapA, const FAttributeSetMap* SetMapB)
 				{
 					if (SetMapA == nullptr)
 					{
@@ -33,64 +270,235 @@ namespace UE::UAF::Transformers
 						}
 						else
 						{
-							// This set map doesn't have a matching pair
-							UnmatchedSetMapsB.Add(SetMapB);
+							// This set map doesn't have a matching pair, duplicate it
+							FAttributeSetMap* SetMapCopy = SetMapB->Duplicate(OutputCollection.GetAllocator());
+							OutputCollection.AppendSetMap(SetMapCopy);
 						}
 					}
 					else
 					{
 						if (SetMapB == nullptr)
 						{
-							// This set map doesn't have a matching pair, keep it
+							// This set map doesn't have a matching pair, duplicate it
+							FAttributeSetMap* SetMapCopy = SetMapA->Duplicate(OutputCollection.GetAllocator());
+							OutputCollection.AppendSetMap(SetMapCopy);
 						}
 						else
 						{
 							if (Transformer == nullptr)
 							{
-								// No transformer was found for this value type
+								// No transformer was found for this value type, retain the one with the largest weight
+								FAttributeSetMap* SetMapCopy;
+								if (Alpha >= 0.5f)
+								{
+									SetMapCopy = SetMapB->Duplicate(OutputCollection.GetAllocator());
+								}
+								else
+								{
+									SetMapCopy = SetMapA->Duplicate(OutputCollection.GetAllocator());
+								}
+
+								OutputCollection.AppendSetMap(SetMapCopy);
 							}
 							else
 							{
 								// Transform our matching pair
-								FTransformFunc TypedTransformer = reinterpret_cast<FTransformFunc>(Transformer);
-								(*TypedTransformer)(SetMapA, SetMapB, Alpha);
+								FAttributeSetMap* OutputSetMap = nullptr;
+
+								FTransformSetMapFunc TypedTransformer = reinterpret_cast<FTransformSetMapFunc>(Transformer);
+								(*TypedTransformer)(SetMapA, SetMapB, Alpha, OutputSetMap, OutputCollection.GetAllocator());
+
+								checkf(OutputSetMap != nullptr, TEXT("Transformer function is responsible for allocating an output buffer"));
+								checkf(OutputSetMap->GetAllocator() == OutputCollection.GetAllocator(), TEXT("Transformer function is responsible for allocating using the output allocator"));
+								OutputCollection.AppendSetMap(OutputSetMap);
 							}
 						}
 					}
 				},
-				TransformerList->CreateTransformerIterator(),
-				CollectionA.CreateSetMapIterator(),
+				TransformerList->CreateSetMapTransformerIterator(),
+				CollectionA.CreateConstSetMapIterator(),
 				CollectionB.CreateConstSetMapIterator());
 
-			// Duplicate any unmatched entries
-			for (const FAttributeSetMap* SetMapB : UnmatchedSetMapsB)
-			{
-				const FAttributeMappingKey MappingKey = SetMapB->GetMappingKey();
-				FAttributeSetMap* SetMapA = CollectionA.AddSetMap(MappingKey);
-				checkf(SetMapA != nullptr, TEXT("Failed to add set mapping"));
+			// Name/value maps
+			OuterJoinBy(
+				FAttributeValueTransformerListWithAttributeCollectionMapJoinOp(),
+				[Alpha, &OutputCollection](FRawTransformerFunc Transformer, const FAttributeMap* MapA, const FAttributeMap* MapB)
+				{
+					if (MapA == nullptr)
+					{
+						if (MapB == nullptr)
+						{
+							// No maps found for the transformer type
+						}
+						else
+						{
+							// This map doesn't have a matching pair, duplicate it
+							FAttributeMap* MapCopy = MapB->Duplicate(OutputCollection.GetAllocator());
+							OutputCollection.AppendMap(MapCopy);
+						}
+					}
+					else
+					{
+						if (MapB == nullptr)
+						{
+							// This map doesn't have a matching pair
+							// Our output is A and thus already contains this
+						}
+						else
+						{
+							if (Transformer == nullptr)
+							{
+								// No transformer was found for this value type, retain the one with the largest weight
+								FAttributeMap* MapCopy;
+								if (Alpha >= 0.5f)
+								{
+									MapCopy = MapB->Duplicate(OutputCollection.GetAllocator());
+								}
+								else
+								{
+									MapCopy = MapA->Duplicate(OutputCollection.GetAllocator());
+								}
 
-				SetMapB->CopyTo(*SetMapA);
-			}
+								OutputCollection.AppendMap(MapCopy);
+							}
+							else
+							{
+								// Transform our matching pair in place
+								FAttributeMap* OutputMap = nullptr;
+
+								FTransformMapFunc TypedTransformer = reinterpret_cast<FTransformMapFunc>(Transformer);
+								(*TypedTransformer)(MapA, MapB, Alpha, OutputMap, OutputCollection.GetAllocator());
+
+								checkf(OutputMap != nullptr, TEXT("Transformer function is responsible for allocating an output buffer"));
+								checkf(OutputMap->GetAllocator() == OutputCollection.GetAllocator(), TEXT("Transformer function is responsible for allocating using the output allocator"));
+								OutputCollection.AppendMap(OutputMap);
+							}
+						}
+					}
+				},
+				TransformerList->CreateMapTransformerIterator(),
+				CollectionA.CreateConstMapIterator(),
+				CollectionB.CreateConstMapIterator());
 		}
 	}
 
-	void FInterpolate_BoneTransformAttribute::Transform(UE::UAF::FAttributeSetMap* SetMapA, const UE::UAF::FAttributeSetMap* SetMapB, float Alpha)
+	void FInterpolate_BoneTransformAttribute::TransformSetMap(const FAttributeSetMap* SetMapA, const FAttributeSetMap* SetMapB, float Alpha, FAttributeSetMap*& OutputSetMap, FReallocFun OutputAllocator)
 	{
 		check(SetMapA->GetTypedSet() == SetMapB->GetTypedSet());
+		check(SetMapA->Num() == SetMapB->Num());
+
+		if (OutputSetMap == nullptr)
+		{
+			const FAttributeSetMap::FConstructArgs Args(SetMapA->GetTypedSet(), SetMapA->GetValueType(), OutputAllocator);
+			OutputSetMap = MakeAttributeSetMap<FBoneTransformAnimationAttribute>(Args);
+		}
+
+		check(SetMapA->GetTypedSet() == OutputSetMap->GetTypedSet());
+		check(SetMapA->Num() == OutputSetMap->Num());
 
-		TAttributeSetMap<FBoneTransformAnimationAttribute>* BoneTransformsA = static_cast<TAttributeSetMap<FBoneTransformAnimationAttribute>*>(SetMapA);
-		const TAttributeSetMap<FBoneTransformAnimationAttribute>* BoneTransformsB = static_cast<const TAttributeSetMap<FBoneTransformAnimationAttribute>*>(SetMapB);
+		const TAttributeSetMap<FBoneTransformAnimationAttribute>* BoneTransformsA = CastChecked<FBoneTransformAnimationAttribute>(SetMapA);
+		const TAttributeSetMap<FBoneTransformAnimationAttribute>* BoneTransformsB = CastChecked<FBoneTransformAnimationAttribute>(SetMapB);
+		TAttributeSetMap<FBoneTransformAnimationAttribute>* OutputBoneTransforms = CastChecked<FBoneTransformAnimationAttribute>(OutputSetMap);
 
-		FBoneTransformAnimationAttribute* TransformPtrA = BoneTransformsA->GetData();
+		const FBoneTransformAnimationAttribute* TransformPtrA = BoneTransformsA->GetData();
 		const FBoneTransformAnimationAttribute* TransformPtrB = BoneTransformsB->GetData();
-		const FBoneTransformAnimationAttribute* EndTransformPtrA = TransformPtrA + BoneTransformsA->Num();
+		FBoneTransformAnimationAttribute* OutputTransformPtr = OutputBoneTransforms->GetData();
+		const FBoneTransformAnimationAttribute* TransformPtrAEnd = TransformPtrA + BoneTransformsA->Num();
 
-		while (TransformPtrA < EndTransformPtrA)
+		while (TransformPtrA < TransformPtrAEnd)
 		{
-			TransformPtrA->Value.Blend(TransformPtrA->Value, TransformPtrB->Value, Alpha);
+			OutputTransformPtr->Value.Blend(TransformPtrA->Value, TransformPtrB->Value, Alpha);
 
 			TransformPtrA++;
 			TransformPtrB++;
+			OutputTransformPtr++;
 		}
 	}
+
+	void FInterpolate_BoneTransformAttribute::TransformMap(const FAttributeMap* MapA, const FAttributeMap* MapB, float Alpha, FAttributeMap*& OutputMap, FReallocFun OutputAllocator)
+	{
+		check(OutputMap == nullptr);
+
+		const TAttributeMap<FBoneTransformAnimationAttribute>* BoneTransformsA = CastChecked<FBoneTransformAnimationAttribute>(MapA);
+		const TAttributeMap<FBoneTransformAnimationAttribute>* BoneTransformsB = CastChecked<FBoneTransformAnimationAttribute>(MapB);
+		TAttributeMap<FBoneTransformAnimationAttribute>* OutputBoneTransforms = MakeAttributeMap<FBoneTransformAnimationAttribute>(OutputAllocator);
+
+		OuterJoinBy(
+			FAttributeMapJoinOp<FBoneTransformAnimationAttribute>(),
+			[&Alpha, OutputBoneTransforms](FName Name, const FBoneTransformAnimationAttribute* A, const FBoneTransformAnimationAttribute* B)
+			{
+				// Missing values are set to the identity
+				const FTransform ValueA = A != nullptr ? A->Value : FTransform::Identity;
+				const FTransform ValueB = B != nullptr ? B->Value : FTransform::Identity;
+
+				FBoneTransformAnimationAttribute Output;
+				Output.Value.Blend(ValueA, ValueB, Alpha);
+
+				OutputBoneTransforms->Append(Name, Output);
+			},
+			BoneTransformsA->CreateConstIterator(),
+			BoneTransformsB->CreateConstIterator());
+
+		OutputMap = OutputBoneTransforms;
+	}
+
+	void FInterpolate_FloatAttribute::TransformSetMap(const FAttributeSetMap* SetMapA, const FAttributeSetMap* SetMapB, float Alpha, FAttributeSetMap*& OutputSetMap, FReallocFun OutputAllocator)
+	{
+		check(SetMapA->GetTypedSet() == SetMapB->GetTypedSet());
+		check(SetMapA->Num() == SetMapB->Num());
+
+		if (OutputSetMap == nullptr)
+		{
+			const FAttributeSetMap::FConstructArgs Args(SetMapA->GetTypedSet(), SetMapA->GetValueType(), OutputAllocator);
+			OutputSetMap = MakeAttributeSetMap<FFloatAnimationAttribute>(Args);
+		}
+
+		check(SetMapA->GetTypedSet() == OutputSetMap->GetTypedSet());
+		check(SetMapA->Num() == OutputSetMap->Num());
+
+		const TAttributeSetMap<FFloatAnimationAttribute>* FloatsA = CastChecked<FFloatAnimationAttribute>(SetMapA);
+		const TAttributeSetMap<FFloatAnimationAttribute>* FloatsB = CastChecked<FFloatAnimationAttribute>(SetMapB);
+		TAttributeSetMap<FFloatAnimationAttribute>* OutputFloats = CastChecked<FFloatAnimationAttribute>(OutputSetMap);
+
+		const FFloatAnimationAttribute* FloatsPtrA = FloatsA->GetData();
+		const FFloatAnimationAttribute* FloatsPtrB = FloatsB->GetData();
+		FFloatAnimationAttribute* OutputFloatsPtr = OutputFloats->GetData();
+		const FFloatAnimationAttribute* FloatsPtrAEnd = FloatsPtrA + FloatsA->Num();
+
+		while (FloatsPtrA < FloatsPtrAEnd)
+		{
+			OutputFloatsPtr->Value = FMath::Lerp(FloatsPtrA->Value, FloatsPtrB->Value, Alpha);
+
+			FloatsPtrA++;
+			FloatsPtrB++;
+			OutputFloatsPtr++;
+		}
+	}
+
+	void FInterpolate_FloatAttribute::TransformMap(const FAttributeMap* MapA, const FAttributeMap* MapB, float Alpha, FAttributeMap*& OutputMap, FReallocFun OutputAllocator)
+	{
+		check(OutputMap == nullptr);
+
+		const TAttributeMap<FFloatAnimationAttribute>* FloatsA = CastChecked<FFloatAnimationAttribute>(MapA);
+		const TAttributeMap<FFloatAnimationAttribute>* FloatsB = CastChecked<FFloatAnimationAttribute>(MapB);
+		TAttributeMap<FFloatAnimationAttribute>* OutputFloats = MakeAttributeMap<FFloatAnimationAttribute>(OutputAllocator);
+
+		OuterJoinBy(
+			FAttributeMapJoinOp<FFloatAnimationAttribute>(),
+			[&Alpha, OutputFloats](FName Name, const FFloatAnimationAttribute* A, const FFloatAnimationAttribute* B)
+			{
+				// Missing values are set to 0.0
+				const float ValueA = A != nullptr ? A->Value : 0.0f;
+				const float ValueB = B != nullptr ? B->Value : 0.0f;
+
+				FFloatAnimationAttribute Output;
+				Output.Value = FMath::Lerp(ValueA, ValueB, Alpha);
+
+				OutputFloats->Append(Name, Output);
+			},
+			FloatsA->CreateConstIterator(),
+			FloatsB->CreateConstIterator());
+
+		OutputMap = OutputFloats;
+	}
 }
diff --git a/Engine/Plugins/Experimental/UAF/UAF/Source/UAF/Private/UAF/Attributes/Transformers/Overwrite.cpp b/Engine/Plugins/Experimental/UAF/UAF/Source/UAF/Private/UAF/Attributes/Transformers/Overwrite.cpp
--- a/Engine/Plugins/Experimental/UAF/UAF/Source/UAF/Private/UAF/Attributes/Transformers/Overwrite.cpp
+++ b/Engine/Plugins/Experimental/UAF/UAF/Source/UAF/Private/UAF/Attributes/Transformers/Overwrite.cpp
@@ -2,97 +2,289 @@
 
 #include "UAF/Attributes/Transformers/Overwrite.h"
 
-#include "UAF/Attributes/AttributeCollection.h"
-#include "UAF/Attributes/AttributeSetMap.h"
 #include "UAF/Attributes/IteratorUtils.h"
 
 namespace UE::UAF::Transformers
 {
 	FName FOverwrite::TransformerName = TEXT("UAF::Transformers::Overwrite");
 
-	void FOverwrite::Apply(const FAttributeValueTransformerMapPtr& TransformerMap, FAttributeCollection& CollectionA, const FAttributeCollection& CollectionB, float Alpha)
+	void FOverwrite::Apply(const FAttributeValueTransformerMapPtr& TransformerMap, const FAttributeCollection& SourceCollection, float Alpha, FAttributeCollection& OutputCollection)
 	{
 		// Named sets must match as it ensures that our inputs have the same sizes/shapes
-		check(CollectionA.GetNamedSet() == CollectionB.GetNamedSet());
+		check(SourceCollection.GetNamedSet() == OutputCollection.GetNamedSet());
 
-		if (const FAttributeValueTransformerList* TransformerList = TransformerMap->Find(TransformerName))
+		const FAttributeValueTransformerList* TransformerList = TransformerMap->Find(TransformerName);
+		if (TransformerList == nullptr)
 		{
-			// Keep a list of set maps that don't have a matching pair
-			// These are duplicated afterwards to avoid mutating the collection while iterating
-			TArray<const FAttributeSetMap*, TInlineAllocator<8>> UnmatchedSetMapsB;
+			// No transformer implementations have been registered
+			return;
+		}
 
-			OuterJoinBy(
-				FAttributeValueTransformerListWithAttributeCollectionJoinOp(),
-				[Alpha, &UnmatchedSetMapsB](FRawTransformerFunc Transformer, FAttributeSetMap* SetMapA, const FAttributeSetMap* SetMapB)
-				{
-					if (SetMapA == nullptr)
+		if (&SourceCollection == &OutputCollection)
+		{
+			if (FAnimWeight::IsFullWeight(Alpha))
+			{
+				// Our input is our output and it has full weight, nothing to do
+				return;
+			}
+
+			// Our output collection matches one of our inputs, we'll modify it in place
+			FReallocFun OutputAllocator = OutputCollection.GetAllocator();
+
+			// Set maps
+			{
+				OuterJoinBy(
+					FAttributeValueTransformerListWithAttributeCollectionSetMapJoinOp(),
+					[Alpha, OutputAllocator](FRawTransformerFunc Transformer, FAttributeSetMap* SourceSetMap)
 					{
-						if (SetMapB == nullptr)
+						if (SourceSetMap == nullptr)
 						{
 							// No maps found for the transformer type
 						}
 						else
 						{
-							// This set map doesn't have a matching pair
-							UnmatchedSetMapsB.Add(SetMapB);
+							if (Transformer == nullptr)
+							{
+								// No transformer was found for this value type, retain it as-is
+							}
+							else
+							{
+								// Transform our input in place
+								FAttributeSetMap* OutputSetMap = SourceSetMap;
+
+								FTransformSetMapFunc TypedTransformer = reinterpret_cast<FTransformSetMapFunc>(Transformer);
+								(*TypedTransformer)(SourceSetMap, Alpha, OutputSetMap, nullptr);
+							}
 						}
-					}
-					else
+					},
+					TransformerList->CreateSetMapTransformerIterator(),
+					// Our input is our output
+					OutputCollection.CreateSetMapIterator());
+			}
+
+			// Name/value maps
+			{
+				OuterJoinBy(
+					FAttributeValueTransformerListWithAttributeCollectionMapJoinOp(),
+					[Alpha, OutputAllocator](FRawTransformerFunc Transformer, FAttributeMap* SourceMap)
 					{
-						if (SetMapB == nullptr)
+						if (SourceMap == nullptr)
 						{
-							// This set map doesn't have a matching pair, keep it
+							// No maps found for the transformer type
 						}
 						else
 						{
 							if (Transformer == nullptr)
 							{
-								// No transformer was found for this value type
+								// No transformer was found for this value type, retain it as-is
 							}
 							else
 							{
-								// Transform our matching pair
-								FTransformFunc TypedTransformer = reinterpret_cast<FTransformFunc>(Transformer);
-								(*TypedTransformer)(SetMapA, SetMapB, Alpha);
+								// Transform our input in place
+								FAttributeMap* OutputMap = SourceMap;
+
+								FTransformMapFunc TypedTransformer = reinterpret_cast<FTransformMapFunc>(Transformer);
+								(*TypedTransformer)(SourceMap, Alpha, OutputMap, nullptr);
 							}
 						}
+					},
+					TransformerList->CreateMapTransformerIterator(),
+					// Our input is our output
+					OutputCollection.CreateMapIterator());
+			}
+		}
+		else
+		{
+			// Our output collection isn't either of our inputs, reset it
+			OutputCollection.Reset(OutputCollection.GetNamedSet());
+
+			// Set maps
+			OuterJoinBy(
+				FAttributeValueTransformerListWithAttributeCollectionSetMapJoinOp(),
+				[Alpha, &OutputCollection](FRawTransformerFunc Transformer, const FAttributeSetMap* SetMapA)
+				{
+					if (SetMapA == nullptr)
+					{
+						// No maps found for the transformer type
+					}
+					else
+					{
+						if (Transformer == nullptr)
+						{
+							// No transformer was found for this value type, retain it as-is
+							FAttributeSetMap* SetMapCopy = SetMapA->Duplicate(OutputCollection.GetAllocator());
+							OutputCollection.AppendSetMap(SetMapCopy);
+						}
+						else
+						{
+							// Transform our matching pair
+							FAttributeSetMap* OutputSetMap = nullptr;
+
+							FTransformSetMapFunc TypedTransformer = reinterpret_cast<FTransformSetMapFunc>(Transformer);
+							(*TypedTransformer)(SetMapA, Alpha, OutputSetMap, OutputCollection.GetAllocator());
+
+							checkf(OutputSetMap != nullptr, TEXT("Transformer function is responsible for allocating an output buffer"));
+							checkf(OutputSetMap->GetAllocator() == OutputCollection.GetAllocator(), TEXT("Transformer function is responsible for allocating using the output allocator"));
+							OutputCollection.AppendSetMap(OutputSetMap);
+						}
 					}
 				},
-				TransformerList->CreateTransformerIterator(),
-				CollectionA.CreateSetMapIterator(),
-				CollectionB.CreateConstSetMapIterator());
+				TransformerList->CreateSetMapTransformerIterator(),
+				SourceCollection.CreateConstSetMapIterator());
 
-			// Duplicate any unmatched entries
-			for (const FAttributeSetMap* SetMapB : UnmatchedSetMapsB)
-			{
-				const FAttributeMappingKey MappingKey = SetMapB->GetMappingKey();
-				FAttributeSetMap* SetMapA = CollectionA.AddSetMap(MappingKey);
-				checkf(SetMapA != nullptr, TEXT("Failed to add set mapping"));
+			// Name/value maps
+			OuterJoinBy(
+				FAttributeValueTransformerListWithAttributeCollectionMapJoinOp(),
+				[Alpha, &OutputCollection](FRawTransformerFunc Transformer, const FAttributeMap* MapA)
+				{
+					if (MapA == nullptr)
+					{
+						// No maps found for the transformer type
+					}
+					else
+					{
+						if (Transformer == nullptr)
+						{
+							// No transformer was found for this value type, retain it as-is
+							FAttributeMap* MapCopy = MapA->Duplicate(OutputCollection.GetAllocator());
+							OutputCollection.AppendMap(MapCopy);
+						}
+						else
+						{
+							// Transform our matching pair
+							FAttributeMap* OutputMap = nullptr;
 
-				SetMapB->CopyTo(*SetMapA);
-			}
+							FTransformMapFunc TypedTransformer = reinterpret_cast<FTransformMapFunc>(Transformer);
+							(*TypedTransformer)(MapA, Alpha, OutputMap, OutputCollection.GetAllocator());
+
+							checkf(OutputMap != nullptr, TEXT("Transformer function is responsible for allocating an output buffer"));
+							checkf(OutputMap->GetAllocator() == OutputCollection.GetAllocator(), TEXT("Transformer function is responsible for allocating using the output allocator"));
+							OutputCollection.AppendMap(OutputMap);
+						}
+					}
+				},
+				TransformerList->CreateMapTransformerIterator(),
+				SourceCollection.CreateConstMapIterator());
 		}
 	}
 
-	void FOverwrite_BoneTransformAttribute::Transform(UE::UAF::FAttributeSetMap* SetMapA, const UE::UAF::FAttributeSetMap* SetMapB, float Alpha)
+	void FOverwrite_BoneTransformAttribute::TransformSetMap(const FAttributeSetMap* SetMap, float Alpha, FAttributeSetMap*& OutputSetMap, FReallocFun OutputAllocator)
 	{
-		check(SetMapA->GetTypedSet() == SetMapB->GetTypedSet());
+		if (OutputSetMap == nullptr)
+		{
+			const FAttributeSetMap::FConstructArgs Args(SetMap->GetTypedSet(), SetMap->GetValueType(), OutputAllocator);
+			OutputSetMap = MakeAttributeSetMap<FBoneTransformAnimationAttribute>(Args);
+		}
+
+		check(SetMap->GetTypedSet() == OutputSetMap->GetTypedSet());
+		check(SetMap->Num() == OutputSetMap->Num());
+
+		const TAttributeSetMap<FBoneTransformAnimationAttribute>* InputBoneTransforms = CastChecked<FBoneTransformAnimationAttribute>(SetMap);
+		TAttributeSetMap<FBoneTransformAnimationAttribute>* OutputBoneTransforms = CastChecked<FBoneTransformAnimationAttribute>(OutputSetMap);
+
+		const FBoneTransformAnimationAttribute* InputTransformPtr = InputBoneTransforms->GetData();
+		FBoneTransformAnimationAttribute* OutputTransformPtr = OutputBoneTransforms->GetData();
+		const FBoneTransformAnimationAttribute* InputTransformPtrEnd = InputTransformPtr + InputBoneTransforms->Num();
+
+		const ScalarRegister VAlpha(Alpha);
+
+		while (InputTransformPtr < InputTransformPtrEnd)
+		{
+			OutputTransformPtr->Value = InputTransformPtr->Value * VAlpha;
+
+			InputTransformPtr++;
+			OutputTransformPtr++;
+		}
+	}
 
-		TAttributeSetMap<FBoneTransformAnimationAttribute>* BoneTransformsA = static_cast<TAttributeSetMap<FBoneTransformAnimationAttribute>*>(SetMapA);
-		const TAttributeSetMap<FBoneTransformAnimationAttribute>* BoneTransformsB = static_cast<const TAttributeSetMap<FBoneTransformAnimationAttribute>*>(SetMapB);
+	void FOverwrite_BoneTransformAttribute::TransformMap(const FAttributeMap* Map, float Alpha, FAttributeMap*& OutputMap, FReallocFun OutputAllocator)
+	{
+		const bool bAllocateOutput = OutputMap == nullptr;
+		if (bAllocateOutput)
+		{
+			OutputMap = MakeAttributeMap<FBoneTransformAnimationAttribute>(OutputAllocator);
+		}
 
-		FBoneTransformAnimationAttribute* TransformPtrA = BoneTransformsA->GetData();
-		const FBoneTransformAnimationAttribute* TransformPtrB = BoneTransformsB->GetData();
-		const FBoneTransformAnimationAttribute* EndTransformPtrA = TransformPtrA + BoneTransformsA->Num();
+		const TAttributeMap<FBoneTransformAnimationAttribute>* InputBoneTransforms = CastChecked<FBoneTransformAnimationAttribute>(Map);
+		TAttributeMap<FBoneTransformAnimationAttribute>* OutputBoneTransforms = CastChecked<FBoneTransformAnimationAttribute>(OutputMap);
 
 		const ScalarRegister VAlpha(Alpha);
 
-		while (TransformPtrA < EndTransformPtrA)
+		if (bAllocateOutput)
+		{
+			// We allocated a new output, append to it
+			for (auto It = InputBoneTransforms->CreateConstIterator(); It; ++It)
+			{
+				FBoneTransformAnimationAttribute Output{ It.GetValue()->Value * VAlpha };
+				OutputBoneTransforms->Append(It.GetName(), Output);
+			}
+		}
+		else
+		{
+			// We overwrite in place
+			check(InputBoneTransforms == OutputBoneTransforms);
+			for (auto It = OutputBoneTransforms->CreateIterator(); It; ++It)
+			{
+				It.GetValue()->Value *= VAlpha;
+			}
+		}
+	}
+
+	void FOverwrite_FloatAttribute::TransformSetMap(const FAttributeSetMap* SetMap, float Alpha, FAttributeSetMap*& OutputSetMap, FReallocFun OutputAllocator)
+	{
+		if (OutputSetMap == nullptr)
+		{
+			const FAttributeSetMap::FConstructArgs Args(SetMap->GetTypedSet(), SetMap->GetValueType(), OutputAllocator);
+			OutputSetMap = MakeAttributeSetMap<FFloatAnimationAttribute>(Args);
+		}
+
+		check(SetMap->GetTypedSet() == OutputSetMap->GetTypedSet());
+		check(SetMap->Num() == OutputSetMap->Num());
+
+		const TAttributeSetMap<FFloatAnimationAttribute>* InputFloats = CastChecked<FFloatAnimationAttribute>(SetMap);
+		TAttributeSetMap<FFloatAnimationAttribute>* OutputFloats = CastChecked<FFloatAnimationAttribute>(OutputSetMap);
+
+		const FFloatAnimationAttribute* FloatsPtr = InputFloats->GetData();
+		FFloatAnimationAttribute* OutputFloatsPtr = OutputFloats->GetData();
+		const FFloatAnimationAttribute* FloatsPtrEnd = FloatsPtr + InputFloats->Num();
+
+		while (FloatsPtr < FloatsPtrEnd)
+		{
+			OutputFloatsPtr->Value = FloatsPtr->Value * Alpha;
+
+			FloatsPtr++;
+			OutputFloatsPtr++;
+		}
+	}
+
+	void FOverwrite_FloatAttribute::TransformMap(const FAttributeMap* Map, float Alpha, FAttributeMap*& OutputMap, FReallocFun OutputAllocator)
+	{
+		const bool bAllocateOutput = OutputMap == nullptr;
+		if (bAllocateOutput)
 		{
-			TransformPtrA->Value = TransformPtrB->Value * VAlpha;
+			OutputMap = MakeAttributeMap<FFloatAnimationAttribute>(OutputAllocator);
+		}
+
+		const TAttributeMap<FFloatAnimationAttribute>* InputFloats = CastChecked<FFloatAnimationAttribute>(Map);
+		TAttributeMap<FFloatAnimationAttribute>* OutputFloats = CastChecked<FFloatAnimationAttribute>(OutputMap);
 
-			TransformPtrA++;
-			TransformPtrB++;
+		if (bAllocateOutput)
+		{
+			// We allocated a new output, append to it
+			for (auto It = InputFloats->CreateConstIterator(); It; ++It)
+			{
+				FFloatAnimationAttribute Output{ It.GetValue()->Value * Alpha };
+				OutputFloats->Append(It.GetName(), Output);
+			}
+		}
+		else
+		{
+			// We overwrite in place
+			check(InputFloats == OutputFloats);
+			for (auto It = OutputFloats->CreateIterator(); It; ++It)
+			{
+				It.GetValue()->Value *= Alpha;
+			}
 		}
 	}
 }
diff --git a/Engine/Plugins/Experimental/UAF/UAF/Source/UAF/Private/UAF/Attributes/Transformers/Sanitize.cpp b/Engine/Plugins/Experimental/UAF/UAF/Source/UAF/Private/UAF/Attributes/Transformers/Sanitize.cpp
--- a/Engine/Plugins/Experimental/UAF/UAF/Source/UAF/Private/UAF/Attributes/Transformers/Sanitize.cpp
+++ b/Engine/Plugins/Experimental/UAF/UAF/Source/UAF/Private/UAF/Attributes/Transformers/Sanitize.cpp
@@ -2,34 +2,45 @@
 
 #include "UAF/Attributes/Transformers/Sanitize.h"
 
-#include "UAF/Attributes/AttributeCollection.h"
-#include "UAF/Attributes/AttributeSetMap.h"
 #include "UAF/Attributes/IteratorUtils.h"
 
-// We validate bone transforms in development and debug builds
-#define UE_VALIDATE_BONE_TRANSFORMS !(UE_BUILD_SHIPPING || UE_BUILD_TEST)
-
 namespace UE::UAF::Transformers
 {
 	FName FSanitize::TransformerName = TEXT("UAF::Transformers::Sanitize");
 
 	void FSanitize::Apply(const FAttributeValueTransformerMapPtr& TransformerMap, FAttributeCollection& Collection)
 	{
-		if (const FAttributeValueTransformerList* TransformerList = TransformerMap->Find(TransformerName))
+		const FAttributeValueTransformerList* TransformerList = TransformerMap->Find(TransformerName);
+		if (TransformerList == nullptr)
 		{
-			InnerJoinBy(
-				FAttributeValueTransformerListWithAttributeCollectionJoinOp(),
-				[](FRawTransformerFunc Transformer, FAttributeSetMap* SetMap)
-				{
-					FTransformFunc TypedTransformer = reinterpret_cast<FTransformFunc>(Transformer);
-					(*TypedTransformer)(SetMap);
-				},
-				TransformerList->CreateTransformerIterator(),
-				Collection.CreateSetMapIterator());
+			// No transformer implementations have been registered
+			return;
 		}
+
+		// Set maps
+		InnerJoinBy(
+			FAttributeValueTransformerListWithAttributeCollectionSetMapJoinOp(),
+			[](FRawTransformerFunc Transformer, FAttributeSetMap* SetMap)
+			{
+				FTransformSetMapFunc TypedTransformer = reinterpret_cast<FTransformSetMapFunc>(Transformer);
+				(*TypedTransformer)(SetMap);
+			},
+			TransformerList->CreateSetMapTransformerIterator(),
+			Collection.CreateSetMapIterator());
+
+		// Name/value maps
+		InnerJoinBy(
+			FAttributeValueTransformerListWithAttributeCollectionMapJoinOp(),
+			[](FRawTransformerFunc Transformer, FAttributeMap* Map)
+			{
+				FTransformMapFunc TypedTransformer = reinterpret_cast<FTransformMapFunc>(Transformer);
+				(*TypedTransformer)(Map);
+			},
+			TransformerList->CreateMapTransformerIterator(),
+			Collection.CreateMapIterator());
 	}
 
-#if UE_VALIDATE_BONE_TRANSFORMS
+#if UE_UAF_VALIDATE_MAPPED_VALUES
 	// Don't inline this function to keep the stack usage down
 	static FORCENOINLINE void IsBoneTransformSetMapValid(const TAttributeSetMap<FBoneTransformAnimationAttribute>* BoneTransforms)
 	{
@@ -46,7 +57,7 @@ namespace UE::UAF::Transformers
 			}
 			else if (!TransformPtr->Value.IsRotationNormalized())
 			{
-				ensureMsgf(false, TEXT("Bone (%s) Rotation not normalized with Rotation:[%s]"),
+				ensureMsgf(false, TEXT("Bone (%s) rotation not normalized with Rotation:[%s]"),
 					*BoneTransforms->GetTypedSet()->GetName(FAttributeSetIndex(TransformPtr - BoneTransforms->GetData())).ToString(),
 					*TransformPtr->Value.GetRotation().ToString());
 			}
@@ -54,24 +65,89 @@ namespace UE::UAF::Transformers
 			TransformPtr++;
 		}
 	}
+
+	// Don't inline this function to keep the stack usage down
+	static FORCENOINLINE void IsBoneTransformMapValid(const TAttributeMap<FBoneTransformAnimationAttribute>* BoneTransforms)
+	{
+		for (auto It = BoneTransforms->CreateConstIterator(); It; ++It)
+		{
+			if (It.GetValue()->Value.ContainsNaN())
+			{
+				ensureMsgf(false, TEXT("Transform attribute (%s) contains NaN with Value:[%s]"),
+					*It.GetName().ToString(),
+					*It.GetValue()->Value.ToString());
+			}
+			else if (!It.GetValue()->Value.IsRotationNormalized())
+			{
+				ensureMsgf(false, TEXT("Transform attribute (%s) rotation not normalized with Rotation:[%s]"),
+					*It.GetName().ToString(),
+					*It.GetValue()->Value.GetRotation().ToString());
+			}
+		}
+	}
 #endif
 
-	void FSanitize_BoneTransformAttribute::Transform(FAttributeSetMap* SetMap)
+	void FSanitize_BoneTransformAttribute::TransformSetMap(FAttributeSetMap* SetMap)
 	{
-		TAttributeSetMap<FBoneTransformAnimationAttribute>* BoneTransforms = static_cast<TAttributeSetMap<FBoneTransformAnimationAttribute>*>(SetMap);
+		TAttributeSetMap<FBoneTransformAnimationAttribute>* BoneTransforms = CastChecked<FBoneTransformAnimationAttribute>(SetMap);
 
 		FBoneTransformAnimationAttribute* TransformPtr = BoneTransforms->GetData();
-		const FBoneTransformAnimationAttribute* EndTransformPtr = TransformPtr + BoneTransforms->Num();
+		const FBoneTransformAnimationAttribute* TransformPtrEnd = TransformPtr + BoneTransforms->Num();
 
-		while (TransformPtr < EndTransformPtr)
+		while (TransformPtr < TransformPtrEnd)
 		{
 			TransformPtr->Value.NormalizeRotation();
 
 			TransformPtr++;
 		}
 
-#if UE_VALIDATE_BONE_TRANSFORMS
+#if UE_UAF_VALIDATE_MAPPED_VALUES
 		IsBoneTransformSetMapValid(BoneTransforms);
 #endif
 	}
+
+	void FSanitize_BoneTransformAttribute::TransformMap(FAttributeMap* Map)
+	{
+		TAttributeMap<FBoneTransformAnimationAttribute>* BoneTransforms = CastChecked<FBoneTransformAnimationAttribute>(Map);
+
+		for (auto It = BoneTransforms->CreateIterator(); It; ++It)
+		{
+			It.GetValue()->Value.NormalizeRotation();
+		}
+
+#if UE_UAF_VALIDATE_MAPPED_VALUES
+		IsBoneTransformMapValid(BoneTransforms);
+#endif
+	}
+
+#if UE_UAF_VALIDATE_MAPPED_VALUES
+	void FSanitize_FloatAttribute::TransformSetMap(FAttributeSetMap* SetMap)
+	{
+		const TAttributeSetMap<FFloatAnimationAttribute>* Floats = CastChecked<FFloatAnimationAttribute>(SetMap);
+
+		const FFloatAnimationAttribute* FloatsPtr = Floats->GetData();
+		const FFloatAnimationAttribute* FloatsPtrEnd = FloatsPtr + Floats->Num();
+
+		while (FloatsPtr < FloatsPtrEnd)
+		{
+			ensureMsgf(!FMath::IsNaN(FloatsPtr->Value), TEXT("Float attribute (%s) contains NaN with Value:[%3.3f]"),
+				*Floats->GetTypedSet()->GetName(FAttributeSetIndex(FloatsPtr - Floats->GetData())).ToString(),
+				FloatsPtr->Value);
+
+			FloatsPtr++;
+		}
+	}
+
+	void FSanitize_FloatAttribute::TransformMap(FAttributeMap* Map)
+	{
+		const TAttributeMap<FFloatAnimationAttribute>* Floats = CastChecked<FFloatAnimationAttribute>(Map);
+
+		for (auto It = Floats->CreateConstIterator(); It; ++It)
+		{
+			ensureMsgf(!FMath::IsNaN(It.GetValue()->Value), TEXT("Float attribute (%s) contains NaN with Value:[%3.3f]"),
+				*It.GetName().ToString(),
+				It.GetValue()->Value);
+		}
+	}
+#endif
 }
