Add support for dynamic attributes in abstract skeleton runtime

diff --git a/Engine/Plugins/Experimental/UAF/UAF/Source/UAF/Private/UAF/Attributes/Transformers/Transformer.cpp b/Engine/Plugins/Experimental/UAF/UAF/Source/UAF/Private/UAF/Attributes/Transformers/Transformer.cpp
--- a/Engine/Plugins/Experimental/UAF/UAF/Source/UAF/Private/UAF/Attributes/Transformers/Transformer.cpp
+++ b/Engine/Plugins/Experimental/UAF/UAF/Source/UAF/Private/UAF/Attributes/Transformers/Transformer.cpp
@@ -11,9 +11,37 @@ namespace UE::UAF
 	{
 	}
 
-	bool FAttributeValueTransformerList::AddTransformer(UScriptStruct* ValueType, FRawTransformerFunc TransformerFunc)
+	FName FAttributeValueTransformerList::GetTransformerName() const
 	{
-		const int32 EntryIndex = Algo::LowerBound(ValueTypes, ValueType);
+		return TransformerName;
+	}
+
+	bool FAttributeValueTransformerList::AddSetMapTransformer(UScriptStruct* ValueType, FRawTransformerFunc TransformerFunc)
+	{
+		const int32 EntryIndex = Algo::LowerBound(MakeArrayView(ValueTypes.GetData(), NumSetMapTransformers()), ValueType);
+		if (EntryIndex < ValueTypes.Num() && ValueTypes[EntryIndex] == ValueType)
+		{
+			// A transformer for this value type has already been registered
+			return false;
+		}
+
+		ValueTypes.Insert(ValueType, EntryIndex);
+		Transformers.Insert(TransformerFunc, EntryIndex);
+		SetMapTransformerCount++;
+
+		return true;
+	}
+
+	bool FAttributeValueTransformerList::AddMapTransformer(UScriptStruct* ValueType, FRawTransformerFunc TransformerFunc)
+	{
+		const int32 NumSetMapEntries = NumSetMapTransformers();
+		const int32 NumMapEntries = NumMapTransformers();
+
+		const int32 PerContainerEntryIndex = Algo::LowerBound(MakeArrayView(ValueTypes.GetData() + NumSetMapEntries, NumMapEntries), ValueType);
+
+		// Fixup our index
+		const int32 EntryIndex = PerContainerEntryIndex + NumSetMapEntries;
+
 		if (EntryIndex < ValueTypes.Num() && ValueTypes[EntryIndex] == ValueType)
 		{
 			// A transformer for this value type has already been registered
diff --git a/Engine/Plugins/Experimental/UAF/UAF/Source/UAF/Public/UAF/Attributes/AttributeCollection.h b/Engine/Plugins/Experimental/UAF/UAF/Source/UAF/Public/UAF/Attributes/AttributeCollection.h
--- a/Engine/Plugins/Experimental/UAF/UAF/Source/UAF/Public/UAF/Attributes/AttributeCollection.h
+++ b/Engine/Plugins/Experimental/UAF/UAF/Source/UAF/Public/UAF/Attributes/AttributeCollection.h
@@ -5,6 +5,7 @@
 #include "CoreMinimal.h"
 #include "Misc/MemStack.h"
 #include "UAF/Attributes/AttributeNamedSet.h"
+#include "UAF/Attributes/AttributeMap.h"
 #include "UAF/Attributes/AttributeMappingKey.h"
 #include "UAF/Attributes/AttributeSetMap.h"
 #include "UAF/Attributes/IndirectAllocator.h"
@@ -54,6 +55,39 @@ namespace UE::UAF
 		EntryPtrType EntryEndPtr = nullptr;
 	};
 
+	// A iterator that iterates over all name/value maps within an attribute collection
+	template<class ContainerType>
+	class TAttributeCollectionMapIterator
+	{
+	public:
+		static constexpr bool bIsContainerConst = std::is_const_v<ContainerType>;
+
+		using EntryPtrType = typename std::conditional<bIsContainerConst, const typename ContainerType::FMapEntry*, typename ContainerType::FMapEntry*>::type;
+		using MapType = typename std::conditional<bIsContainerConst, const FAttributeMap*, FAttributeMap*>::type;
+
+		// Creates an iterator over the name/value map entries within the specified collection
+		[[nodiscard]] TAttributeCollectionMapIterator(ContainerType& Collection);
+
+		// Increments and moves the iterator to the next name/value map entry
+		TAttributeCollectionMapIterator& operator++();
+
+		// Returns whether or not the iterator still contains values
+		[[nodiscard]] explicit operator bool() const;
+
+		// Returns the value type of the current name/value map entry
+		[[nodiscard]] UScriptStruct* GetValueType() const;
+
+		// Returns the name/value map from the current entry
+		[[nodiscard]] MapType GetMap() const;
+
+	private:
+		// Current name/value map entry
+		EntryPtrType EntryPtr = nullptr;
+
+		// One past the last name/value map entry
+		EntryPtrType EntryEndPtr = nullptr;
+	};
+
 	/*
 	 * Attribute Collection
 	 *
@@ -73,7 +107,17 @@ namespace UE::UAF
 			FAttributeSetMap* Map = nullptr;
 		};
 
+		struct FMapEntry
+		{
+			// The value type (key) that identifies this entry
+			UScriptStruct* ValueType = nullptr;
+
+			// The map containing our values
+			FAttributeMap* Map = nullptr;
+		};
+
 	public:
+		UE_API FAttributeCollection(const FAttributeNamedSetPtr& NamedSet, FReallocFun InReallocFun);
 		UE_API FAttributeCollection(const FAttributeCollection& Other);
 		UE_API FAttributeCollection(FAttributeCollection&& Other);
 		UE_API ~FAttributeCollection();
@@ -87,6 +131,12 @@ namespace UE::UAF
 		// Returns the named set this collection maps to
 		[[nodiscard]] UE_API const FAttributeNamedSetPtr& GetNamedSet() const;
 
+		// Resets the collection with the specified named set
+		UE_API void Reset(const FAttributeNamedSetPtr& NamedSet);
+
+		// Returns the allocator function this collection was initialized with
+		[[nodiscard]] UE_API FReallocFun GetAllocator() const;
+
 		// Returns the number of set maps contained within
 		[[nodiscard]] UE_API int32 NumSetMaps() const;
 
@@ -98,6 +148,14 @@ namespace UE::UAF
 		// Returns nullptr if we fail to add an entry (e.g. duplicate mapping)
 		[[nodiscard]] UE_API FAttributeSetMap* AddSetMap(const FAttributeMappingKey& MappingKey);
 
+		// Adds a new set mapping to this collection.
+		// Returns false if we fail to add an entry (e.g. duplicate mapping, mismatched allocator)
+		UE_API bool AddSetMap(FAttributeSetMap* Map);
+
+		// Appends a new set mapping to this collection (the input must belong at the end of the sorted set maps).
+		// Returns false if we fail to add an entry (e.g. duplicate mapping, mismatched allocator, not sorted)
+		UE_API bool AppendSetMap(FAttributeSetMap* Map);
+
 		// Removes the specified value set mapping from this collection
 		// Returns whether or not the entry was successfully removed
 		UE_API bool RemoveSetMap(FAttributeSetMap* Map);
@@ -138,11 +196,68 @@ namespace UE::UAF
 		[[nodiscard]] TSetMapIterator CreateSetMapIterator();
 		[[nodiscard]] TConstSetMapIterator CreateConstSetMapIterator() const;
 
-	protected:
-		UE_API FAttributeCollection(const FAttributeNamedSetPtr& NamedSet, FReallocFun InReallocFun);
+		// Returns the number of name/value maps contained within
+		[[nodiscard]] UE_API int32 NumMaps() const;
+
+		// Constructs and adds a new name/value map to this collection.
+		// Returns nullptr if we fail to add an entry (e.g. duplicate mapping)
+		template <class ValueType> [[nodiscard]] TAttributeMap<ValueType>* AddMap();
+
+		// Constructs and adds a new name/value mapping to this collection.
+		// Returns nullptr if we fail to add an entry (e.g. duplicate mapping)
+		[[nodiscard]] UE_API FAttributeMap* AddMap(UScriptStruct* ValueType);
+
+		// Adds a new name/value mapping to this collection.
+		// Returns false if we fail to add an entry (e.g. duplicate mapping, mismatched allocator)
+		UE_API bool AddMap(FAttributeMap* Map);
+
+		// Appends a new name/value mapping to this collection (the input must belong at the end of the sorted maps).
+		// Returns false if we fail to add an entry (e.g. duplicate mapping, mismatched allocator, not sorted)
+		UE_API bool AppendMap(FAttributeMap* Map);
+
+		// Removes the specified name/value mapping from this collection
+		// Returns whether or not the entry was successfully removed
+		UE_API bool RemoveMap(FAttributeMap* Map);
+
+		// Removes the specified name/value mapping from this collection
+		// Returns whether or not the entry was successfully removed
+		UE_API bool RemoveMap(UScriptStruct* ValueType);
+
+		// Finds the name/value mapping of the specified ValueType in this collection
+		// Returns nullptr if this collection does not contain such a mapping
+		template <class ValueType> [[nodiscard]] TAttributeMap<ValueType>* FindMap();
+
+		// Finds the name/value mapping of the specified ValueType in this collection
+		// Returns nullptr if this collection does not contain such a mapping
+		template <class ValueType> [[nodiscard]] const TAttributeMap<ValueType>* FindMap() const;
+
+		// Finds the name/value mapping of the specified ValueType in this collection
+		// Returns nullptr if this collection does not contain such a mapping
+		[[nodiscard]] UE_API FAttributeMap* FindMap(UScriptStruct* ValueType);
+
+		// Finds the name/value mapping of the specified ValueType in this collection
+		// Returns nullptr if this collection does not contain such a mapping
+		[[nodiscard]] UE_API const FAttributeMap* FindMap(UScriptStruct* ValueType) const;
+
+		// Finds or adds the name/value mapping of the specified ValueType in this collection
+		// Returns nullptr if we fail to add an entry (e.g. duplicate mapping)
+		template <class ValueType> [[nodiscard]] TAttributeMap<ValueType>* FindOrAddMap();
+
+		// Finds or adds the name/value mapping of the specified ValueType in this collection
+		// Returns nullptr if we fail to add an entry (e.g. duplicate mapping)
+		[[nodiscard]] UE_API FAttributeMap* FindOrAddMap(UScriptStruct* ValueType);
+
+		// Iterators
+		using TMapIterator = TAttributeCollectionMapIterator<FAttributeCollection>;
+		using TConstMapIterator = TAttributeCollectionMapIterator<const FAttributeCollection>;
+
+		// Returns an iterator over all name/value maps contained within, sorted by their key (value type)
+		[[nodiscard]] TMapIterator CreateMapIterator();
+		[[nodiscard]] TConstMapIterator CreateConstMapIterator() const;
 
 	private:
 		UE_API FAttributeSetMap** AddSetMapImpl(const FAttributeSetMap::FConstructArgs& Args, bool bConstructSetMap);
+		UE_API FAttributeMap** AddMapImpl(UScriptStruct* ValueType, bool bConstructMap);
 
 		// The named set we are based on, every set mapping uses a typed set within it
 		FAttributeNamedSetPtr NamedSet;
@@ -156,11 +271,22 @@ namespace UE::UAF
 		// How many set mappings we can contain within this collection
 		int32 SetMapCapacity = 0;
 
+		// An array of name/value mappings (sorted by value type)
+		FMapEntry* Maps = nullptr;
+
+		// How many mappings are contained within this collection
+		int32 MapCount = 0;
+
+		// How many mappings we can contain within this collection
+		int32 MapCapacity = 0;
+
 		// The indirect allocator function pointer
 		FReallocFun ReallocFun = nullptr;
 
 		friend TSetMapIterator;
 		friend TConstSetMapIterator;
+		friend TMapIterator;
+		friend TConstMapIterator;
 	};
 
 	// A heap allocated attribute collection
@@ -198,6 +324,21 @@ namespace UE::UAF
 		FAttributeSetMap* GetValue(const FAttributeCollection::TSetMapIterator& It) const { return It.GetMap(); }
 		const FAttributeSetMap* GetValue(const FAttributeCollection::TConstSetMapIterator& It) const { return It.GetMap(); }
 	};
+
+	// An operator structure to join attribute collections using InnerJoin/OuterJoin/etc
+	struct FAttributeCollectionMapJoinOp
+	{
+		// Key sorting uses signed comparison, use the largest signed value regardless of pointer size
+		static constexpr intptr_t LARGEST_KEY_VALUE = static_cast<intptr_t>(static_cast<uintptr_t>(~0ULL) >> 1);
+
+		UScriptStruct* GetLargestKey() const { return reinterpret_cast<UScriptStruct*>(LARGEST_KEY_VALUE); }
+
+		UScriptStruct* GetKey(const FAttributeCollection::TMapIterator& It) const { return It.GetValueType(); }
+		UScriptStruct* GetKey(const FAttributeCollection::TConstMapIterator& It) const { return It.GetValueType(); }
+
+		FAttributeMap* GetValue(const FAttributeCollection::TMapIterator& It) const { return It.GetMap(); }
+		const FAttributeMap* GetValue(const FAttributeCollection::TConstMapIterator& It) const { return It.GetMap(); }
+	};
 }
 
 /*
@@ -274,7 +415,7 @@ namespace UE::UAF
 			return nullptr;
 		}
 
-		return static_cast<TAttributeSetMap<ValueType>*>(FindSetMap(MappingKey));
+		return Cast<ValueType>(FindSetMap(MappingKey));
 	}
 
 	template <class ValueType>
@@ -286,7 +427,7 @@ namespace UE::UAF
 			return nullptr;
 		}
 
-		return static_cast<const TAttributeSetMap<ValueType>*>(FindSetMap(MappingKey));
+		return Cast<ValueType>(FindSetMap(MappingKey));
 	}
 
 	template <class ValueType, typename... TArgs>
@@ -300,7 +441,7 @@ namespace UE::UAF
 
 		if (FAttributeSetMap* SetMap = FindSetMap(MappingKey))
 		{
-			return static_cast<TAttributeSetMap<ValueType>*>(SetMap);
+			return Cast<ValueType>(SetMap);
 		}
 
 		return AddSetMap<ValueType>(MappingKey, Args...);
@@ -316,6 +457,67 @@ namespace UE::UAF
 		return TConstSetMapIterator(*this);
 	}
 
+	template <class ValueType>
+	inline TAttributeMap<ValueType>* FAttributeCollection::AddMap()
+	{
+		if (!NamedSet)
+		{
+			// We are not initialized
+			return nullptr;
+		}
+
+		UScriptStruct* ValueTypeStruct = ValueType::StaticStruct();
+
+		// Don't construct the map internally, it'll have to use indirection to do so and here we know the concrete type
+		constexpr bool bConstructMap = false;
+
+		FAttributeMap** MapPtr = AddMapImpl(ValueTypeStruct, bConstructMap);
+		if (MapPtr == nullptr)
+		{
+			return nullptr;
+		}
+
+		TAttributeMap<ValueType>* Map = MakeAttributeMap<ValueType>(ReallocFun);
+
+		check(*MapPtr == nullptr);
+		*MapPtr = Map;
+
+		return Map;
+	}
+
+	template <class ValueType>
+	inline TAttributeMap<ValueType>* FAttributeCollection::FindMap()
+	{
+		return Cast<ValueType>(FindMap(ValueType::StaticStruct()));
+	}
+
+	template <class ValueType>
+	inline const TAttributeMap<ValueType>* FAttributeCollection::FindMap() const
+	{
+		return Cast<ValueType>(FindMap(ValueType::StaticStruct()));
+	}
+
+	template <class ValueType>
+	inline TAttributeMap<ValueType>* FAttributeCollection::FindOrAddMap()
+	{
+		if (TAttributeMap<ValueType>* Map = FindMap<ValueType>())
+		{
+			return Map;
+		}
+
+		return AddMap<ValueType>();
+	}
+
+	inline typename FAttributeCollection::TMapIterator FAttributeCollection::CreateMapIterator()
+	{
+		return TMapIterator(*this);
+	}
+
+	inline typename FAttributeCollection::TConstMapIterator FAttributeCollection::CreateConstMapIterator() const
+	{
+		return TConstMapIterator(*this);
+	}
+
 	inline FAttributeCollectionHeap::FAttributeCollectionHeap()
 		: FAttributeCollection(nullptr, &FAllocatorTypeTrait<FDefaultAllocator>::Realloc)
 	{
@@ -379,6 +581,38 @@ namespace UE::UAF
 	{
 		return EntryPtr->Map;
 	}
+
+	template<class ContainerType>
+	inline TAttributeCollectionMapIterator<ContainerType>::TAttributeCollectionMapIterator(ContainerType& Collection)
+		: EntryPtr(Collection.Maps)
+		, EntryEndPtr(Collection.Maps + Collection.MapCount)
+	{
+	}
+
+	template<class ContainerType>
+	inline TAttributeCollectionMapIterator<ContainerType>& TAttributeCollectionMapIterator<ContainerType>::operator++()
+	{
+		++EntryPtr;
+		return *this;
+	}
+
+	template<class ContainerType>
+	inline TAttributeCollectionMapIterator<ContainerType>::operator bool() const
+	{
+		return EntryPtr < EntryEndPtr;
+	}
+
+	template<class ContainerType>
+	inline UScriptStruct* TAttributeCollectionMapIterator<ContainerType>::GetValueType() const
+	{
+		return EntryPtr->ValueType;
+	}
+
+	template<class ContainerType>
+	inline typename TAttributeCollectionMapIterator<ContainerType>::MapType TAttributeCollectionMapIterator<ContainerType>::GetMap() const
+	{
+		return EntryPtr->Map;
+	}
 }
 
 #undef UE_API
diff --git a/Engine/Plugins/Experimental/UAF/UAF/Source/UAF/Public/UAF/Attributes/AttributeMap.h b/Engine/Plugins/Experimental/UAF/UAF/Source/UAF/Public/UAF/Attributes/AttributeMap.h
new file mode 100644
--- /dev/null
+++ b/Engine/Plugins/Experimental/UAF/UAF/Source/UAF/Public/UAF/Attributes/AttributeMap.h
@@ -0,0 +1,813 @@
+// Copyright Epic Games, Inc. All Rights Reserved.
+
+#pragma once
+
+#include "CoreMinimal.h"
+#include "Algo/BinarySearch.h"
+#include "Algo/Sort.h"
+#include "Templates/MemoryOps.h"
+#include "UAF/Attributes/IndirectAllocator.h"
+
+#define UE_API UAF_API
+
+namespace UE::UAF
+{
+	// Untyped base class for TAttributeMap
+	class FAttributeMap
+	{
+	public:
+		UE_NONCOPYABLE(FAttributeMap);	// We disallow copy/move semantics
+
+		// Returns whether or not this map has any values
+		[[nodiscard]] UE_API bool IsEmpty() const;
+
+		// Returns the number of mapped values
+		[[nodiscard]] UE_API int32 Num() const;
+
+		// Returns the maximum number of values we can map
+		[[nodiscard]] UE_API int32 Max() const;
+
+		// Returns the mapped value type
+		[[nodiscard]] UE_API UScriptStruct* GetValueType() const;
+
+		// Returns the total allocated size for this map
+		[[nodiscard]] UE_API int32 GetAllocatedSize() const;
+
+		// Returns the allocator function this map was initialized with
+		[[nodiscard]] UE_API FReallocFun GetAllocator() const;
+
+		// Returns a copy of this map using the specified allocator
+		[[nodiscard]] virtual FAttributeMap* Duplicate(FReallocFun InReallocFun) const = 0;
+
+		// Copies the content of this map into another
+		// Both must have matching value types
+		virtual void CopyTo(FAttributeMap& Other) const = 0;
+
+		// Moves the content of this map to another
+		// Both must have matching value types
+		// The content of this map after the move is undefined, just like with move assignment
+		UE_API virtual void MoveTo(FAttributeMap& Other);
+
+		// Finds a mapping and returns its index or INDEX_NONE if not found, O(logN)
+		[[nodiscard]] int32 IndexOf(FName Name) const;
+
+		// Returns whether or not this map contains the specified mapped value, O(logN)
+		[[nodiscard]] bool Contains(FName Name) const;
+
+		// Retrieves a mapping name by index, O(1)
+		[[nodiscard]] FName GetName(int32 Index) const;
+
+		// Adds a new mapping, sets the value using the provided setter, O(logN)
+		// Returns whether or not we inserted (we can fail if a value with that name already exists)
+		// The value setter (and/or caller) is responsible for validating the value type
+		virtual bool AddWithSetter(FName Name, const TFunctionRef<void(UScriptStruct* ValueType, uint8* OutValuePtr)>& ValueSetter) = 0;
+
+		// Gets a value at the specified index using a value getter function, O(1)
+		// The value getter (and/or caller) is responsible for validating the value type
+		virtual void GetValueWithGetter(int32 Index, const TFunctionRef<void(UScriptStruct* ValueType, const uint8* ValuePtr)>& ValueGetter) const = 0;
+
+	protected:
+		// Default construction for an empty map
+		UE_API FAttributeMap();
+
+		// Constructs a map but does not allocate its memory, derived types are responsible for it
+		UE_API FAttributeMap(UScriptStruct* ValueType, FReallocFun ReallocFun, int32 ContainerSize);
+
+		// Destroys the map and frees its memory
+		UE_API virtual ~FAttributeMap();
+
+		// Returns a pointer to the start of the names
+		[[nodiscard]] FName* GetNames();
+		[[nodiscard]] const FName* GetNames() const;
+
+		// The type of values this map contains
+		UScriptStruct* ValueTypeStruct = nullptr;
+
+		// The number of values mapped in our buffer
+		int32 NumValues = 0;
+
+		// The maximum number of values that we can map in our buffer
+		int32 MaxValues = 0;
+
+		// The base untyped pointer for our names/values
+		// Names and values are allocated together within the same contiguous memory buffer
+		// Names are first within the buffer followed by value types using default alignment
+		uint8* DataPtr = nullptr;
+
+		// The indirect allocator function pointer
+		FReallocFun ReallocFun = nullptr;
+
+		// The size in bytes of this container (excluding any externally allocated data)
+		int32 ContainerSize = 0;
+
+		friend class FAttributeCollection;
+		friend class FAttributeRuntimeRegistry;
+		friend void ReleaseAttributeMap(FAttributeMap* Map);
+	};
+
+	// A iterator that iterates over all name/value pairs within an attribute map
+	template<class ContainerType>
+	class TAttributeMapIterator
+	{
+	public:
+		static constexpr bool bIsContainerConst = std::is_const_v<ContainerType>;
+
+		using NamePtrType = typename std::conditional<bIsContainerConst, const FName*, FName*>::type;
+		using ValuePtrType = typename std::conditional<bIsContainerConst, const typename ContainerType::ValueType*, typename ContainerType::ValueType*>::type;
+
+		// Creates an iterator over the name/value map entries within the specified attribute map
+		[[nodiscard]] TAttributeMapIterator(ContainerType& Map);
+
+		// Increments and moves the iterator to the next name/value pair
+		TAttributeMapIterator& operator++();
+
+		// Returns whether or not the iterator still contains values
+		[[nodiscard]] explicit operator bool() const;
+
+		// Returns the name of the current name/value pair
+		[[nodiscard]] FName GetName() const;
+
+		// Returns the value of the current name/value pair
+		[[nodiscard]] ValuePtrType GetValue() const;
+
+	private:
+		// Current name/value pair
+		NamePtrType NamePtr = nullptr;
+		ValuePtrType ValuePtr = nullptr;
+
+		// One past the last name/value pair
+		NamePtrType NameEndPtr = nullptr;
+	};
+
+	// Encapsulates a mapping of FName to value type
+	// Value types must be a UStruct
+	template<class InValueType>
+	class TAttributeMap : public FAttributeMap
+	{
+	public:
+		// The type of values held in this map
+		using ValueType = InValueType;
+
+		// Returns a copy of this map using the specified allocator
+		[[nodiscard]] virtual FAttributeMap* Duplicate(FReallocFun InReallocFun) const override;
+
+		// Copies the content of this map into another
+		// Both must have matching value types
+		virtual void CopyTo(FAttributeMap& Other) const override;
+
+		// Adds a new mapping, sets the value using the provided setter, O(logN)
+		// Returns whether or not we inserted (we can fail if a value with that name already exists)
+		// The value setter (and/or caller) is responsible for validating the value type
+		virtual bool AddWithSetter(FName Name, const TFunctionRef<void(UScriptStruct* ValueType, uint8* OutValuePtr)>& ValueSetter) override;
+
+		// Gets a value at the specified index using a value getter function, O(1)
+		// The value getter (and/or caller) is responsible for validating the value type
+		virtual void GetValueWithGetter(int32 Index, const TFunctionRef<void(UScriptStruct* ValueType, const uint8* ValuePtr)>& ValueGetter) const override;
+
+		// Adds a new mapping, O(logN)
+		// Returns whether or not we inserted (we can fail if a value with that name already exists)
+		bool Add(FName Name, const InValueType& Value);
+
+		// Adds a new mapping at the end of the map, must be sorted, O(1)
+		// Returns whether or not we inserted (we can fail if the value isn't sorted to land at the end of the map)
+		bool Append(FName Name, const InValueType& Value);
+
+		// Removes a mapping, O(logN)
+		// Returns whether or not we removed the value (we can fail if the name isn't found)
+		bool Remove(FName Name);
+
+		// Removes a mapping, O(1)
+		void RemoveAt(int32 Index);
+
+		// Finds a mapping and returns a pointer to it or nullptr if not found, O(logN)
+		[[nodiscard]] InValueType* Find(FName Name);
+		[[nodiscard]] const InValueType* Find(FName Name) const;
+
+		// Retrieves a mapping value by index, O(1)
+		[[nodiscard]] InValueType& GetValue(int32 Index);
+		[[nodiscard]] const InValueType& GetValue(int32 Index) const;
+
+		// Initializes the mapping using sorted inputs, O(N)
+		template<class GetNamePredicateType, class GetValuePredicateType>
+		void SetSorted(int32 NumValues, GetNamePredicateType GetNamePredicate, GetValuePredicateType GetValuePredicate);
+
+		// Initializes the mapping using unsorted inputs, O(NlogN)
+		template<class GetNamePredicateType, class GetValuePredicateType>
+		void SetUnsorted(int32 NumValues, GetNamePredicateType GetNamePredicate, GetValuePredicateType GetValuePredicate);
+
+		// Iterators
+		using TIterator = TAttributeMapIterator<TAttributeMap<InValueType>>;
+		using TConstIterator = TAttributeMapIterator<const TAttributeMap<InValueType>>;
+
+		// Returns an iterator over all name/value pairs contained within, sorted by their name
+		[[nodiscard]] TIterator CreateIterator();
+		[[nodiscard]] TConstIterator CreateConstIterator() const;
+
+	protected:
+		static const int32 ValueTypeAlignment = sizeof(InValueType) >= 16 ? 16 : 8;
+
+		// Constructs a map
+		TAttributeMap(UScriptStruct* ValueType, FReallocFun ReallocFun, int32 ContainerSize);
+
+		// Destroys the map and frees its memory
+		virtual ~TAttributeMap() override;
+
+		void Grow();
+
+		// Returns a pointer to the start of the values
+		InValueType* GetValues();
+		const InValueType* GetValues() const;
+
+		template<class T>
+		friend TAttributeMap<T>* MakeAttributeMap(FReallocFun ReallocFun);
+		friend TIterator;
+		friend TConstIterator;
+	};
+
+	// An operator structure to join attribute maps using InnerJoin/OuterJoin/etc
+	template<class ValueType>
+	struct FAttributeMapJoinOp
+	{
+		// We generally want the key in our predicate since the values by themselves have little meaning
+		static constexpr bool HasPredicateWithKey() { return true; }
+
+		bool IsLessThan(FName LHS, FName RHS) const
+		{
+			return LHS.FastLess(RHS);
+		}
+
+		// Special dummy FName that is larger than all other FNames in fast sorted order
+		FName GetLargestKey() const { return FName(FNameEntryId::FromUnstableInt(~0), FNameEntryId::FromUnstableInt(~0), 0x7fffffff); }
+
+		FName GetKey(const typename TAttributeMap<ValueType>::TIterator& It) const { return It.GetName(); }
+		FName GetKey(const typename TAttributeMap<ValueType>::TConstIterator& It) const { return It.GetName(); }
+
+		ValueType* GetValue(const typename TAttributeMap<ValueType>::TIterator& It) const { return It.GetValue(); }
+		const ValueType* GetValue(const typename TAttributeMap<ValueType>::TConstIterator& It) const { return It.GetValue(); }
+	};
+
+	// Allocates and constructs a new attribute map for the specified value type
+	template<class ValueType>
+	TAttributeMap<ValueType>* MakeAttributeMap(FReallocFun ReallocFun);
+
+	// Releases and destroys the specified attribute map
+	void ReleaseAttributeMap(FAttributeMap* Map);
+
+	// Cast from an untyped name/value map to a typed name/value map
+	template<class ValueType>
+	UE_FORCEINLINE_HINT TAttributeMap<ValueType>* Cast(FAttributeMap* Src)
+	{
+		return Src != nullptr && ValueType::StaticStruct() == Src->GetValueType() ? static_cast<TAttributeMap<ValueType>*>(Src) : nullptr;
+	}
+
+	// Cast from an untyped name/value map to a typed name/value map
+	template<class ValueType>
+	UE_FORCEINLINE_HINT const TAttributeMap<ValueType>* Cast(const FAttributeMap* Src)
+	{
+		return Src != nullptr && ValueType::StaticStruct() == Src->GetValueType() ? static_cast<const TAttributeMap<ValueType>*>(Src) : nullptr;
+	}
+
+	// Cast from an untyped name/value map to a typed name/value map with an assert if the types mismatch
+	template<class ValueType>
+	inline TAttributeMap<ValueType>* CastChecked(FAttributeMap* Src)
+	{
+		TAttributeMap<ValueType>* CastResult = Src != nullptr && ValueType::StaticStruct() == Src->GetValueType() ? static_cast<TAttributeMap<ValueType>*>(Src) : nullptr;
+		check(CastResult);
+		return CastResult;
+	}
+
+	// Cast from an untyped name/value map to a typed name/value map with an assert if the types mismatch
+	template<class ValueType>
+	inline const TAttributeMap<ValueType>* CastChecked(const FAttributeMap* Src)
+	{
+		const TAttributeMap<ValueType>* CastResult = Src != nullptr && ValueType::StaticStruct() == Src->GetValueType() ? static_cast<const TAttributeMap<ValueType>*>(Src) : nullptr;
+		check(CastResult);
+		return CastResult;
+	}
+
+	//////////////////////////////////////////////////////////////////////////
+	// Implementation
+
+	inline int32 FAttributeMap::IndexOf(FName Name) const
+	{
+		const FName* Names = GetNames();
+
+		const int32 Index = Algo::LowerBound(TArrayView<const FName>(Names, NumValues), Name, FNameFastLess());
+		if (Index >= NumValues || Names[Index] != Name)
+		{
+			// No value is mapped to that name
+			return INDEX_NONE;
+		}
+
+		return Index;
+	}
+
+	inline bool FAttributeMap::Contains(FName Name) const
+	{
+		return IndexOf(Name) != INDEX_NONE;
+	}
+
+	inline FName FAttributeMap::GetName(int32 Index) const
+	{
+		checkf((Index >= 0) & (Index < NumValues), TEXT("Map index out of bounds: %lld into an map of size %lld"), (long long)Index, (long long)NumValues); // & for one branch
+		return GetNames()[Index];
+	}
+
+	inline FName* FAttributeMap::GetNames()
+	{
+		return reinterpret_cast<FName*>(DataPtr);
+	}
+
+	inline const FName* FAttributeMap::GetNames() const
+	{
+		return reinterpret_cast<const FName*>(DataPtr);
+	}
+
+	template<class InValueType>
+	inline TAttributeMap<InValueType>::TAttributeMap(UScriptStruct* ValueTypeStruct, FReallocFun InReallocFun, int32 InContainerSize)
+		: FAttributeMap(ValueTypeStruct, InReallocFun, InContainerSize)
+	{
+		check(ValueTypeStruct == InValueType::StaticStruct());
+		check(InContainerSize == sizeof(TAttributeMap<InValueType>));
+	}
+
+	template<class InValueType>
+	inline TAttributeMap<InValueType>::~TAttributeMap()
+	{
+		if (DataPtr != nullptr)
+		{
+			// FNames don't need to be destroyed, only destroy the values
+			DestructItems<InValueType>(GetValues(), NumValues);
+
+			DataPtr = nullptr;
+			NumValues = MaxValues = 0;
+		}
+	}
+
+	template<class InValueType>
+	inline FAttributeMap* TAttributeMap<InValueType>::Duplicate(FReallocFun InReallocFun) const
+	{
+		TAttributeMap<InValueType>* Copy = MakeAttributeMap<InValueType>(InReallocFun);
+
+		if (NumValues != 0)
+		{
+			check(Copy->MaxValues == 0 && Copy->DataPtr == nullptr);
+
+			const int32 NewSize = Align(sizeof(FName) * NumValues, ValueTypeAlignment) + (sizeof(InValueType) * NumValues);
+
+			Copy->DataPtr = (*InReallocFun)(nullptr, 0, NewSize);
+			Copy->NumValues = NumValues;
+			Copy->MaxValues = NumValues;
+
+			ConstructItems<FName>(Copy->GetNames(), GetNames(), NumValues);
+			ConstructItems<InValueType>(Copy->GetValues(), GetValues(), NumValues);
+		}
+
+		return Copy;
+	}
+
+	template<class InValueType>
+	inline void TAttributeMap<InValueType>::CopyTo(FAttributeMap& Other) const
+	{
+		checkf(ValueTypeStruct == Other.GetValueType(), TEXT("Value types must match to allow copying"));
+
+		TAttributeMap<InValueType>& Copy = reinterpret_cast<TAttributeMap<InValueType>&>(Other);
+
+		// Destroy old values
+		if (Copy.NumValues != 0)
+		{
+			// FNames don't need to be destroyed, only destroy the values
+			DestructItems<InValueType>(Copy.GetValues(), Copy.NumValues);
+		}
+
+		// Resize
+		if (Copy.MaxValues != NumValues)
+		{
+			const int32 OldSize = Align(sizeof(FName) * Copy.MaxValues, ValueTypeAlignment) + (sizeof(InValueType) * Copy.MaxValues);
+			const int32 NewSize = Align(sizeof(FName) * NumValues, ValueTypeAlignment) + (sizeof(InValueType) * NumValues);
+
+			Copy.DataPtr = (*Copy.ReallocFun)(Copy.DataPtr, OldSize, NewSize);
+			Copy.MaxValues = NumValues;
+		}
+
+		// Copy construct our new values
+		if (NumValues != 0)
+		{
+			ConstructItems<FName>(Copy.GetNames(), GetNames(), NumValues);
+			ConstructItems<InValueType>(Copy.GetValues(), GetValues(), NumValues);
+		}
+
+		Copy.NumValues = NumValues;
+	}
+
+	template<class InValueType>
+	inline bool TAttributeMap<InValueType>::AddWithSetter(FName Name, const TFunctionRef<void(UScriptStruct* ValueType, uint8* OutValuePtr)>& ValueSetter)
+	{
+		FName* Names = GetNames();
+
+		const int32 Index = Algo::LowerBound(TArrayView<const FName>(Names, NumValues), Name, FNameFastLess());
+		if (Index < NumValues && Names[Index] == Name)
+		{
+			// A mapping already exists that matches the specified name
+			return false;
+		}
+
+		if (NumValues >= MaxValues)
+		{
+			// Not enough space, grow our buffer
+			Grow();
+
+			Names = GetNames();
+		}
+
+		InValueType* Values = GetValues();
+
+		// Insert our value at the right position, making space for it
+		const int32 NumToMove = NumValues - Index;
+		RelocateConstructItems<FName>(Names + Index + 1, Names + Index, NumToMove);
+		RelocateConstructItems<InValueType>(Values + Index + 1, Values + Index, NumToMove);
+		new(Names + Index) FName(Name);
+		ValueSetter(ValueTypeStruct, reinterpret_cast<uint8*>(Values + Index));
+		NumValues++;
+
+		return true;
+	}
+
+	template<class InValueType>
+	inline void TAttributeMap<InValueType>::GetValueWithGetter(int32 Index, const TFunctionRef<void(UScriptStruct* ValueType, const uint8* ValuePtr)>& ValueGetter) const
+	{
+		checkf((Index >= 0) & (Index < NumValues), TEXT("Map index out of bounds: %lld into an map of size %lld"), (long long)Index, (long long)NumValues); // & for one branch
+		const InValueType* ValuePtr = GetValues() + Index;
+		ValueGetter(ValueTypeStruct, reinterpret_cast<const uint8*>(ValuePtr));
+	}
+
+	template<class InValueType>
+	inline bool TAttributeMap<InValueType>::Add(FName Name, const InValueType& Value)
+	{
+		FName* Names = GetNames();
+
+		const int32 Index = Algo::LowerBound(TArrayView<const FName>(Names, NumValues), Name, FNameFastLess());
+		if (Index < NumValues && Names[Index] == Name)
+		{
+			// A mapping already exists that matches the specified name
+			return false;
+		}
+
+		if (NumValues >= MaxValues)
+		{
+			// Not enough space, grow our buffer
+			Grow();
+			
+			Names = GetNames();
+		}
+
+		InValueType* Values = GetValues();
+
+		// Insert our value at the right position, making space for it
+		const int32 NumToMove = NumValues - Index;
+		RelocateConstructItems<FName>(Names + Index + 1, Names + Index, NumToMove);
+		RelocateConstructItems<InValueType>(Values + Index + 1, Values + Index, NumToMove);
+		new(Names + Index) FName(Name);
+		new(Values + Index) InValueType(Value);
+		NumValues++;
+
+		return true;
+	}
+
+	template<class InValueType>
+	inline bool TAttributeMap<InValueType>::Append(FName Name, const InValueType& Value)
+	{
+		FName* Names = GetNames();
+
+		if (NumValues != 0 && !Names[NumValues - 1].FastLess(Name))
+		{
+			// The name doesn't belong at the end, it isn't sorted
+			return false;
+		}
+
+		if (NumValues >= MaxValues)
+		{
+			// Not enough space, grow our buffer
+			Grow();
+
+			Names = GetNames();
+		}
+
+		InValueType* Values = GetValues();
+
+		// Append our new value at the end
+		new(Names + NumValues) FName(Name);
+		new(Values + NumValues) InValueType(Value);
+		NumValues++;
+
+		return true;
+	}
+
+	template<class InValueType>
+	inline bool TAttributeMap<InValueType>::Remove(FName Name)
+	{
+		const FName* Names = GetNames();
+
+		const int32 Index = Algo::LowerBound(TArrayView<const FName>(Names, NumValues), Name, FNameFastLess());
+		if (Index >= NumValues || Names[Index] != Name)
+		{
+			// No value is mapped to that name
+			return false;
+		}
+
+		RemoveAt(Index);
+
+		return true;
+	}
+
+	template<class InValueType>
+	inline void TAttributeMap<InValueType>::RemoveAt(int32 Index)
+	{
+		checkf((Index >= 0) & (Index < NumValues), TEXT("Map index out of bounds: %lld into an map of size %lld"), (long long)Index, (long long)NumValues); // & for one branch
+
+		FName* Names = GetNames();
+		InValueType* Values = GetValues();
+
+		// FNames don't need to be destroyed, only destroy the values
+		DestructItem<InValueType>(Values + Index);
+
+		const int32 NumToMove = NumValues - Index - 1;
+		RelocateConstructItems<FName>(Names + Index, Names + Index + 1, NumToMove);
+		RelocateConstructItems<InValueType>(Values + Index, Values + Index + 1, NumToMove);
+		NumValues--;
+	}
+
+	template<class InValueType>
+	inline InValueType* TAttributeMap<InValueType>::Find(FName Name)
+	{
+		const FName* Names = GetNames();
+
+		const int32 Index = Algo::LowerBound(TArrayView<const FName>(Names, NumValues), Name, FNameFastLess());
+		if (Index >= NumValues || Names[Index] != Name)
+		{
+			// No value is mapped to that name
+			return nullptr;
+		}
+
+		return GetValues() + Index;
+	}
+
+	template<class InValueType>
+	inline const InValueType* TAttributeMap<InValueType>::Find(FName Name) const
+	{
+		const FName* Names = GetNames();
+
+		const int32 Index = Algo::LowerBound(TArrayView<const FName>(Names, NumValues), Name, FNameFastLess());
+		if (Index >= NumValues || Names[Index] != Name)
+		{
+			// No value is mapped to that name
+			return nullptr;
+		}
+
+		return GetValues() + Index;
+	}
+
+	template<class InValueType>
+	inline InValueType& TAttributeMap<InValueType>::GetValue(int32 Index)
+	{
+		checkf((Index >= 0) & (Index < NumValues), TEXT("Map index out of bounds: %lld into an map of size %lld"), (long long)Index, (long long)NumValues); // & for one branch
+
+		return GetValues()[Index];
+	}
+
+	template<class InValueType>
+	inline const InValueType& TAttributeMap<InValueType>::GetValue(int32 Index) const
+	{
+		checkf((Index >= 0) & (Index < NumValues), TEXT("Map index out of bounds: %lld into an map of size %lld"), (long long)Index, (long long)NumValues); // & for one branch
+
+		return GetValues()[Index];
+	}
+
+	template<class InValueType>
+	template<class GetNamePredicateType, class GetValuePredicateType>
+	inline void TAttributeMap<InValueType>::SetSorted(int32 InNumValues, GetNamePredicateType GetNamePredicate, GetValuePredicateType GetValuePredicate)
+	{
+		// Destroy old values
+		if (NumValues != 0)
+		{
+			// FNames don't need to be destroyed, only destroy the values
+			DestructItems<InValueType>(GetValues(), NumValues);
+		}
+
+		// Resize
+		if (MaxValues != InNumValues)
+		{
+			const int32 OldSize = Align(sizeof(FName) * MaxValues, ValueTypeAlignment) + (sizeof(InValueType) * MaxValues);
+			const int32 NewSize = Align(sizeof(FName) * InNumValues, ValueTypeAlignment) + (sizeof(InValueType) * InNumValues);
+
+			DataPtr = (*ReallocFun)(DataPtr, OldSize, NewSize);
+			MaxValues = InNumValues;
+		}
+
+		FName* Names = GetNames();
+		InValueType* Values = GetValues();
+
+		// Set now to make debugger easier
+		NumValues = InNumValues;
+
+		// Set our values
+		for (int32 Index = 0; Index < InNumValues; ++Index)
+		{
+			const FName Name = GetNamePredicate(Index);
+			checkf(Index == 0 || Names[Index - 1].FastLess(Name), TEXT("Attempting to initialize an attribute map from unsorted data"));
+
+			new(Names + Index) FName(Name);
+			new(Values + Index) InValueType(GetValuePredicate(Index));
+		}
+	}
+
+	template<class InValueType>
+	template<class GetNamePredicateType, class GetValuePredicateType>
+	inline void TAttributeMap<InValueType>::SetUnsorted(int32 InNumValues, GetNamePredicateType GetNamePredicate, GetValuePredicateType GetValuePredicate)
+	{
+		// Destroy old values
+		if (NumValues != 0)
+		{
+			// FNames don't need to be destroyed, only destroy the values
+			DestructItems<InValueType>(GetValues(), NumValues);
+		}
+
+		// Resize
+		if (MaxValues != InNumValues)
+		{
+			const int32 OldSize = Align(sizeof(FName) * MaxValues, ValueTypeAlignment) + (sizeof(InValueType) * MaxValues);
+			const int32 NewSize = Align(sizeof(FName) * InNumValues, ValueTypeAlignment) + (sizeof(InValueType) * InNumValues);
+
+			DataPtr = (*ReallocFun)(DataPtr, OldSize, NewSize);
+			MaxValues = InNumValues;
+		}
+
+		FName* Names = GetNames();
+		InValueType* Values = GetValues();
+
+		// Set now to make debugger easier
+		NumValues = InNumValues;
+
+#if 1
+		// TODO: We don't appear to have a version of Algo::Sort that works on iterators
+		// and so we can't sort our two arrays in tandem, use a temporary array instead
+		using FTmpElementType = TPair<FName, InValueType>;
+		TArray<FTmpElementType> TmpElements;
+		TmpElements.Reserve(InNumValues);
+		for (int32 Index = 0; Index < InNumValues; ++Index)
+		{
+			TmpElements.Emplace(FTmpElementType{ FName(GetNamePredicate(Index)), InValueType(GetValuePredicate(Index)) });
+		}
+
+		// Sort our temporary elements
+		Algo::Sort(TmpElements, [](const FTmpElementType& LHS, const FTmpElementType& RHS) { return LHS.Key.FastLess(RHS.Key); });
+
+		// Set our values
+		for (int32 Index = 0; Index < InNumValues; ++Index)
+		{
+			const FTmpElementType& Element = TmpElements[Index];
+
+			new(Names + Index) FName(Element.Key);
+			new(Values + Index) InValueType(Element.Value);
+		}
+#else
+		// Set our values
+		for (int32 Index = 0; Index < InNumValues; ++Index)
+		{
+			new(Names + Index) FName(GetNamePredicate(Index));
+			new(Values + Index) InValueType(GetValuePredicate(Index));
+		}
+
+		// Sort our values
+		//Algo::Sort(Elements, FElementSortPredicate());	// TODO
+#endif
+	}
+
+	template<class InValueType>
+	inline typename TAttributeMap<InValueType>::TIterator TAttributeMap<InValueType>::CreateIterator()
+	{
+		return TIterator(*this);
+	}
+
+	template<class InValueType>
+	inline typename TAttributeMap<InValueType>::TConstIterator TAttributeMap<InValueType>::CreateConstIterator() const
+	{
+		return TConstIterator(*this);
+	}
+
+	template<class InValueType>
+	FORCENOINLINE void TAttributeMap<InValueType>::Grow()
+	{
+		const bool bAllowQuantize = false;
+		const int32 NewCapacity = DefaultCalculateSlackGrow(MaxValues + 1, MaxValues, 1, bAllowQuantize);
+
+		const int32 OldValuesOffset = Align(sizeof(FName) * MaxValues, ValueTypeAlignment);
+		const int32 NewValuesOffset = Align(sizeof(FName) * NewCapacity, ValueTypeAlignment);
+		const int32 OldSize = OldValuesOffset + (sizeof(InValueType) * MaxValues);
+		const int32 NewSize = NewValuesOffset + (sizeof(InValueType) * NewCapacity);
+
+		// Our names and values live within the same buffer but when we re-size, they need to end up at
+		// different offsets since both halves grow. To avoid reallocate moving our names/values and
+		// then having another pass to offset our values, we handle it manually.
+
+		// Cache our old pointers
+		FName* OldNames = GetNames();
+		InValueType* OldValues = GetValues();
+
+		// Allocate our new buffer
+		uint8* NewDataPtr = (*ReallocFun)(nullptr, 0, NewSize);
+
+		// Cache our new pointers
+		FName* NewNames = reinterpret_cast<FName*>(NewDataPtr);
+		InValueType* NewValues = reinterpret_cast<InValueType*>(NewDataPtr + NewValuesOffset);
+
+		// Relocate our name/value pairs
+		RelocateConstructItems<FName>(NewNames, OldNames, MaxValues);
+		RelocateConstructItems<InValueType>(NewValues, OldValues, MaxValues);
+
+		// Deallocate our old buffer
+		(*ReallocFun)(DataPtr, OldSize, 0);
+
+		// Update our state
+		DataPtr = NewDataPtr;
+
+		MaxValues = NewCapacity;
+	}
+
+	template<class InValueType>
+	inline InValueType* TAttributeMap<InValueType>::GetValues()
+	{
+		const int32 ValuesOffset = Align(MaxValues * sizeof(FName), ValueTypeAlignment);
+		return reinterpret_cast<InValueType*>(DataPtr + ValuesOffset);
+	}
+
+	template<class InValueType>
+	inline const InValueType* TAttributeMap<InValueType>::GetValues() const
+	{
+		const int32 ValuesOffset = Align(MaxValues * sizeof(FName), ValueTypeAlignment);
+		return reinterpret_cast<const InValueType*>(DataPtr + ValuesOffset);
+	}
+
+	template<class ContainerType>
+	inline TAttributeMapIterator<ContainerType>::TAttributeMapIterator(ContainerType& Map)
+		: NamePtr(Map.GetNames())
+		, ValuePtr(Map.GetValues())
+		, NameEndPtr(NamePtr + Map.Num())
+	{
+	}
+
+	template<class ContainerType>
+	inline TAttributeMapIterator<ContainerType>& TAttributeMapIterator<ContainerType>::operator++()
+	{
+		++NamePtr;
+		++ValuePtr;
+		return *this;
+	}
+
+	template<class ContainerType>
+	inline TAttributeMapIterator<ContainerType>::operator bool() const
+	{
+		return NamePtr < NameEndPtr;
+	}
+
+	template<class ContainerType>
+	inline FName TAttributeMapIterator<ContainerType>::GetName() const
+	{
+		return *NamePtr;
+	}
+
+	template<class ContainerType>
+	inline typename TAttributeMapIterator<ContainerType>::ValuePtrType TAttributeMapIterator<ContainerType>::GetValue() const
+	{
+		return ValuePtr;
+	}
+
+	template<class ValueType>
+	inline TAttributeMap<ValueType>* MakeAttributeMap(FReallocFun ReallocFun)
+	{
+		const int32 ContainerSize = sizeof(TAttributeMap<ValueType>);
+
+		// Allocate our map container buffer
+		uint8* Ptr = (*ReallocFun)(nullptr, 0, ContainerSize);
+
+		// Construct our map over it
+		return new(Ptr) TAttributeMap<ValueType>(ValueType::StaticStruct(), ReallocFun, ContainerSize);
+	}
+
+	inline void ReleaseAttributeMap(FAttributeMap* Map)
+	{
+		if (Map == nullptr)
+		{
+			return;
+		}
+
+		const int32 ContainerSize = Map->ContainerSize;
+		const FReallocFun ReallocFun = Map->ReallocFun;
+
+		Map->~FAttributeMap();
+
+		(*ReallocFun)(reinterpret_cast<uint8*>(Map), ContainerSize, 0);
+	}
+}
+
+#undef UE_API
diff --git a/Engine/Plugins/Experimental/UAF/UAF/Source/UAF/Public/UAF/Attributes/AttributeRuntimeRegistry.h b/Engine/Plugins/Experimental/UAF/UAF/Source/UAF/Public/UAF/Attributes/AttributeRuntimeRegistry.h
--- a/Engine/Plugins/Experimental/UAF/UAF/Source/UAF/Public/UAF/Attributes/AttributeRuntimeRegistry.h
+++ b/Engine/Plugins/Experimental/UAF/UAF/Source/UAF/Public/UAF/Attributes/AttributeRuntimeRegistry.h
@@ -3,6 +3,7 @@
 #pragma once
 
 #include "CoreMinimal.h"
+#include "UAF/Attributes/AttributeMap.h"
 #include "UAF/Attributes/AttributeSetMap.h"
 #include "UAF/Attributes/Transformers/Transformer.h"
 
@@ -17,6 +18,7 @@ namespace UE::UAF
 		UE_NONCOPYABLE(FAttributeRuntimeRegistry);	// We disallow copy/move semantics
 
 		using SetMapConstructorFun = TFunction<FAttributeSetMap*(const FAttributeSetMap::FConstructArgs& Args)>;
+		using MapConstructorFun = TFunction<FAttributeMap*(FReallocFun ReallocFun)>;
 
 		// Returns the set map registry instance
 		[[nodiscard]] static UE_API FAttributeRuntimeRegistry& Get();
@@ -24,6 +26,9 @@ namespace UE::UAF
 		// Constructs a set map from its UStruct value type
 		[[nodiscard]] UE_API FAttributeSetMap* ConstructSetMap(const FAttributeSetMap::FConstructArgs& Args) const;
 
+		// Constructs a name/value map from its UStruct value type
+		[[nodiscard]] UE_API FAttributeMap* ConstructMap(UScriptStruct* ValueType, FReallocFun ReallocFun) const;
+
 		// Returns the map of known attribute value transformers
 		[[nodiscard]] UE_API FAttributeValueTransformerMapPtr GetTransformerMap() const;
 
@@ -31,13 +36,19 @@ namespace UE::UAF
 		template<class ValueType>
 		void RegisterSetMapInitializer();
 
+		// Registers a name/value map of the specified value type
+		template<class ValueType>
+		void RegisterMapInitializer();
+
 		template<class TransformerSpecializationType>
 		bool RegisterAttributeValueTransformer();
 
 	private:
 		UE_API FAttributeRuntimeRegistry();
 		UE_API void RegisterSetMapInitializer(UScriptStruct* ValueType, SetMapConstructorFun Initializer);
-		UE_API bool RegisterAttributeValueTransformer(FName TransformerName, UScriptStruct* ValueType, FRawTransformerFunc TransformerFunc);
+		UE_API void RegisterMapInitializer(UScriptStruct* ValueType, MapConstructorFun Initializer);
+		UE_API bool RegisterAttributeValueSetMapTransformer(FName TransformerName, UScriptStruct* ValueType, FRawTransformerFunc TransformerFunc);
+		UE_API bool RegisterAttributeValueMapTransformer(FName TransformerName, UScriptStruct* ValueType, FRawTransformerFunc TransformerFunc);
 
 		// Initialize/destroy the global registry
 		static UE_API void Init();
@@ -48,6 +59,11 @@ namespace UE::UAF
 		TMap<UScriptStruct*, SetMapConstructorFun> SetMapInitializers;
 		mutable FRWLock SetMapInitializerLock;
 
+		// A map of ValueType -> Map constructors
+		// This is used for indirect construction when we have a UScriptStruct value type
+		TMap<UScriptStruct*, MapConstructorFun> MapInitializers;
+		mutable FRWLock MapInitializerLock;
+
 		// A shared pointer to a map of transformer name -> transformer list (a transformer per value type)
 		FAttributeValueTransformerMapPtr TransformerMap;
 		mutable FRWLock TransformerMapLock;
@@ -69,35 +85,74 @@ namespace UE::UAF
 			});
 	}
 
+	template<class ValueType>
+	inline void FAttributeRuntimeRegistry::RegisterMapInitializer()
+	{
+		RegisterMapInitializer(ValueType::StaticStruct(), [](FReallocFun ReallocFun)
+			{
+				return MakeAttributeMap<ValueType>(ReallocFun);
+			});
+	}
+
 	template<class TransformerSpecializationType>
 	inline bool FAttributeRuntimeRegistry::RegisterAttributeValueTransformer()
 	{
-		using TransformerType = typename TransformerSpecializationType::TransformerType;
+		using TransformerCppType = typename TransformerSpecializationType::TransformerType;
 		using ValueType = typename TransformerSpecializationType::ValueType;
-		using FTransformFunc = typename TransformerType::FTransformFunc;
 
-		static_assert(std::is_base_of_v<FAttributeValueTransformer, TransformerType>, "Specified transformer must derive from FAttributeValueTransformer and provide the necessary API");
+		static_assert(std::is_base_of_v<FAttributeValueTransformer, TransformerCppType>, "Specified transformer must derive from FAttributeValueTransformer and provide the necessary API");
 
-		const FName TransformerName = TransformerType::TransformerName;
+		const FName TransformerName = TransformerCppType::TransformerName;
 		UScriptStruct* ValueTypeStruct = ValueType::StaticStruct();
 
-		// Signature must match our transformer
-		FTransformFunc SpecializedTransformFunc = &TransformerSpecializationType::Transform;
+		bool bSuccess = false;
+
+		if constexpr (requires { &TransformerSpecializationType::TransformSetMap; })
+		{
+			// Signature must match our transformer
+			using FTransformSetMapFunc = typename TransformerCppType::FTransformSetMapFunc;
+			FTransformSetMapFunc SpecializedTransformFunc = &TransformerSpecializationType::TransformSetMap;
 
 #ifdef _MSC_VER
 	#pragma warning(push)
 	#pragma warning(disable : 4191) // unsafe conversion from 'function pointer type A' to 'function pointer type B'
 #endif
 
-		// Cast to a raw dummy function pointer type that we exclusively use for storage
-		// The transformer is responsible for casting it back to the original type it specifies
-		auto RawSpecializedTransformFunc = reinterpret_cast<FRawTransformerFunc>(SpecializedTransformFunc);
+			// Cast to a raw dummy function pointer type that we exclusively use for storage
+			// The transformer is responsible for casting it back to the original type it specifies
+			auto RawSpecializedTransformFunc = reinterpret_cast<FRawTransformerFunc>(SpecializedTransformFunc);
 
 #ifdef _MSC_VER
 	#pragma warning(pop)
 #endif
 
-		return RegisterAttributeValueTransformer(TransformerName, ValueTypeStruct, RawSpecializedTransformFunc);
+			bSuccess |= RegisterAttributeValueSetMapTransformer(TransformerName, ValueTypeStruct, RawSpecializedTransformFunc);
+		}
+
+		if constexpr (requires { &TransformerSpecializationType::TransformMap; })
+		{
+			// Signature must match our transformer
+			using FTransformMapFunc = typename TransformerCppType::FTransformMapFunc;
+			FTransformMapFunc SpecializedTransformFunc = &TransformerSpecializationType::TransformMap;
+
+#ifdef _MSC_VER
+	#pragma warning(push)
+	#pragma warning(disable : 4191) // unsafe conversion from 'function pointer type A' to 'function pointer type B'
+#endif
+
+			// Cast to a raw dummy function pointer type that we exclusively use for storage
+			// The transformer is responsible for casting it back to the original type it specifies
+			auto RawSpecializedTransformFunc = reinterpret_cast<FRawTransformerFunc>(SpecializedTransformFunc);
+
+#ifdef _MSC_VER
+	#pragma warning(pop)
+#endif
+
+			bSuccess |= RegisterAttributeValueMapTransformer(TransformerName, ValueTypeStruct, RawSpecializedTransformFunc);
+		}
+
+		// Returns true if we registered with at least one transformer type
+		return bSuccess;
 	}
 }
 
diff --git a/Engine/Plugins/Experimental/UAF/UAF/Source/UAF/Public/UAF/Attributes/AttributeSetMap.h b/Engine/Plugins/Experimental/UAF/UAF/Source/UAF/Public/UAF/Attributes/AttributeSetMap.h
--- a/Engine/Plugins/Experimental/UAF/UAF/Source/UAF/Public/UAF/Attributes/AttributeSetMap.h
+++ b/Engine/Plugins/Experimental/UAF/UAF/Source/UAF/Public/UAF/Attributes/AttributeSetMap.h
@@ -3,7 +3,6 @@
 #pragma once
 
 #include "CoreMinimal.h"
-#include "Misc/MemStack.h"
 #include "Templates/MemoryOps.h"
 #include "UAF/Attributes/AttributeMappingKey.h"
 #include "UAF/Attributes/AttributeTypedSet.h"
@@ -56,6 +55,9 @@ namespace UE::UAF
 		// Returns the total allocated size for this set map
 		[[nodiscard]] UE_API int32 GetAllocatedSize() const;
 
+		// Returns the allocator function this set map was initialized with
+		[[nodiscard]] UE_API FReallocFun GetAllocator() const;
+
 		// Returns a copy of this set map using the specified allocator
 		[[nodiscard]] virtual FAttributeSetMap* Duplicate(FReallocFun InReallocFun) const = 0;
 
@@ -67,6 +69,10 @@ namespace UE::UAF
 		// The value setter (and/or caller) is responsible for validating the value type
 		virtual void SetValueWithSetter(FAttributeSetIndex Index, const TFunctionRef<void(UScriptStruct* ValueType, uint8* OutValuePtr)>& ValueSetter) = 0;
 
+		// Gets a value at the specified index using a value getter function
+		// The value getter (and/or caller) is responsible for validating the value type
+		virtual void GetValueWithGetter(FAttributeSetIndex Index, const TFunctionRef<void(UScriptStruct* ValueType, const uint8* ValuePtr)>& ValueGetter) const = 0;
+
 		// Sets a value at the specified index
 		// The called is responsible for validating the value type matches the one contained within this set map
 		template <class ValueType>
@@ -120,10 +126,13 @@ namespace UE::UAF
 	// Our keys are held by the attribute set we point to (keys are shared) while we hold the
 	// mapped values internally. Set maps and their values are allocated together in a contiguous buffer.
 	// Value types must be a UStruct
-	template<class ValueType>
+	template<class InValueType>
 	class TAttributeSetMap : public FAttributeSetMap
 	{
 	public:
+		// The type of values held in this map
+		using ValueType = InValueType;
+
 		// Returns a copy of this set map using the specified allocator
 		[[nodiscard]] virtual FAttributeSetMap* Duplicate(FReallocFun InReallocFun) const override;
 
@@ -132,16 +141,16 @@ namespace UE::UAF
 		virtual void CopyTo(FAttributeSetMap& Other) const override;
 
 		// Value accessors by index
-		[[nodiscard]] ValueType& operator[](FAttributeSetIndex Index);
+		[[nodiscard]] InValueType& operator[](FAttributeSetIndex Index);
 
 		// Value accessors by index
-		[[nodiscard]] const ValueType& operator[](FAttributeSetIndex Index) const;
+		[[nodiscard]] const InValueType& operator[](FAttributeSetIndex Index) const;
 
 		// Returns a pointer to the contiguous values held within
-		[[nodiscard]] ValueType* GetData();
+		[[nodiscard]] InValueType* GetData();
 
 		// Returns a pointer to the contiguous values held within
-		[[nodiscard]] const ValueType* GetData() const;
+		[[nodiscard]] const InValueType* GetData() const;
 
 		// Returns the total size required for an instance of a set map of this type
 		[[nodiscard]] static int32 CalculateAllocationSize(int32 NumAttributes);
@@ -153,7 +162,8 @@ namespace UE::UAF
 		// Destroys the set map and frees its memory
 		virtual ~TAttributeSetMap() override;
 
-		virtual void SetValueWithSetter(FAttributeSetIndex Index, const TFunctionRef<void(UScriptStruct* InValueType, uint8* OutValuePtr)>& ValueSetter) override;
+		virtual void SetValueWithSetter(FAttributeSetIndex Index, const TFunctionRef<void(UScriptStruct* ValueType, uint8* OutValuePtr)>& ValueSetter) override;
+		virtual void GetValueWithGetter(FAttributeSetIndex Index, const TFunctionRef<void(UScriptStruct* ValueType, const uint8* ValuePtr)>& ValueGetter) const override;
 		virtual void SetValueImpl(FAttributeSetIndex Index, const void* ValuePtr) override;
 		virtual void GetValueImpl(FAttributeSetIndex Index, void* OutValuePtr) const override;
 
@@ -169,6 +179,38 @@ namespace UE::UAF
 	// Releases and destroys the specified attribute set map
 	void ReleaseAttributeSetMap(FAttributeSetMap* SetMap);
 
+	// Cast from an untyped set map to a typed set map
+	template<class ValueType>
+	UE_FORCEINLINE_HINT TAttributeSetMap<ValueType>* Cast(FAttributeSetMap* Src)
+	{
+		return Src != nullptr && ValueType::StaticStruct() == Src->GetValueType() ? static_cast<TAttributeSetMap<ValueType>*>(Src) : nullptr;
+	}
+
+	// Cast from an untyped set map to a typed set map
+	template<class ValueType>
+	UE_FORCEINLINE_HINT const TAttributeSetMap<ValueType>* Cast(const FAttributeSetMap* Src)
+	{
+		return Src != nullptr && ValueType::StaticStruct() == Src->GetValueType() ? static_cast<const TAttributeSetMap<ValueType>*>(Src) : nullptr;
+	}
+
+	// Cast from an untyped set map to a typed set map with an assert if the types mismatch
+	template<class ValueType>
+	inline TAttributeSetMap<ValueType>* CastChecked(FAttributeSetMap* Src)
+	{
+		TAttributeSetMap<ValueType>* CastResult = Src != nullptr && ValueType::StaticStruct() == Src->GetValueType() ? static_cast<TAttributeSetMap<ValueType>*>(Src) : nullptr;
+		check(CastResult);
+		return CastResult;
+	}
+
+	// Cast from an untyped set map to a typed set map with an assert if the types mismatch
+	template<class ValueType>
+	inline const TAttributeSetMap<ValueType>* CastChecked(const FAttributeSetMap* Src)
+	{
+		const TAttributeSetMap<ValueType>* CastResult = Src != nullptr && ValueType::StaticStruct() == Src->GetValueType() ? static_cast<const TAttributeSetMap<ValueType>*>(Src) : nullptr;
+		check(CastResult);
+		return CastResult;
+	}
+
 	//////////////////////////////////////////////////////////////////////////
 	// Implementation
 
@@ -190,113 +232,121 @@ namespace UE::UAF
 		return Result;
 	}
 
-	template<class ValueType>
-	inline TAttributeSetMap<ValueType>::TAttributeSetMap(const FConstructArgs& Args, int32 InAllocatedSize)
+	template<class InValueType>
+	inline TAttributeSetMap<InValueType>::TAttributeSetMap(const FConstructArgs& Args, int32 InAllocatedSize)
 		: FAttributeSetMap(Args, InAllocatedSize)
 	{
 		const int32 NumAttributes = Args.TypedSet->Num();
 
-		check(InAllocatedSize == TAttributeSetMap<ValueType>::CalculateAllocationSize(NumAttributes));
+		check(InAllocatedSize == TAttributeSetMap<InValueType>::CalculateAllocationSize(NumAttributes));
 
 		if (NumAttributes != 0)
 		{
-			ValuesPtr = Align(reinterpret_cast<uint8*>(this) + sizeof(TAttributeSetMap<ValueType>), alignof(ValueType));
+			ValuesPtr = Align(reinterpret_cast<uint8*>(this) + sizeof(TAttributeSetMap<InValueType>), alignof(InValueType));
 			NumValues = NumAttributes;
 
-			DefaultConstructItems<ValueType>(ValuesPtr, NumAttributes);
+			DefaultConstructItems<InValueType>(ValuesPtr, NumAttributes);
 		}
 	}
 
-	template<class ValueType>
-	inline TAttributeSetMap<ValueType>::~TAttributeSetMap()
+	template<class InValueType>
+	inline TAttributeSetMap<InValueType>::~TAttributeSetMap()
 	{
 		if (ValuesPtr != nullptr)
 		{
-			DestructItems<ValueType>(GetData(), NumValues);
+			DestructItems<InValueType>(GetData(), NumValues);
 
 			ValuesPtr = nullptr;
 			NumValues = 0;
 		}
 	}
 
-	template<class ValueType>
-	inline FAttributeSetMap* TAttributeSetMap<ValueType>::Duplicate(FReallocFun InReallocFun) const
+	template<class InValueType>
+	inline FAttributeSetMap* TAttributeSetMap<InValueType>::Duplicate(FReallocFun InReallocFun) const
 	{
-		TAttributeSetMap<ValueType>* Copy = MakeAttributeSetMap<ValueType>(FConstructArgs(TypedSet, ValueTypeStruct, InReallocFun));
+		TAttributeSetMap<InValueType>* Copy = MakeAttributeSetMap<InValueType>(FConstructArgs(TypedSet, ValueTypeStruct, InReallocFun));
 
 		if (NumValues != 0)
 		{
-			ConstructItems<ValueType>(Copy->ValuesPtr, GetData(), NumValues);
+			ConstructItems<InValueType>(Copy->GetData(), GetData(), NumValues);
 		}
 
 		return Copy;
 	}
 
-	template<class ValueType>
-	inline void TAttributeSetMap<ValueType>::CopyTo(FAttributeSetMap& Other) const
+	template<class InValueType>
+	inline void TAttributeSetMap<InValueType>::CopyTo(FAttributeSetMap& Other) const
 	{
 		checkf(TypedSet == Other.GetTypedSet() && ValueTypeStruct == Other.GetValueType(), TEXT("Typed set and value type must match to allow copying"));
 		check(Other.Num() == Num());	// If the typed sets match, their sizes does too
 
-		TAttributeSetMap<ValueType>& Copy = reinterpret_cast<TAttributeSetMap<ValueType>&>(Other);
+		TAttributeSetMap<InValueType>& Copy = reinterpret_cast<TAttributeSetMap<InValueType>&>(Other);
 
-		CopyAssignItems<ValueType>(Copy.GetData(), GetData(), NumValues);
+		CopyAssignItems<InValueType>(Copy.GetData(), GetData(), NumValues);
 	}
 
-	template<class ValueType>
-	inline ValueType& TAttributeSetMap<ValueType>::operator[](FAttributeSetIndex Index)
+	template<class InValueType>
+	inline InValueType& TAttributeSetMap<InValueType>::operator[](FAttributeSetIndex Index)
 	{
 		checkf((Index.GetInt() >= 0) & (Index.GetInt() < NumValues), TEXT("Map index out of bounds: %lld into an map of size %lld"), (long long)Index.GetInt(), (long long)NumValues); // & for one branch
 		return GetData()[Index.GetInt()];
 	}
 
-	template<class ValueType>
-	inline const ValueType& TAttributeSetMap<ValueType>::operator[](FAttributeSetIndex Index) const
+	template<class InValueType>
+	inline const InValueType& TAttributeSetMap<InValueType>::operator[](FAttributeSetIndex Index) const
 	{
 		checkf((Index.GetInt() >= 0) & (Index.GetInt() < NumValues), TEXT("Map index out of bounds: %lld into an map of size %lld"), (long long)Index.GetInt(), (long long)NumValues); // & for one branch
 		return GetData()[Index.GetInt()];
 	}
 
-	template<class ValueType>
-	inline ValueType* TAttributeSetMap<ValueType>::GetData()
+	template<class InValueType>
+	inline InValueType* TAttributeSetMap<InValueType>::GetData()
 	{
-		return reinterpret_cast<ValueType*>(ValuesPtr);
+		return reinterpret_cast<InValueType*>(ValuesPtr);
 	}
 
-	template<class ValueType>
-	inline const ValueType* TAttributeSetMap<ValueType>::GetData() const
+	template<class InValueType>
+	inline const InValueType* TAttributeSetMap<InValueType>::GetData() const
 	{
-		return reinterpret_cast<const ValueType*>(ValuesPtr);
+		return reinterpret_cast<const InValueType*>(ValuesPtr);
 	}
 
-	template<class ValueType>
-	inline void TAttributeSetMap<ValueType>::SetValueWithSetter(FAttributeSetIndex Index, const TFunctionRef<void(UScriptStruct* InValueType, uint8* OutValuePtr)>& ValueSetter)
+	template<class InValueType>
+	inline void TAttributeSetMap<InValueType>::SetValueWithSetter(FAttributeSetIndex Index, const TFunctionRef<void(UScriptStruct* ValueType, uint8* OutValuePtr)>& ValueSetter)
 	{
 		checkf((Index.GetInt() >= 0) & (Index.GetInt() < NumValues), TEXT("Map index out of bounds: %lld into an map of size %lld"), (long long)Index.GetInt(), (long long)NumValues); // & for one branch
-		ValueType* ValuePtr = GetData() + Index.GetInt();
+		InValueType* ValuePtr = GetData() + Index.GetInt();
 		ValueSetter(ValueTypeStruct, reinterpret_cast<uint8*>(ValuePtr));
 	}
 
-	template<class ValueType>
-	inline void TAttributeSetMap<ValueType>::SetValueImpl(FAttributeSetIndex Index, const void* ValuePtr)
+	template<class InValueType>
+	inline void TAttributeSetMap<InValueType>::GetValueWithGetter(FAttributeSetIndex Index, const TFunctionRef<void(UScriptStruct* ValueType, const uint8* ValuePtr)>& ValueGetter) const
 	{
 		checkf((Index.GetInt() >= 0) & (Index.GetInt() < NumValues), TEXT("Map index out of bounds: %lld into an map of size %lld"), (long long)Index.GetInt(), (long long)NumValues); // & for one branch
-		const ValueType* Value = reinterpret_cast<const ValueType*>(ValuePtr);
+		const InValueType* ValuePtr = GetData() + Index.GetInt();
+		ValueGetter(ValueTypeStruct, reinterpret_cast<const uint8*>(ValuePtr));
+	}
+
+	template<class InValueType>
+	inline void TAttributeSetMap<InValueType>::SetValueImpl(FAttributeSetIndex Index, const void* ValuePtr)
+	{
+		checkf((Index.GetInt() >= 0) & (Index.GetInt() < NumValues), TEXT("Map index out of bounds: %lld into an map of size %lld"), (long long)Index.GetInt(), (long long)NumValues); // & for one branch
+		const InValueType* Value = reinterpret_cast<const InValueType*>(ValuePtr);
 		GetData()[Index.GetInt()] = *Value;
 	}
 
-	template<class ValueType>
-	inline void TAttributeSetMap<ValueType>::GetValueImpl(FAttributeSetIndex Index, void* OutValuePtr) const
+	template<class InValueType>
+	inline void TAttributeSetMap<InValueType>::GetValueImpl(FAttributeSetIndex Index, void* OutValuePtr) const
 	{
 		checkf((Index.GetInt() >= 0) & (Index.GetInt() < NumValues), TEXT("Map index out of bounds: %lld into an map of size %lld"), (long long)Index.GetInt(), (long long)NumValues); // & for one branch
-		ValueType* OutValue = reinterpret_cast<ValueType*>(OutValuePtr);
+		InValueType* OutValue = reinterpret_cast<InValueType*>(OutValuePtr);
 		*OutValue = GetData()[Index.GetInt()];
 	}
 
-	template<class ValueType>
-	inline int32 TAttributeSetMap<ValueType>::CalculateAllocationSize(int32 NumAttributes)
+	template<class InValueType>
+	inline int32 TAttributeSetMap<InValueType>::CalculateAllocationSize(int32 NumAttributes)
 	{
-		return Align(sizeof(TAttributeSetMap<ValueType>), alignof(ValueType)) + sizeof(ValueType) * NumAttributes;
+		return Align(sizeof(TAttributeSetMap<InValueType>), alignof(InValueType)) + sizeof(InValueType) * NumAttributes;
 	}
 
 	template<class ValueType>
diff --git a/Engine/Plugins/Experimental/UAF/UAF/Source/UAF/Public/UAF/Attributes/EngineAttributes.h b/Engine/Plugins/Experimental/UAF/UAF/Source/UAF/Public/UAF/Attributes/EngineAttributes.h
--- a/Engine/Plugins/Experimental/UAF/UAF/Source/UAF/Public/UAF/Attributes/EngineAttributes.h
+++ b/Engine/Plugins/Experimental/UAF/UAF/Source/UAF/Public/UAF/Attributes/EngineAttributes.h
@@ -7,6 +7,8 @@
 
 #include "EngineAttributes.generated.h"
 
+#define UE_API UAF_API
+
 // A bone transform is a special attribute type that allows for template specialization
 USTRUCT(BlueprintType)
 struct FBoneTransformAnimationAttribute : public FTransformAnimationAttribute
@@ -22,3 +24,11 @@ namespace UE::Anim
 		// Inherit the flags from FTransformAnimationAttribute
 	};
 }
+
+namespace UE::UAF
+{
+	// Handles registration for built-in engine attributes and value types
+	UE_API void RegisterEngineAttributes();
+}
+
+#undef UE_API
diff --git a/Engine/Plugins/Experimental/UAF/UAF/Source/UAF/Public/UAF/Attributes/IteratorUtils.h b/Engine/Plugins/Experimental/UAF/UAF/Source/UAF/Public/UAF/Attributes/IteratorUtils.h
--- a/Engine/Plugins/Experimental/UAF/UAF/Source/UAF/Public/UAF/Attributes/IteratorUtils.h
+++ b/Engine/Plugins/Experimental/UAF/UAF/Source/UAF/Public/UAF/Attributes/IteratorUtils.h
@@ -87,6 +87,38 @@ namespace UE::UAF
 				);
 		};
 
+		/** Concept used to detect whether or not a Join Operator specifies a custom LessThan operator */
+		struct CHasCustomLessThan
+		{
+			template <typename T>
+			auto Requires(T& Val) -> decltype(
+				Val.IsLessThan({}, {})
+				);
+		};
+
+		/** Concept used to detect whether or not a Join Operator specifies a HasPredicateWithKey function */
+		struct CHasPredicateWithKey
+		{
+			template <typename T>
+			auto Requires(T& Val) -> decltype(
+				Val.HasPredicateWithKey()
+				);
+		};
+
+		// Returns whether or not the predicate requires the Key as its first argument
+		template<class JoinOpType>
+		constexpr bool HasPredicateWithKey()
+		{
+			if constexpr (TModels_V<CHasPredicateWithKey, JoinOpType>)
+			{
+				return JoinOpType::HasPredicateWithKey();
+			}
+			else
+			{
+				return false;
+			}
+		}
+
 		// Returns whether or not one of the specified iterators has reached the end and is empty
 		template<class... IteratorTypes>
 		inline bool AnyIteratorEmpty(const IteratorTypes&... Iterators)
@@ -154,7 +186,15 @@ namespace UE::UAF
 			checkf(HeadIterator, TEXT("Iterator should have remaining values"));
 			auto HeadKey = JoinOp.GetKey(HeadIterator);
 			auto TailKey = IteratorMin<bTestIfEmpty>(JoinOp, TailIterators...);
-			return HeadKey < TailKey ? HeadKey : TailKey;
+
+			if constexpr (TModels_V<Private::CHasCustomLessThan, JoinOpType>)
+			{
+				return JoinOp.IsLessThan(HeadKey, TailKey) ? HeadKey : TailKey;
+			}
+			else
+			{
+				return HeadKey < TailKey ? HeadKey : TailKey;
+			}
 		}
 
 		// Conditionally increment iterators that match the specified key
@@ -238,6 +278,9 @@ namespace UE::UAF
 		// We know every iterator has at least one value
 		constexpr bool bTestIfEmpty = false;
 
+		// Whether or not to call the predicate with the first argument being the key
+		constexpr bool bWithPredicateKey = Private::HasPredicateWithKey<JoinOpType>();
+
 		// We loop as long as iterators still have values
 		bool bAnyIteratorEmpty = Private::AnyIteratorEmpty(Iterators...);
 		while (!bAnyIteratorEmpty)
@@ -249,7 +292,14 @@ namespace UE::UAF
 			{
 				// Our keys match, call the predicate
 				// All our iterators match our key, return their value (even if the match is fuzzy)
-				Predicate(JoinOp.GetValue(Iterators)...);
+				if constexpr (bWithPredicateKey)
+				{
+					Predicate(MinKey, JoinOp.GetValue(Iterators)...);
+				}
+				else
+				{
+					Predicate(JoinOp.GetValue(Iterators)...);
+				}
 
 				// Move our iterators forward
 				if constexpr (TModels_V<Private::CHasFuzzyMatching, JoinOpType>)
@@ -286,6 +336,9 @@ namespace UE::UAF
 		// We might have iterators that have ran out of values
 		constexpr bool bTestIfEmpty = true;
 
+		// Whether or not to call the predicate with the first argument being the key
+		constexpr bool bWithPredicateKey = Private::HasPredicateWithKey<JoinOpType>();
+
 		// We loop as long as one iterator still has values
 		bool bAllIteratorsEmpty = Private::AllIteratorsEmpty(Iterators...);
 		while (!bAllIteratorsEmpty)
@@ -295,7 +348,14 @@ namespace UE::UAF
 
 			// Call the predicate with a conditional fetch of our iterator values
 			// If an iterator's key matches the smallest key, its value is used (even if the match is fuzzy) otherwise we use nullptr
-			Predicate(Private::ConditionalGetValue(JoinOp, MinKey, Iterators)...);
+			if constexpr (bWithPredicateKey)
+			{
+				Predicate(MinKey, Private::ConditionalGetValue(JoinOp, MinKey, Iterators)...);
+			}
+			else
+			{
+				Predicate(Private::ConditionalGetValue(JoinOp, MinKey, Iterators)...);
+			}
 
 			// Increment iterators that match the smallest key, it isn't needed anymore
 			const uint32 NumEmptyIterators = Private::ConditionalIncrement<bTestIfEmpty>(JoinOp, MinKey, Iterators...);
diff --git a/Engine/Plugins/Experimental/UAF/UAF/Source/UAF/Public/UAF/Attributes/Transformers/Accumulate.h b/Engine/Plugins/Experimental/UAF/UAF/Source/UAF/Public/UAF/Attributes/Transformers/Accumulate.h
--- a/Engine/Plugins/Experimental/UAF/UAF/Source/UAF/Public/UAF/Attributes/Transformers/Accumulate.h
+++ b/Engine/Plugins/Experimental/UAF/UAF/Source/UAF/Public/UAF/Attributes/Transformers/Accumulate.h
@@ -18,22 +18,35 @@ namespace UE::UAF::Transformers
 	{
 		static UE_API FName TransformerName;
 
-		using FTransformFunc = void (*)(UE::UAF::FAttributeSetMap* SetMapA, const UE::UAF::FAttributeSetMap* SetMapB, float Alpha);
+		// If the output buffer is nullptr, an allocator is provided and the implementation is responsible for allocating the output buffer and returning it
+		using FTransformSetMapFunc = void (*)(const FAttributeSetMap* SetMapA, const FAttributeSetMap* SetMapB, float Alpha, FAttributeSetMap*& OutputSetMap, FReallocFun OutputAllocator);
+		using FTransformMapFunc = void (*)(const FAttributeMap* MapA, const FAttributeMap* MapB, float Alpha, FAttributeMap*& OutputMap, FReallocFun OutputAllocator);
 
 		// Accumulate attribute values
-		// accumulate(A, B, Alpha) -> A += (B * Alpha)
-		static UE_API void Apply(const FAttributeValueTransformerMapPtr& TransformerMap, FAttributeCollection& CollectionA, const FAttributeCollection& CollectionB, float Alpha);
+		// Output = accumulate(A, B, Alpha) -> Output = A + (B * Alpha)
+		// This function supports outputting to either input
+		static UE_API void Apply(const FAttributeValueTransformerMapPtr& TransformerMap, const FAttributeCollection& CollectionA, const FAttributeCollection& CollectionB, float Alpha, FAttributeCollection& OutputCollection);
 	};
 
 	//////////////////////////////////////////////////////////////////////////
-	// Specific attribute value type specializations
+	// Engine built-in value type specializations
 
 	struct FAccumulate_BoneTransformAttribute
 	{
 		using TransformerType = FAccumulate;
 		using ValueType = FBoneTransformAnimationAttribute;
 
-		static void Transform(UE::UAF::FAttributeSetMap* SetMapA, const UE::UAF::FAttributeSetMap* SetMapB, float Alpha);
+		static void TransformSetMap(const FAttributeSetMap* SetMapA, const FAttributeSetMap* SetMapB, float Alpha, FAttributeSetMap*& OutputSetMap, FReallocFun OutputAllocator);
+		static void TransformMap(const FAttributeMap* MapA, const FAttributeMap* MapB, float Alpha, FAttributeMap*& OutputMap, FReallocFun OutputAllocator);
+	};
+
+	struct FAccumulate_FloatAttribute
+	{
+		using TransformerType = FAccumulate;
+		using ValueType = FFloatAnimationAttribute;
+
+		static void TransformSetMap(const FAttributeSetMap* SetMapA, const FAttributeSetMap* SetMapB, float Alpha, FAttributeSetMap*& OutputSetMap, FReallocFun OutputAllocator);
+		static void TransformMap(const FAttributeMap* MapA, const FAttributeMap* MapB, float Alpha, FAttributeMap*& OutputMap, FReallocFun OutputAllocator);
 	};
 }
 
diff --git a/Engine/Plugins/Experimental/UAF/UAF/Source/UAF/Public/UAF/Attributes/Transformers/Interpolate.h b/Engine/Plugins/Experimental/UAF/UAF/Source/UAF/Public/UAF/Attributes/Transformers/Interpolate.h
--- a/Engine/Plugins/Experimental/UAF/UAF/Source/UAF/Public/UAF/Attributes/Transformers/Interpolate.h
+++ b/Engine/Plugins/Experimental/UAF/UAF/Source/UAF/Public/UAF/Attributes/Transformers/Interpolate.h
@@ -17,22 +17,35 @@ namespace UE::UAF::Transformers
 	{
 		static UE_API FName TransformerName;
 
-		using FTransformFunc = void (*)(UE::UAF::FAttributeSetMap* SetMapA, const UE::UAF::FAttributeSetMap* SetMapB, float Alpha);
+		// If the output buffer is nullptr, an allocator is provided and the implementation is responsible for allocating the output buffer and returning it
+		using FTransformSetMapFunc = void (*)(const FAttributeSetMap* SetMapA, const FAttributeSetMap* SetMapB, float Alpha, FAttributeSetMap*& OutputSetMap, FReallocFun OutputAllocator);
+		using FTransformMapFunc = void (*)(const FAttributeMap* MapA, const FAttributeMap* MapB, float Alpha, FAttributeMap*& OutputMap, FReallocFun OutputAllocator);
 
 		// Interpolate attribute values
-		// lerp(A, B, Alpha) -> A = (A * (1.0 - Alpha)) + (B * Alpha)
-		static UE_API void Apply(const FAttributeValueTransformerMapPtr& TransformerMap, FAttributeCollection& CollectionA, const FAttributeCollection& CollectionB, float Alpha);
+		// Output = lerp(A, B, Alpha) -> Output = (A * (1.0 - Alpha)) + (B * Alpha)
+		// This function supports outputting to either input
+		static UE_API void Apply(const FAttributeValueTransformerMapPtr& TransformerMap, const FAttributeCollection& CollectionA, const FAttributeCollection& CollectionB, float Alpha, FAttributeCollection& OutputCollection);
 	};
 
 	//////////////////////////////////////////////////////////////////////////
-	// Specific attribute value type specializations
+	// Engine built-in value type specializations
 
 	struct FInterpolate_BoneTransformAttribute
 	{
 		using TransformerType = FInterpolate;
 		using ValueType = FBoneTransformAnimationAttribute;
 
-		static void Transform(UE::UAF::FAttributeSetMap* SetMapA, const UE::UAF::FAttributeSetMap* SetMapB, float Alpha);
+		static void TransformSetMap(const FAttributeSetMap* SetMapA, const FAttributeSetMap* SetMapB, float Alpha, FAttributeSetMap*& OutputSetMap, FReallocFun OutputAllocator);
+		static void TransformMap(const FAttributeMap* MapA, const FAttributeMap* MapB, float Alpha, FAttributeMap*& OutputMap, FReallocFun OutputAllocator);
+	};
+
+	struct FInterpolate_FloatAttribute
+	{
+		using TransformerType = FInterpolate;
+		using ValueType = FFloatAnimationAttribute;
+
+		static void TransformSetMap(const FAttributeSetMap* SetMapA, const FAttributeSetMap* SetMapB, float Alpha, FAttributeSetMap*& OutputSetMap, FReallocFun OutputAllocator);
+		static void TransformMap(const FAttributeMap* MapA, const FAttributeMap* MapB, float Alpha, FAttributeMap*& OutputMap, FReallocFun OutputAllocator);
 	};
 }
 
diff --git a/Engine/Plugins/Experimental/UAF/UAF/Source/UAF/Public/UAF/Attributes/Transformers/Overwrite.h b/Engine/Plugins/Experimental/UAF/UAF/Source/UAF/Public/UAF/Attributes/Transformers/Overwrite.h
--- a/Engine/Plugins/Experimental/UAF/UAF/Source/UAF/Public/UAF/Attributes/Transformers/Overwrite.h
+++ b/Engine/Plugins/Experimental/UAF/UAF/Source/UAF/Public/UAF/Attributes/Transformers/Overwrite.h
@@ -18,22 +18,36 @@ namespace UE::UAF::Transformers
 	{
 		static UE_API FName TransformerName;
 
-		using FTransformFunc = void (*)(UE::UAF::FAttributeSetMap* SetMapA, const UE::UAF::FAttributeSetMap* SetMapB, float Alpha);
+		// If the output buffer is nullptr, an allocator is provided and the implementation is responsible for allocating the output buffer and returning it
+		// Note that the input/output buffers can be the same buffer
+		using FTransformSetMapFunc = void (*)(const FAttributeSetMap* SetMap, float Alpha, FAttributeSetMap*& OutputSetMap, FReallocFun OutputAllocator);
+		using FTransformMapFunc = void (*)(const FAttributeMap* Map, float Alpha, FAttributeMap*& OutputMap, FReallocFun OutputAllocator);
 
 		// Overwrite attribute values
-		// overwrite(A, B, Alpha) -> A = (B * Alpha)
-		static UE_API void Apply(const FAttributeValueTransformerMapPtr& TransformerMap, FAttributeCollection& CollectionA, const FAttributeCollection& CollectionB, float Alpha);
+		// Output = overwrite(A, Alpha) -> Output = (A * Alpha)
+		// This function supports outputting to its input
+		static UE_API void Apply(const FAttributeValueTransformerMapPtr& TransformerMap, const FAttributeCollection& SourceCollection, float Alpha, FAttributeCollection& OutputCollection);
 	};
 
 	//////////////////////////////////////////////////////////////////////////
-	// Specific attribute value type specializations
+	// Engine built-in value type specializations
 
 	struct FOverwrite_BoneTransformAttribute
 	{
 		using TransformerType = FOverwrite;
 		using ValueType = FBoneTransformAnimationAttribute;
 
-		static void Transform(UE::UAF::FAttributeSetMap* SetMapA, const UE::UAF::FAttributeSetMap* SetMapB, float Alpha);
+		static void TransformSetMap(const FAttributeSetMap* SetMap, float Alpha, FAttributeSetMap*& OutputSetMap, FReallocFun OutputAllocator);
+		static void TransformMap(const FAttributeMap* Map, float Alpha, FAttributeMap*& OutputMap, FReallocFun OutputAllocator);
+	};
+
+	struct FOverwrite_FloatAttribute
+	{
+		using TransformerType = FOverwrite;
+		using ValueType = FFloatAnimationAttribute;
+
+		static void TransformSetMap(const FAttributeSetMap* SetMap, float Alpha, FAttributeSetMap*& OutputSetMap, FReallocFun OutputAllocator);
+		static void TransformMap(const FAttributeMap* Map, float Alpha, FAttributeMap*& OutputMap, FReallocFun OutputAllocator);
 	};
 }
 
diff --git a/Engine/Plugins/Experimental/UAF/UAF/Source/UAF/Public/UAF/Attributes/Transformers/Sanitize.h b/Engine/Plugins/Experimental/UAF/UAF/Source/UAF/Public/UAF/Attributes/Transformers/Sanitize.h
--- a/Engine/Plugins/Experimental/UAF/UAF/Source/UAF/Public/UAF/Attributes/Transformers/Sanitize.h
+++ b/Engine/Plugins/Experimental/UAF/UAF/Source/UAF/Public/UAF/Attributes/Transformers/Sanitize.h
@@ -7,6 +7,9 @@
 #include "UAF/Attributes/EngineAttributes.h"
 #include "UAF/Attributes/Transformers/Transformer.h"
 
+// We validate mapped values in development and debug builds
+#define UE_UAF_VALIDATE_MAPPED_VALUES !(UE_BUILD_SHIPPING || UE_BUILD_TEST)
+
 #define UE_API UAF_API
 
 namespace UE::UAF::Transformers
@@ -17,15 +20,16 @@ namespace UE::UAF::Transformers
 	{
 		static UE_API FName TransformerName;
 
-		using FTransformFunc = void (*)(UE::UAF::FAttributeSetMap* SetMap);
+		using FTransformSetMapFunc = void (*)(FAttributeSetMap* SetMap);
+		using FTransformMapFunc = void (*)(FAttributeMap* Map);
 
 		// Sanitize attribute values
-		// This can be used to normalize rotations and other like values as well as check for NaN/Inf values
+		// This can be used to normalize rotations (and other like values) as well as check for NaN/Inf values
 		static UE_API void Apply(const FAttributeValueTransformerMapPtr& TransformerMap, FAttributeCollection& Collection);
 	};
 
 	//////////////////////////////////////////////////////////////////////////
-	// Specific attribute value type specializations
+	// Engine built-in value type specializations
 
 	// Sanitize bone transforms (normalize rotation, ensure finite)
 	struct FSanitize_BoneTransformAttribute
@@ -33,8 +37,21 @@ namespace UE::UAF::Transformers
 		using TransformerType = FSanitize;
 		using ValueType = FBoneTransformAnimationAttribute;
 
-		static void Transform(UE::UAF::FAttributeSetMap* SetMap);
+		static void TransformSetMap(FAttributeSetMap* SetMap);
+		static void TransformMap(FAttributeMap* Map);
+	};
+
+#if UE_UAF_VALIDATE_MAPPED_VALUES
+	// Ensure values are not NaN (but Inf is valid)
+	struct FSanitize_FloatAttribute
+	{
+		using TransformerType = FSanitize;
+		using ValueType = FFloatAnimationAttribute;
+
+		static void TransformSetMap(FAttributeSetMap* SetMap);
+		static void TransformMap(FAttributeMap* Map);
 	};
+#endif
 }
 
 #undef UE_API
diff --git a/Engine/Plugins/Experimental/UAF/UAF/Source/UAF/Public/UAF/Attributes/Transformers/Transformer.h b/Engine/Plugins/Experimental/UAF/UAF/Source/UAF/Public/UAF/Attributes/Transformers/Transformer.h
--- a/Engine/Plugins/Experimental/UAF/UAF/Source/UAF/Public/UAF/Attributes/Transformers/Transformer.h
+++ b/Engine/Plugins/Experimental/UAF/UAF/Source/UAF/Public/UAF/Attributes/Transformers/Transformer.h
@@ -23,15 +23,18 @@ namespace UE::UAF
 		// The name of the transformer (exposed in derived types)
 		//static FName TransformerName;
 
-		// The function signature of the transformed (exposed in derived types)
-		//using FTransformFunc = void (*)(...);
+		// An optional function signature of the transformer for set maps (exposed in derived types)
+		//using FTransformSetMapFunc = void (*)(...);
+
+		// An optional function signature of the transformer for name/value maps (exposed in derived types)
+		//using FTransformMapFunc = void (*)(...);
 	};
 
 	// An iterator over a transformer list within the attribute runtime registry
 	class FAttributeValueTransformerListIterator
 	{
 	public:
-		[[nodiscard]] FAttributeValueTransformerListIterator(const FAttributeValueTransformerList& TransformerList);
+		[[nodiscard]] FAttributeValueTransformerListIterator(const FAttributeValueTransformerList& TransformerList, int32 StartOffset, int32 Count);
 
 		// Increments and moves the iterator to the next transformer entry
 		FAttributeValueTransformerListIterator& operator++();
@@ -46,16 +49,15 @@ namespace UE::UAF
 		[[nodiscard]] FRawTransformerFunc GetTransformerFunc() const;
 
 	private:
-		// Iterator over our value types
-		TArray<UScriptStruct*>::TConstIterator ValueTypesIt;
-
-		// Iterator over our transformer instances
-		TArray<FRawTransformerFunc>::TConstIterator TransformerFuncIt;
+		UScriptStruct* const* ValueTypePtr = nullptr;
+		const FRawTransformerFunc* TransformerFuncPtr = nullptr;
+		UScriptStruct* const* ValueTypeEndPtr = nullptr;
 	};
 
-	// A list of transformer specializations that operators over various value types
+	// A list of transformer specializations that operators over various value types.
 	// Transformers have a base type (e.g. interpolate, sanitize) and can be specialized to operate on arbitrary
-	// value types (e.g. FFloatAnimationAttribute)
+	// value types (e.g. FFloatAnimationAttribute). They can also transform two types of containers: set maps
+	// and name/value maps.
 	class FAttributeValueTransformerList
 	{
 	public:
@@ -65,20 +67,42 @@ namespace UE::UAF
 		// Creates a transformer list for the specified transformer type
 		[[nodiscard]] UE_API explicit FAttributeValueTransformerList(FName TransformerName);
 
-		// Adds a new transformer entry
-		UE_API bool AddTransformer(UScriptStruct* ValueType, FRawTransformerFunc TransformerFunc);
+		// Returns the transformer name
+		[[nodiscard]] UE_API FName GetTransformerName() const;
+
+		// Returns the number of set map transformers in this list
+		[[nodiscard]] int32 NumSetMapTransformers() const { return SetMapTransformerCount; }
+
+		// Returns the number of name/value map transformers in this list
+		[[nodiscard]] int32 NumMapTransformers() const { return ValueTypes.Num() - SetMapTransformerCount; }
+
+		// Adds a new set map transformer entry
+		UE_API bool AddSetMapTransformer(UScriptStruct* ValueType, FRawTransformerFunc TransformerFunc);
 
-		// Returns an iterator over all transformer specializations within
-		[[nodiscard]] FAttributeValueTransformerListIterator CreateTransformerIterator() const { return FAttributeValueTransformerListIterator(*this); }
+		// Adds a new name/value map transformer entry
+		UE_API bool AddMapTransformer(UScriptStruct* ValueType, FRawTransformerFunc TransformerFunc);
+
+		// Returns an iterator over all set map transformer specializations within
+		[[nodiscard]] FAttributeValueTransformerListIterator CreateSetMapTransformerIterator() const { return FAttributeValueTransformerListIterator(*this, 0, NumSetMapTransformers()); }
+
+		// Returns an iterator over all name/value map transformer specializations within
+		[[nodiscard]] FAttributeValueTransformerListIterator CreateMapTransformerIterator() const { return FAttributeValueTransformerListIterator(*this, NumSetMapTransformers(), NumMapTransformers()); }
 
 	private:
 		// The type of the transformer (e.g. interpolate, sanitize)
 		FName TransformerName;
 
-		// List of attribute value types for which we have transformers registered (e.g. FTransformAnimationAttribute, FFloatAnimationAttribute) (sorted by value)
+		// The number of set map transformers in this list
+		int32 SetMapTransformerCount = 0;
+
+		// List of attribute value types for which we have transformers registered (e.g. FTransformAnimationAttribute, FFloatAnimationAttribute) (sorted by container type, then by value type)
+		// This list is split in two, set map transformers come first, then map transformers
+		// Each sub-list is sorted by value type
 		TArray<UScriptStruct*> ValueTypes;
 
-		// List of transformer functions for each value type (sorted by value type)
+		// List of transformer functions for each value type (sorted by container type, then by value type)
+		// This list is split in two, set map transformers come first, then map transformers
+		// Each sub-list is sorted by value type
 		TArray<FRawTransformerFunc> Transformers;
 
 		friend FAttributeRuntimeRegistry;
@@ -89,7 +113,7 @@ namespace UE::UAF
 
 	// An operator structure to join attribute value transformer lists and attribute collections using InnerJoin/OuterJoin/etc
 	// Ensures that the predicate is called for each set map with a value type that has a corresponding transformer
-	struct FAttributeValueTransformerListWithAttributeCollectionJoinOp
+	struct FAttributeValueTransformerListWithAttributeCollectionSetMapJoinOp
 	{
 		// Key sorting uses signed comparison, use the largest signed value regardless of pointer size
 		static constexpr intptr_t WILDCARD_ATTRIBUTE_TYPE_VALUE = static_cast<intptr_t>(static_cast<uintptr_t>(~0ULL) >> 1);
@@ -128,36 +152,58 @@ namespace UE::UAF
 		const FAttributeSetMap* GetValue(const FAttributeCollection::TConstSetMapIterator& It) const { return It.GetMap(); }
 	};
 
+	// An operator structure to join attribute value transformer lists and attribute collections using InnerJoin/OuterJoin/etc
+	// Ensures that the predicate is called for each set map with a value type that has a corresponding transformer
+	struct FAttributeValueTransformerListWithAttributeCollectionMapJoinOp
+	{
+		// Key sorting uses signed comparison, use the largest signed value regardless of pointer size
+		static constexpr intptr_t KEY_LARGEST_VALUE = static_cast<intptr_t>(static_cast<uintptr_t>(~0ULL) >> 1);
+
+		UScriptStruct* GetLargestKey() const { return reinterpret_cast<UScriptStruct*>(KEY_LARGEST_VALUE); }
+
+		// FAttributeValueTransformerListIterator entries
+		UScriptStruct* GetKey(const FAttributeValueTransformerListIterator& It) const { return It.GetValueType(); }
+		FRawTransformerFunc GetValue(const FAttributeValueTransformerListIterator& It) const { return It.GetTransformerFunc(); }
+
+		// Normal FAttributeCollectionMapJoinOp
+		UScriptStruct* GetKey(const FAttributeCollection::TMapIterator& It) const { return It.GetValueType(); }
+		UScriptStruct* GetKey(const FAttributeCollection::TConstMapIterator& It) const { return It.GetValueType(); }
+
+		FAttributeMap* GetValue(const FAttributeCollection::TMapIterator& It) const { return It.GetMap(); }
+		const FAttributeMap* GetValue(const FAttributeCollection::TConstMapIterator& It) const { return It.GetMap(); }
+	};
+
 	//////////////////////////////////////////////////////////////////////////
 	// Implementation
 
-	inline FAttributeValueTransformerListIterator::FAttributeValueTransformerListIterator(const FAttributeValueTransformerList& TransformerList)
-		: ValueTypesIt(TransformerList.ValueTypes.CreateConstIterator())
-		, TransformerFuncIt(TransformerList.Transformers.CreateConstIterator())
+	inline FAttributeValueTransformerListIterator::FAttributeValueTransformerListIterator(const FAttributeValueTransformerList& TransformerList, int32 StartOffset, int32 Count)
+		: ValueTypePtr(TransformerList.ValueTypes.GetData() + StartOffset)
+		, TransformerFuncPtr(TransformerList.Transformers.GetData() + StartOffset)
+		, ValueTypeEndPtr(ValueTypePtr + Count)
 	{
 	}
 
 	inline FAttributeValueTransformerListIterator& FAttributeValueTransformerListIterator::operator++()
 	{
-		++ValueTypesIt;
-		++TransformerFuncIt;
+		++ValueTypePtr;
+		++TransformerFuncPtr;
 
 		return *this;
 	}
 
 	inline FAttributeValueTransformerListIterator::operator bool() const
 	{
-		return (bool)ValueTypesIt;
+		return ValueTypePtr < ValueTypeEndPtr;
 	}
 
 	inline UScriptStruct* FAttributeValueTransformerListIterator::GetValueType() const
 	{
-		return *ValueTypesIt;
+		return *ValueTypePtr;
 	}
 
 	inline FRawTransformerFunc FAttributeValueTransformerListIterator::GetTransformerFunc() const
 	{
-		return *TransformerFuncIt;
+		return *TransformerFuncPtr;
 	}
 }
 
diff --git a/Engine/Plugins/Experimental/UAF/UAFAnimGraph/Source/UAFAnimGraph/Private/EvaluationVM/Tasks/BlendKeyframes.cpp b/Engine/Plugins/Experimental/UAF/UAFAnimGraph/Source/UAFAnimGraph/Private/EvaluationVM/Tasks/BlendKeyframes.cpp
--- a/Engine/Plugins/Experimental/UAF/UAFAnimGraph/Source/UAFAnimGraph/Private/EvaluationVM/Tasks/BlendKeyframes.cpp
+++ b/Engine/Plugins/Experimental/UAF/UAFAnimGraph/Source/UAFAnimGraph/Private/EvaluationVM/Tasks/BlendKeyframes.cpp
@@ -70,7 +70,7 @@ void FAnimNextBlendTwoKeyframesTask::Execute(UE::UAF::FEvaluationVM& VM) const
 			// Named sets must match as it ensures that our inputs have the same sizes/shapes
 			check(CollectionA->GetNamedSet() == CollectionB->GetNamedSet());
 
-			Transformers::FInterpolate::Apply(VM.GetTransformerMap(), *CollectionA, *CollectionB, WeightOfPoseB);
+			Transformers::FInterpolate::Apply(VM.GetTransformerMap(), *CollectionA, *CollectionB, WeightOfPoseB, *CollectionA);
 			Transformers::FSanitize::Apply(VM.GetTransformerMap(), *CollectionA);
 
 			VM.PushValue(ATTRIBUTE_STACK_NAME, MoveTemp(CollectionA));
@@ -227,7 +227,7 @@ void FAnimNextBlendOverwriteKeyframeWithScaleTask::Execute(UE::UAF::FEvaluationV
 
 		// Overwrite in place
 
-		Transformers::FOverwrite::Apply(VM.GetTransformerMap(), *Collection, *Collection, ScaleFactor);
+		Transformers::FOverwrite::Apply(VM.GetTransformerMap(), *Collection, ScaleFactor, *Collection);
 
 		VM.PushValue(ATTRIBUTE_STACK_NAME, MoveTemp(Collection));
 	}
@@ -300,7 +300,7 @@ void FAnimNextBlendAddKeyframeWithScaleTask::Execute(UE::UAF::FEvaluationVM& VM)
 		// Named sets must match as it ensures that our inputs have the same sizes/shapes
 		check(CollectionA->GetNamedSet() == CollectionB->GetNamedSet());
 
-		Transformers::FAccumulate::Apply(VM.GetTransformerMap(), *CollectionB, *CollectionA, ScaleFactor);
+		Transformers::FAccumulate::Apply(VM.GetTransformerMap(), *CollectionB, *CollectionA, ScaleFactor, *CollectionB);
 
 		VM.PushValue(ATTRIBUTE_STACK_NAME, MoveTemp(CollectionB));
 	}
diff --git a/Engine/Plugins/Experimental/UAF/UAFAnimGraph/Source/UAFAnimGraph/Private/TraitInterfaces/IEvaluate.cpp b/Engine/Plugins/Experimental/UAF/UAFAnimGraph/Source/UAFAnimGraph/Private/TraitInterfaces/IEvaluate.cpp
--- a/Engine/Plugins/Experimental/UAF/UAFAnimGraph/Source/UAFAnimGraph/Private/TraitInterfaces/IEvaluate.cpp
+++ b/Engine/Plugins/Experimental/UAF/UAFAnimGraph/Source/UAFAnimGraph/Private/TraitInterfaces/IEvaluate.cpp
@@ -2,6 +2,7 @@
 
 #include "TraitInterfaces/IEvaluate.h"
 
+#include "Animation/AnimCurveUtils.h"
 #include "TraitCore/ExecutionContext.h"
 #include "TraitInterfaces/IHierarchy.h"
 #include "AnimNextAnimGraphStats.h"
@@ -258,26 +259,93 @@ namespace UE::UAF
 				if (Private::GUseAbstractSkeletonRuntime && EvaluationVM.PopValue(ATTRIBUTE_STACK_NAME, Collection))
 				{
 					const FAttributeMappingKey MappingKey = FAttributeMappingKey::MakeFromTo<FBoneTransformAnimationAttribute>();
-					if (TAttributeSetMap<FBoneTransformAnimationAttribute>* BoneTransforms = Collection->FindSetMap<FBoneTransformAnimationAttribute>(MappingKey))
+					if (!Collection->IsEmpty())
 					{
-						const FAttributeTypedSetPtr& BoneSet = BoneTransforms->GetTypedSet();
-						const int32 NumPoseLODBones = OutputPose.LODPose.GetNumBones();
+						// Attributes should no longer be attached to bones in UAF
+						// Attributes should have unique names instead
+						// Attach to root for now for backwards compatibility
+						const FCompactPoseBoneIndex RootCompactPoseBoneIndex(0);
 
-						for (int32 LODBoneIndex = 0; LODBoneIndex < NumPoseLODBones; ++LODBoneIndex)
+						// Set maps
+						for (auto It = Collection->CreateConstSetMapIterator(); It; ++It)
 						{
-							const int32 SkeletonBoneIndex = RefPose.GetSkeletonBoneIndexFromLODBoneIndex(LODBoneIndex);
-							const FAttributeBindingIndex BoneBindingIndex = FAttributeBindingIndex(FSkeletonPoseBoneIndex(SkeletonBoneIndex));
-							check(BoneBindingIndex.IsValid());
+							const FAttributeSetMap* SetMap = It.GetMap();
+							const FAttributeTypedSetPtr& TypedSet = SetMap->GetTypedSet();
 
-							const FAttributeSetIndex BoneSetIndex = BoneSet->GetIndex(BoneBindingIndex);
-							if (BoneSetIndex.IsValid())
+							if (const TAttributeSetMap<FBoneTransformAnimationAttribute>* BoneTransforms = Cast<FBoneTransformAnimationAttribute>(SetMap))
 							{
-								OutputPose.LODPose.LocalTransformsView[LODBoneIndex] = (*BoneTransforms)[BoneSetIndex].Value;
+								const int32 NumPoseLODBones = OutputPose.LODPose.GetNumBones();
+
+								for (int32 LODBoneIndex = 0; LODBoneIndex < NumPoseLODBones; ++LODBoneIndex)
+								{
+									const int32 SkeletonBoneIndex = RefPose.GetSkeletonBoneIndexFromLODBoneIndex(LODBoneIndex);
+									const FAttributeBindingIndex BoneBindingIndex = FAttributeBindingIndex(FSkeletonPoseBoneIndex(SkeletonBoneIndex));
+									check(BoneBindingIndex.IsValid());
+
+									const FAttributeSetIndex BoneSetIndex = TypedSet->GetIndex(BoneBindingIndex);
+									if (BoneSetIndex.IsValid())
+									{
+										OutputPose.LODPose.LocalTransformsView[LODBoneIndex] = (*BoneTransforms)[BoneSetIndex].Value;
+									}
+									else
+									{
+										// This bone isn't present with the current named set, use the bind pose value instead
+										OutputPose.LODPose.LocalTransformsView[LODBoneIndex] = RefPose.GetRefPoseTransform(LODBoneIndex);
+									}
+								}
+							}
+							else if (const TAttributeSetMap<FFloatAnimationAttribute>* Floats = Cast<FFloatAnimationAttribute>(SetMap))
+							{
+								const auto GetNameFromIndex = [&TypedSet](int32 CurveIndex)
+									{
+										return TypedSet->GetName(FAttributeSetIndex(CurveIndex));
+									};
+
+								const auto GetValueFromIndex = [Floats](int32 CurveIndex)
+									{
+										return (*Floats)[FAttributeSetIndex(CurveIndex)].Value;
+									};
+
+								// Convert float set attributes into curves
+								UE::Anim::FCurveUtils::BuildSorted(OutputPose.Curves, Floats->Num(), GetNameFromIndex, GetValueFromIndex);
 							}
 							else
 							{
-								// This bone isn't present with the current named set, use the bind pose value instead
-								OutputPose.LODPose.LocalTransformsView[LODBoneIndex] = RefPose.GetRefPoseTransform(LODBoneIndex);
+								// Everything else
+								const UScriptStruct* ValueType = SetMap->GetValueType();
+								const int32 NumAttributes = SetMap->Num();
+
+								for (FAttributeSetIndex SetIndex(0); SetIndex < NumAttributes; ++SetIndex)
+								{
+									const UE::Anim::FAttributeId AttributeId(TypedSet->GetName(SetIndex), RootCompactPoseBoneIndex);
+
+									uint8* AttributePtr = OutputPose.Attributes.FindOrAdd(ValueType, AttributeId);
+
+									SetMap->GetValueWithGetter(SetIndex, [AttributePtr](UScriptStruct* ValueType_, const uint8* ValuePtr)
+										{
+											ValueType_->CopyScriptStruct(AttributePtr, ValuePtr);
+										});
+								}
+							}
+						}
+
+						// Name/value maps
+						for (auto It = Collection->CreateConstMapIterator(); It; ++It)
+						{
+							const FAttributeMap* Map = It.GetMap();
+							const UScriptStruct* ValueType = Map->GetValueType();
+							const int32 NumValues = Map->Num();
+
+							for (int32 ValueIndex = 0; ValueIndex < NumValues; ++ValueIndex)
+							{
+								const UE::Anim::FAttributeId AttributeId(Map->GetName(ValueIndex), RootCompactPoseBoneIndex);
+
+								uint8* AttributePtr = OutputPose.Attributes.FindOrAdd(ValueType, AttributeId);
+
+								Map->GetValueWithGetter(ValueIndex, [AttributePtr](UScriptStruct* ValueType_, const uint8* ValuePtr)
+									{
+										ValueType_->CopyScriptStruct(AttributePtr, ValuePtr);
+									});
 							}
 						}
 					}
