- Added AddUnitNodeOfClass to allow spawning of AnimNext typed unit nodes from BP/Python scripting
- Updated UAF to use AddUnitNodeOfClass instead of AddUnitNode to provide clear python code message-log information

diff --git a/Engine/Plugins/Experimental/UAF/UAF/Source/UAFEditor/Private/Common/GraphEditorSchemaActions.cpp b/Engine/Plugins/Experimental/UAF/UAF/Source/UAFEditor/Private/Common/GraphEditorSchemaActions.cpp
--- a/Engine/Plugins/Experimental/UAF/UAF/Source/UAFEditor/Private/Common/GraphEditorSchemaActions.cpp
+++ b/Engine/Plugins/Experimental/UAF/UAF/Source/UAFEditor/Private/Common/GraphEditorSchemaActions.cpp
@@ -19,6 +19,7 @@
 #include "UAFCompilationScope.h"
 #include "ScopedTransaction.h"
 #include "Variables/AnimNextSharedVariables.h"
+#include "UObject/Class.h"
 
 #include UE_INLINE_GENERATED_CPP_BY_NAME(GraphEditorSchemaActions)
 
@@ -53,14 +54,14 @@ UEdGraphNode* FAnimNextSchemaAction_RigUnit::PerformAction(UEdGraph* ParentGraph
 	if (Host != nullptr && EdGraph != nullptr)
 	{
 		FName Name = UE::UAF::Editor::FUtils::ValidateName(Cast<UObject>(Host), StructTemplate->GetFName().ToString());
-		URigVMController* Controller = Host->GetRigVMClient()->GetController(ParentGraph);
+		UAnimNextControllerBase* Controller = CastChecked<UAnimNextControllerBase>(Host->GetRigVMClient()->GetController(ParentGraph));
 
 		Controller->OpenUndoBracket(FString::Printf(TEXT("Add '%s' Node"), *Name.ToString()));
 
 		FRigVMUnitNodeCreatedContext& UnitNodeCreatedContext = Controller->GetUnitNodeCreatedContext();
 		FRigVMUnitNodeCreatedContext::FScope ReasonScope(UnitNodeCreatedContext, ERigVMNodeCreatedReason::NodeSpawner, Host);
 
-		if (URigVMUnitNode* ModelNode = Controller->AddUnitNode(StructTemplate, NodeClass, FRigVMStruct::ExecuteName, FDeprecateSlateVector2D(Location), Name.ToString(), true, true))
+		if (URigVMUnitNode* ModelNode = Controller->AddUnitNodeOfClass(StructTemplate->GetPathName(), NodeClass.Get()->GetPathName(), FRigVMStruct::ExecuteName, FDeprecateSlateVector2D(Location), Name.ToString(), true, true))
 		{
 			NewNode = Cast<URigVMEdGraphNode>(EdGraph->FindNodeForModelNodeName(ModelNode->GetFName()));
 			check(NewNode);
diff --git a/Engine/Plugins/Experimental/UAF/UAF/Source/UAFUncookedOnly/Internal/AnimNextControllerBase.h b/Engine/Plugins/Experimental/UAF/UAF/Source/UAFUncookedOnly/Internal/AnimNextControllerBase.h
--- a/Engine/Plugins/Experimental/UAF/UAF/Source/UAFUncookedOnly/Internal/AnimNextControllerBase.h
+++ b/Engine/Plugins/Experimental/UAF/UAF/Source/UAFUncookedOnly/Internal/AnimNextControllerBase.h
@@ -35,6 +35,12 @@ public:
 	UFUNCTION(BlueprintCallable, Category = RigVMController)
 	UAFUNCOOKEDONLY_API UAnimNextSharedVariableNode* AddSharedVariableNode(const FString& InSourceObjectPath, const FName& InVariableName, const FString& InCPPType, const FString& InCPPTypeObjectPath, bool bIsGetter, const FString& InDefaultValue, const FVector2D& InPosition, const FString& InNodeName, bool bSetupUndoRedo, bool bPrintPythonCommand);
 	
+	// Adds a Function / Struct Node to the edited Graph.
+	// UnitNode represent a RIGVM_METHOD declaration on a USTRUCT.
+	// This causes a NodeAdded modified event.
+	UFUNCTION(BlueprintCallable, Category = "UAF|Controller")
+	UAFUNCOOKEDONLY_API URigVMUnitNode* AddUnitNodeOfClass(const FString& InScriptStructPath, const FString& InUnitNodeClassPath, const FName& InMethodName, const FVector2D& InPosition, const FString& InNodeName, bool bSetupUndoRedo = true, bool bPrintPythonCommand = false);
+	
 	UAnimNextSharedVariableNode* ReplaceVariableNodeWithSharedVariableNode(URigVMVariableNode* InVariableNode, FName InNewVariableName, const UObject* InAssetOrStruct, bool bSetupUndoRedo, bool bPrintPythonCommand);
 	URigVMVariableNode* ReplaceSharedVariableNodeWithVariableNode(UAnimNextSharedVariableNode* InVariableNode, FName InNewVariableName, bool bSetupUndoRedo, bool bPrintPythonCommand);
 };
diff --git a/Engine/Plugins/Experimental/UAF/UAF/Source/UAFUncookedOnly/Private/AnimNextControllerBase.cpp b/Engine/Plugins/Experimental/UAF/UAF/Source/UAFUncookedOnly/Private/AnimNextControllerBase.cpp
--- a/Engine/Plugins/Experimental/UAF/UAF/Source/UAFUncookedOnly/Private/AnimNextControllerBase.cpp
+++ b/Engine/Plugins/Experimental/UAF/UAF/Source/UAFUncookedOnly/Private/AnimNextControllerBase.cpp
@@ -102,6 +102,51 @@ UAnimNextSharedVariableNode* UAnimNextControllerBase::AddSharedVariableNode(cons
 	return VariableNode;
 }
 
+URigVMUnitNode* UAnimNextControllerBase::AddUnitNodeOfClass(const FString& InScriptStructPath, const FString& InUnitNodeClassPath, const FName& InMethodName, const FVector2D& InPosition, const FString& InNodeName, bool bSetupUndoRedo, bool bPrintPythonCommand)
+{
+	UScriptStruct* ScriptStruct = RigVMTypeUtils::FindObjectFromCPPTypeObjectPath<UScriptStruct>(InScriptStructPath);
+	if (ScriptStruct == nullptr)
+	{
+		ReportErrorf(TEXT("Cannot find struct for path '%s'."), *InScriptStructPath);
+		return nullptr;
+	}
+	
+	UClass* UnitNodeClass = RigVMTypeUtils::FindObjectFromCPPTypeObjectPath<UClass>(InUnitNodeClassPath);
+	if (UnitNodeClass == nullptr)
+	{
+		ReportErrorf(TEXT("Cannot find class for path '%s'."), *InUnitNodeClassPath);
+		return nullptr;
+	}
+	
+	if (!UnitNodeClass->IsChildOf(URigVMUnitNode::StaticClass()))
+	{
+		ReportErrorf(TEXT("Provided UnitNodeClass (%s) is not a subclass of URigVMUnitNode."), *UnitNodeClass->GetName());
+		return nullptr;
+	}
+
+	URigVMUnitNode* UnitNode = AddUnitNode(ScriptStruct, UnitNodeClass, InMethodName, InPosition, InNodeName, bSetupUndoRedo, false);
+	
+	if (!UnitNode)
+	{
+		ReportErrorf(TEXT("Failed to AddUnitNodeOfClass (%s)"), *UnitNodeClass->GetName());
+		return nullptr;
+	}
+	
+	if (bPrintPythonCommand)
+	{
+		const FString GraphName = GetSchema()->GetSanitizedGraphName(GetGraph()->GetGraphName());
+		RigVMPythonUtils::Print(GetSchema()->GetGraphOuterName(GetGraph()), FString::Printf(TEXT("blueprint.get_controller_by_name('%s').add_unit_node_of_class('%s', '%s', '%s', %s, '%s')"),
+					*GraphName,
+					*InScriptStructPath,
+					*InUnitNodeClassPath,
+					*InMethodName.ToString(),
+					*RigVMPythonUtils::Vector2DToPythonString(UnitNode->GetPosition()),
+					*UnitNode->GetName()));
+	}
+	
+	return UnitNode;
+}
+
 void UAnimNextControllerBase::RefreshSharedVariableNode(const FName& InNodeName, const FString& InSourceObjectPath, const FName& InVariableName, const FString& InCPPType, UObject* InCPPTypeObject, bool bSetupUndoRedo, bool bSetupOrphanPins, bool bPrintPythonCommand)
 {
 	UObject* SourceObject = nullptr;
diff --git a/Engine/Plugins/Experimental/UAF/UAFAnimGraph/Source/UAFAnimGraphUncookedOnly/Private/AnimNextController.cpp b/Engine/Plugins/Experimental/UAF/UAFAnimGraph/Source/UAFAnimGraphUncookedOnly/Private/AnimNextController.cpp
--- a/Engine/Plugins/Experimental/UAF/UAFAnimGraph/Source/UAFAnimGraphUncookedOnly/Private/AnimNextController.cpp
+++ b/Engine/Plugins/Experimental/UAF/UAFAnimGraph/Source/UAFAnimGraphUncookedOnly/Private/AnimNextController.cpp
@@ -298,7 +298,7 @@ URigVMUnitNode* UAnimNextController::AddUnitNodeWithPins(UScriptStruct* InScript
 		OpenUndoBracket(TEXT("Add unit node with pins"));
 	}
 
-	URigVMUnitNode* Node = AddUnitNode(InScriptStruct, UAnimNextUnitNode::StaticClass(), InMethodName, InPosition, InNodeName, bSetupUndoRedo, bPrintPythonCommand);
+	URigVMUnitNode* Node = AddUnitNodeOfClass(InScriptStruct->GetPathName(), UAnimNextUnitNode::StaticClass()->GetPathName(), InMethodName, InPosition, InNodeName, bSetupUndoRedo, bPrintPythonCommand);
 
 	if (Node == nullptr)
 	{
diff --git a/Engine/Plugins/Experimental/UAF/UAFAnimGraph/Source/UAFAnimGraphUncookedOnly/Private/Templates/UAFGraphNodeTemplate.cpp b/Engine/Plugins/Experimental/UAF/UAFAnimGraph/Source/UAFAnimGraphUncookedOnly/Private/Templates/UAFGraphNodeTemplate.cpp
--- a/Engine/Plugins/Experimental/UAF/UAFAnimGraph/Source/UAFAnimGraphUncookedOnly/Private/Templates/UAFGraphNodeTemplate.cpp
+++ b/Engine/Plugins/Experimental/UAF/UAFAnimGraph/Source/UAFAnimGraphUncookedOnly/Private/Templates/UAFGraphNodeTemplate.cpp
@@ -244,7 +244,7 @@ void UUAFGraphNodeTemplate::SetCategoryForPinsInLayout(const TArray<FString>& Pi
 
 URigVMUnitNode* UUAFGraphNodeTemplate::CreateNewNode(UAnimNextController* Controller, const FVector2D& InLocation) const
 {
-	UAnimNextTraitStackUnitNode* TraitStackNode = Cast<UAnimNextTraitStackUnitNode>(Controller->AddUnitNode(FRigUnit_AnimNextTraitStack::StaticStruct(), UAnimNextTraitStackUnitNode::StaticClass(), FRigVMStruct::ExecuteName, InLocation, GetName(), true, true));
+	UAnimNextTraitStackUnitNode* TraitStackNode = Cast<UAnimNextTraitStackUnitNode>(Controller->AddUnitNodeOfClass(FRigUnit_AnimNextTraitStack::StaticStruct()->GetPathName(), UAnimNextTraitStackUnitNode::StaticClass()->GetPathName(), FRigVMStruct::ExecuteName, InLocation, GetName(), true, true));
 	if (TraitStackNode == nullptr)
 	{
 		Controller->CancelUndoBracket();
@@ -351,4 +351,4 @@ void UUAFGraphNodeTemplate::RefreshNodeFromTemplate(UAnimNextController* InContr
 	InController->CloseUndoBracket();
 }
 
-#undef LOCTEXT_NAMESPACE
\ No newline at end of file
+#undef LOCTEXT_NAMESPACE
