Fix being able to open multiple UAF asset wizard windows
* Also show error when trying to make a new BP asset with a clashing name

diff --git a/Engine/Plugins/Experimental/UAF/UAF/Source/UAFEditor/Private/Wizard/AssetWizard.cpp b/Engine/Plugins/Experimental/UAF/UAF/Source/UAFEditor/Private/Wizard/AssetWizard.cpp
--- a/Engine/Plugins/Experimental/UAF/UAF/Source/UAFEditor/Private/Wizard/AssetWizard.cpp
+++ b/Engine/Plugins/Experimental/UAF/UAF/Source/UAFEditor/Private/Wizard/AssetWizard.cpp
@@ -223,6 +223,15 @@ namespace UE::UAF::Editor
 
 	void FAssetWizard::ShowTemplatePicker()
 	{
+		static TWeakPtr<SWindow> AssetWizardWindowWeak;
+
+		TSharedPtr<SWindow> AssetWizardWindow = AssetWizardWindowWeak.Pin();
+		if (AssetWizardWindow.IsValid())
+		{
+			AssetWizardWindow->BringToFront(true);
+			return;
+		}
+			
 		FContentBrowserModule& ContentBrowserModule = FModuleManager::LoadModuleChecked<FContentBrowserModule>("ContentBrowser");
 
 		TArray<FString> SelectedPaths;
@@ -236,6 +245,13 @@ namespace UE::UAF::Editor
 			.SupportsMinimize(false)
 			.SupportsMaximize(false);
 
+		AssetWizardWindowWeak = TemplateWindow;
+
+		TemplateWindow->SetOnWindowClosed(FOnWindowClosed::CreateLambda([](const TSharedRef<SWindow>&)
+		{
+			AssetWizardWindowWeak.Reset();
+		}));
+
 		TemplateWindow->SetContent(
 			SNew(STemplatePicker)
 			.OnTemplateSelected_Lambda([TemplateWindow](const TObjectPtr<const UUAFTemplateDataAsset> InTemplate, const TObjectPtr<const UUAFTemplateConfig> InTemplateConfig) mutable
diff --git a/Engine/Plugins/Experimental/UAF/UAF/Source/UAFEditor/Private/Wizard/STemplatePicker.cpp b/Engine/Plugins/Experimental/UAF/UAF/Source/UAFEditor/Private/Wizard/STemplatePicker.cpp
--- a/Engine/Plugins/Experimental/UAF/UAF/Source/UAFEditor/Private/Wizard/STemplatePicker.cpp
+++ b/Engine/Plugins/Experimental/UAF/UAF/Source/UAFEditor/Private/Wizard/STemplatePicker.cpp
@@ -703,6 +703,13 @@ namespace UE::UAF::Editor
 					return false;
 				}
 			}
+
+			UEditorAssetSubsystem* EditorAssetSubsystem = GEditor->GetEditorSubsystem<UEditorAssetSubsystem>();
+			if (EditorAssetSubsystem->DoesAssetExist(FPaths::Combine(TemplateConfig->OutputPath.Path, TemplateConfig->BlueprintAssetName)))
+			{
+				ErrorText = LOCTEXT("ClashingBlueprintAssetName", "Blueprint asset name clashes with an existing asset");
+				return false;
+			}
 		}
 
 		return true;
