UAF: Fix monolithic editor build linker issue with FUAFAssetPreviewFactory

diff --git a/Engine/Plugins/Experimental/UAF/UAF/Source/UAFEditor/Public/Common/AssetPreview/IUAFAssetPreview.h b/Engine/Plugins/Experimental/UAF/UAF/Source/UAFEditor/Public/Common/AssetPreview/IUAFAssetPreview.h
--- a/Engine/Plugins/Experimental/UAF/UAF/Source/UAFEditor/Public/Common/AssetPreview/IUAFAssetPreview.h
+++ b/Engine/Plugins/Experimental/UAF/UAF/Source/UAFEditor/Public/Common/AssetPreview/IUAFAssetPreview.h
@@ -47,7 +47,7 @@ public:
 	/** Get the preview type this factory is used for. */
 	UE_API virtual const UStruct* GetPreviewType() const = 0;
 
-	UE_API virtual ~FUAFAssetPreviewFactory() = default;
+	virtual ~FUAFAssetPreviewFactory() = default;
 };
 
 
