Update workspace compile button
* Remove old compile options that no longer apply.
* Fix compile status not updating depending on focussed document.

diff --git a/Engine/Plugins/Experimental/UAF/UAF/Source/UAF/Public/AnimNextRigVMAsset.h b/Engine/Plugins/Experimental/UAF/UAF/Source/UAF/Public/AnimNextRigVMAsset.h
--- a/Engine/Plugins/Experimental/UAF/UAF/Source/UAF/Public/AnimNextRigVMAsset.h
+++ b/Engine/Plugins/Experimental/UAF/UAF/Source/UAF/Public/AnimNextRigVMAsset.h
@@ -59,6 +59,7 @@ class UAnimNextRigVMAsset : public URigVMHost
 protected:
 	friend class UAnimNextVariableEntry;
 	friend class UAnimNextRigVMAssetEditorData;
+	friend class UAnimNextWorkspaceEditorMode;
 	friend UAnimNextAnimGraphSettings;
 	friend FUAFRigVMComponent;
 	friend FUAFRigVMSystemScriptEventsComponent;
diff --git a/Engine/Plugins/Experimental/UAF/UAF/Source/UAFEditor/Private/AnimNextEditorModule.cpp b/Engine/Plugins/Experimental/UAF/UAF/Source/UAFEditor/Private/AnimNextEditorModule.cpp
--- a/Engine/Plugins/Experimental/UAF/UAF/Source/UAFEditor/Private/AnimNextEditorModule.cpp
+++ b/Engine/Plugins/Experimental/UAF/UAF/Source/UAFEditor/Private/AnimNextEditorModule.cpp
@@ -88,6 +88,7 @@
 #include "Interfaces/IPluginManager.h"
 #include "AssetToolsModule.h"
 #include "UAF/Debug/VariableWatchTabSummoner.h"
+#include "WorkspaceSchema.h"
 
 #define LOCTEXT_NAMESPACE "AnimNextEditorModule"
 
@@ -454,30 +455,9 @@ void FAnimNextEditorModule::StartupModule()
 						}
 
 						FTextBuilder TextBuilder;
-						if(EditorMode->GetState().bCompileDirtyFiles)
-						{
-							if(EditorMode->GetState().bCompileWholeWorkspace)
-							{
-								TextBuilder.AppendLine(LOCTEXT("CompileWorkspaceDirtyTooltip", "Compile all assets in the workspace that are dirty or have errors"));
-							}
-							else
-							{
-								TextBuilder.AppendLine(LOCTEXT("CompileDirtyTooltip", "Compile the current asset if it is dirty or has errors"));
-							}
-						}
-						else
-						{
-							if(EditorMode->GetState().bCompileWholeWorkspace)
-							{
-								TextBuilder.AppendLine(LOCTEXT("CompileWorkspaceTooltip", "Compile all assets in the workspace"));
-							}
-							else
-							{
-								TextBuilder.AppendLine(LOCTEXT("CompileCurrentTooltip", "Compile the current asset"));
-							}
-						}
+						TextBuilder.AppendLine(LOCTEXT("CompileCurrentTooltip", "Compile the current asset"));
 
-						if(EditorMode->GetState().bAutoCompile)
+						if(EditorMode->IsAutoCompileChecked())
 						{
 							TextBuilder.AppendLine(LOCTEXT("AutoCompileEnabledTooltip", "Auto-compilation is enabled"));
 						}
@@ -498,8 +478,6 @@ void FAnimNextEditorModule::StartupModule()
 						{
 							FToolMenuSection& CompileOptionsSection = InToolMenu->AddSection("CompileOptionsSection", LOCTEXT("CompileOptionsComboLabel", "Compilation Options"));
 							CompileOptionsSection.AddEntry(FToolMenuEntry::InitMenuEntry(FRigVMCommands::Get().AutoCompile, TAttribute<FText>(), TAttribute<FText>(), FSlateIcon("RigVMEditorStyle", "RigVM.AutoCompileGraph")));
-							CompileOptionsSection.AddEntry(FToolMenuEntry::InitMenuEntry(FRigVMCommands::Get().CompileWholeWorkspace, TAttribute<FText>(), TAttribute<FText>(), FSlateIcon(FAppStyle::GetAppStyleSetName(), "Icons.World")));
-							CompileOptionsSection.AddEntry(FToolMenuEntry::InitMenuEntry(FRigVMCommands::Get().CompileDirtyFiles, TAttribute<FText>(), TAttribute<FText>(), FSlateIcon(FAppStyle::GetAppStyleSetName(), "Icons.DirtyBadge")));
 						}),
 						FText::GetEmpty(),
 						LOCTEXT("CompileOptionsComboTooltip", "Compilation Options"),
diff --git a/Engine/Plugins/Experimental/UAF/UAF/Source/UAFEditor/Private/RigVMCommands.cpp b/Engine/Plugins/Experimental/UAF/UAF/Source/UAFEditor/Private/RigVMCommands.cpp
--- a/Engine/Plugins/Experimental/UAF/UAF/Source/UAFEditor/Private/RigVMCommands.cpp
+++ b/Engine/Plugins/Experimental/UAF/UAF/Source/UAFEditor/Private/RigVMCommands.cpp
@@ -1,4 +1,4 @@
-﻿// Copyright Epic Games, Inc. All Rights Reserved.
+// Copyright Epic Games, Inc. All Rights Reserved.
 
 #include "RigVMCommands.h"
 
@@ -16,8 +16,6 @@ void FRigVMCommands::RegisterCommands()
 {
 	UI_COMMAND(Compile, "Compile", "Compile all relevant assets", EUserInterfaceActionType::Button, FInputChord());
 	UI_COMMAND(AutoCompile, "Auto Compile", "Automatically compile on every edit", EUserInterfaceActionType::ToggleButton, FInputChord());
-	UI_COMMAND(CompileWholeWorkspace, "Compile Whole Workspace", "When manually compiling, whether to compile the current file or the whole workspace", EUserInterfaceActionType::ToggleButton, FInputChord());
-	UI_COMMAND(CompileDirtyFiles, "Compile Dirty Files", "When manually compiling, whether to compile only dirty files or files with errors. Files that are not dirty and do not have errors will get skipped.", EUserInterfaceActionType::ToggleButton, FInputChord());
 }
 
 }
diff --git a/Engine/Plugins/Experimental/UAF/UAF/Source/UAFEditor/Private/RigVMCommands.h b/Engine/Plugins/Experimental/UAF/UAF/Source/UAFEditor/Private/RigVMCommands.h
--- a/Engine/Plugins/Experimental/UAF/UAF/Source/UAFEditor/Private/RigVMCommands.h
+++ b/Engine/Plugins/Experimental/UAF/UAF/Source/UAFEditor/Private/RigVMCommands.h
@@ -1,4 +1,4 @@
-﻿// Copyright Epic Games, Inc. All Rights Reserved.
+// Copyright Epic Games, Inc. All Rights Reserved.
 
 #pragma once
 
@@ -17,8 +17,6 @@ public:
 
 	TSharedPtr<FUICommandInfo> Compile;
 	TSharedPtr<FUICommandInfo> AutoCompile;
-	TSharedPtr<FUICommandInfo> CompileWholeWorkspace;
-	TSharedPtr<FUICommandInfo> CompileDirtyFiles;
 };
 
 }
diff --git a/Engine/Plugins/Experimental/UAF/UAF/Source/UAFEditor/Private/Workspace/AnimNextWorkspaceEditorMode.cpp b/Engine/Plugins/Experimental/UAF/UAF/Source/UAFEditor/Private/Workspace/AnimNextWorkspaceEditorMode.cpp
--- a/Engine/Plugins/Experimental/UAF/UAF/Source/UAFEditor/Private/Workspace/AnimNextWorkspaceEditorMode.cpp
+++ b/Engine/Plugins/Experimental/UAF/UAF/Source/UAFEditor/Private/Workspace/AnimNextWorkspaceEditorMode.cpp
@@ -67,6 +67,8 @@ void UAnimNextWorkspaceEditorMode::Enter()
 		EditorData->RigVMCompiledEvent.AddUObject(this, &UAnimNextWorkspaceEditorMode::HandleRigVMCompiledEvent);
 		EditorData->RigVMGraphModifiedEvent.RemoveAll(this);
 		EditorData->RigVMGraphModifiedEvent.AddUObject(this, &UAnimNextWorkspaceEditorMode::HandleRigVMModifiedEvent);
+		EditorData->RecompileRequiredChangedEvent.RemoveAll(this);
+		EditorData->RecompileRequiredChangedEvent.AddUObject(this, &UAnimNextWorkspaceEditorMode::UpdateCompileStatus);
 
 		WeakAssets.Add(Asset);
 	}
@@ -151,8 +153,6 @@ void UAnimNextWorkspaceEditorMode::Exit()
 		const UE::UAF::FRigVMCommands& RigVMCommands = UE::UAF::FRigVMCommands::Get();
 		CommandList->UnmapAction(RigVMCommands.Compile);
 		CommandList->UnmapAction(RigVMCommands.AutoCompile);
-		CommandList->UnmapAction(RigVMCommands.CompileWholeWorkspace);
-		CommandList->UnmapAction(RigVMCommands.CompileDirtyFiles);
 
 		if (UToolMenu* Menu = UToolMenus::Get()->ExtendMenu(WorkspaceEditor->GetToolMenuToolbarName()))
 		{
@@ -182,6 +182,7 @@ void UAnimNextWorkspaceEditorMode::Exit()
 
 		EditorData->RigVMCompiledEvent.RemoveAll(this);
 		EditorData->RigVMGraphModifiedEvent.RemoveAll(this);
+		EditorData->RecompileRequiredChangedEvent.RemoveAll(this);
 	}
 
 	WeakAssets.Reset();
@@ -210,16 +211,6 @@ void UAnimNextWorkspaceEditorMode::BindCommands()
 		FExecuteAction::CreateUObject(this, &UAnimNextWorkspaceEditorMode::HandleAutoCompile),
 		FCanExecuteAction::CreateUObject(this, &UAnimNextWorkspaceEditorMode::CanCompile),
 		FIsActionChecked::CreateUObject(this, &UAnimNextWorkspaceEditorMode::IsAutoCompileChecked));
-
-	ToolkitCommands->MapAction(RigVMCommands.CompileWholeWorkspace,
-		FExecuteAction::CreateUObject(this, &UAnimNextWorkspaceEditorMode::HandleCompileWholeWorkspace),
-		FCanExecuteAction::CreateUObject(this, &UAnimNextWorkspaceEditorMode::CanCompile),
-		FIsActionChecked::CreateUObject(this, &UAnimNextWorkspaceEditorMode::IsCompileWholeWorkspaceChecked));
-
-	ToolkitCommands->MapAction(RigVMCommands.CompileDirtyFiles,
-		FExecuteAction::CreateUObject(this, &UAnimNextWorkspaceEditorMode::HandleCompileDirtyFiles),
-		FCanExecuteAction::CreateUObject(this, &UAnimNextWorkspaceEditorMode::CanCompile),
-		FIsActionChecked::CreateUObject(this, &UAnimNextWorkspaceEditorMode::IsCompileDirtyFilesChecked));
 }
 
 void UAnimNextWorkspaceEditorMode::HandleCompile()
@@ -233,23 +224,7 @@ void UAnimNextWorkspaceEditorMode::HandleCompile()
 	TArray<UAnimNextRigVMAsset*> AssetsToCompile;
 
 	FText JobName = LOCTEXT("DefaultJobName", "Job");
-	if(IsCompileWholeWorkspaceChecked())
-	{
-		TArray<UObject*> OpenObjects;
-		WorkspaceEditor->GetOpenedAssets<UObject>(OpenObjects);
-
-		Algo::TransformIf(OpenObjects, AssetsToCompile, [](UObject* Obj){ return Cast<UAnimNextRigVMAsset>(Obj); }, [](UObject* Obj){ return CastChecked<UAnimNextRigVMAsset>(Obj); } );
-
-		if(AssetsToCompile.Num() == 1)
-		{
-			JobName = FText::FromName(AssetsToCompile[0]->GetFName());
-		}
-		else
-		{
-			JobName = FText::FromName(WorkspaceEditor->GetWorkspaceAsset()->GetFName());
-		}
-	}
-	else if(UObject* FocusedDocument = WorkspaceEditor->GetFocusedWorkspaceDocument().GetObject())
+	if(UObject* FocusedDocument = WorkspaceEditor->GetFocusedWorkspaceDocument().GetObject())
 	{
 		// Find the outer asset for this document
 		UObject* FocussedObject = FocusedDocument;
@@ -265,24 +240,6 @@ void UAnimNextWorkspaceEditorMode::HandleCompile()
 		}
 	}
 
-	if(IsCompileDirtyFilesChecked())
-	{
-		for(int32 AssetIndex = 0; AssetIndex < AssetsToCompile.Num(); ++AssetIndex)
-		{
-			UObject* AssetToCheck = AssetsToCompile[AssetIndex];
-			if(TSharedPtr<UE::UAF::Editor::IAssetCompilationHandler> FoundCompiler = GetAssetCompiler(AssetToCheck))
-			{
-				UE::UAF::Editor::ECompileStatus AssetCompileStatus = FoundCompiler->GetCompileStatus(WorkspaceEditor.ToSharedRef(), AssetToCheck);
-				if( AssetCompileStatus != UE::UAF::Editor::ECompileStatus::Dirty &&
-					AssetCompileStatus != UE::UAF::Editor::ECompileStatus::Error)
-				{
-					AssetsToCompile.RemoveAtSwap(AssetIndex);
-					AssetIndex--;
-				}
-			}
-		}
-	}
-
 	if(AssetsToCompile.Num() == 0)
 	{
 		return;
@@ -309,52 +266,30 @@ bool UAnimNextWorkspaceEditorMode::CanCompile() const
 
 void UAnimNextWorkspaceEditorMode::HandleAutoCompile()
 {
-	State.bAutoCompile = !State.bAutoCompile;
-
 	TSharedPtr<UE::Workspace::IWorkspaceEditor> WorkspaceEditor = GetWorkspaceEditor();
-	if(WorkspaceEditor.IsValid())
+	if (WorkspaceEditor.IsValid())
 	{
-		PropagateAutoCompile(WorkspaceEditor.ToSharedRef(), State.bAutoCompile);
+		if (UAnimNextRigVMAsset* FocusedAsset = Cast<UAnimNextRigVMAsset>(WorkspaceEditor->GetFocusedAsset()))
+		{
+			UAnimNextRigVMAssetEditorData* FocusedAssetEditorData = CastChecked<UAnimNextRigVMAssetEditorData>(FocusedAsset->EditorData);
+			FocusedAssetEditorData->SetAutoVMRecompile(!FocusedAssetEditorData->GetAutoVMRecompile());
+		}
 	}
 }
 
-void UAnimNextWorkspaceEditorMode::PropagateAutoCompile(TSharedRef<UE::Workspace::IWorkspaceEditor> InWorkspaceEditor, bool bInAutoCompile)
+bool UAnimNextWorkspaceEditorMode::IsAutoCompileChecked() const
 {
-	TArray<UObject*> Assets;
-	InWorkspaceEditor->GetOpenedAssets(Assets);
-	UE::UAF::Editor::FAnimNextEditorModule& AnimNextEditorModule = FModuleManager::LoadModuleChecked<UE::UAF::Editor::FAnimNextEditorModule>("UAFEditor");
-	for(UObject* Asset : Assets)
+	TSharedPtr<UE::Workspace::IWorkspaceEditor> WorkspaceEditor = GetWorkspaceEditor();
+	if (WorkspaceEditor.IsValid())
 	{
-		if(TSharedPtr<UE::UAF::Editor::IAssetCompilationHandler> FoundCompiler = GetAssetCompiler(Asset))
+		if (UAnimNextRigVMAsset* FocusedAsset = Cast<UAnimNextRigVMAsset>(WorkspaceEditor->GetFocusedAsset()))
 		{
-			FoundCompiler->SetAutoCompile(InWorkspaceEditor, Asset, bInAutoCompile);
+			UAnimNextRigVMAssetEditorData* FocusedAssetEditorData = CastChecked<UAnimNextRigVMAssetEditorData>(FocusedAsset->EditorData);
+			return FocusedAssetEditorData->GetAutoVMRecompile();
 		}
 	}
-}
 
-bool UAnimNextWorkspaceEditorMode::IsAutoCompileChecked() const
-{
-	return State.bAutoCompile;
-}
-
-void UAnimNextWorkspaceEditorMode::HandleCompileWholeWorkspace()
-{
-	State.bCompileWholeWorkspace = !State.bCompileWholeWorkspace;
-}
-
-bool UAnimNextWorkspaceEditorMode::IsCompileWholeWorkspaceChecked() const
-{
-	return State.bCompileWholeWorkspace;
-}
-
-void UAnimNextWorkspaceEditorMode::HandleCompileDirtyFiles()
-{
-	State.bCompileDirtyFiles = !State.bCompileDirtyFiles;
-}
-
-bool UAnimNextWorkspaceEditorMode::IsCompileDirtyFilesChecked() const
-{
-	return State.bCompileDirtyFiles;
+	return false;
 }
 
 void UAnimNextWorkspaceEditorMode::HandleFocusedDocumentChanged(const UE::Workspace::FWorkspaceDocument& InDocument)
@@ -377,73 +312,41 @@ void UAnimNextWorkspaceEditorMode::HandleFocusedDocumentChanged(const UE::Worksp
 		return;
 	}
 
-	// Ensure auto-compilation is propagated for newly opened asset 
-	UAnimNextRigVMAssetEditorData* EditorData = UE::UAF::UncookedOnly::FUtils::GetEditorData<UAnimNextRigVMAssetEditorData>(AnimNextRigVMAsset);
-	EditorData->SetAutoVMRecompile(State.bAutoCompile);
-
 	// Subscribe to asset compilation/modification
+	UAnimNextRigVMAssetEditorData* EditorData = UE::UAF::UncookedOnly::FUtils::GetEditorData<UAnimNextRigVMAssetEditorData>(AnimNextRigVMAsset);
 	EditorData->RigVMCompiledEvent.RemoveAll(this);
 	EditorData->RigVMCompiledEvent.AddUObject(this, &UAnimNextWorkspaceEditorMode::HandleRigVMCompiledEvent);
 	EditorData->RigVMGraphModifiedEvent.RemoveAll(this);
 	EditorData->RigVMGraphModifiedEvent.AddUObject(this, &UAnimNextWorkspaceEditorMode::HandleRigVMModifiedEvent);
+	EditorData->RecompileRequiredChangedEvent.RemoveAll(this);
+	EditorData->RecompileRequiredChangedEvent.AddUObject(this, &UAnimNextWorkspaceEditorMode::UpdateCompileStatus);
 
 	WeakAssets.Add(AnimNextRigVMAsset);
+
+	UpdateCompileStatus();
 }
 
 void UAnimNextWorkspaceEditorMode::UpdateCompileStatus()
 {
-	bUpdateCompileStatus = false;
-
 	TSharedPtr<UE::Workspace::IWorkspaceEditor> WorkspaceEditor = GetWorkspaceEditor();
 	if(!WorkspaceEditor.IsValid())
 	{
 		return;
 	}
 
-	int32 NumAssetsWithWarnings = 0;
-	int32 NumAssetsWithErrors = 0;
-	int32 NumAssetsDirty = 0;
-
-	TArray<UObject*> Assets;
-	WorkspaceEditor->GetOpenedAssets<UAnimNextRigVMAsset>(Assets);
+	UAnimNextRigVMAsset* Asset = WorkspaceEditor->GetFocusedWorkspaceDocument().Export.GetFirstObjectOfType<UAnimNextRigVMAsset>();
 
-	for(UObject* Asset : Assets)
+	if (Asset)
 	{
-		if(TSharedPtr<UE::UAF::Editor::IAssetCompilationHandler> FoundCompiler = GetAssetCompiler(Asset))
+		if (TSharedPtr<UE::UAF::Editor::IAssetCompilationHandler> FoundCompiler = GetAssetCompiler(Asset))
 		{
 			UE::UAF::Editor::ECompileStatus AssetCompileStatus = FoundCompiler->GetCompileStatus(WorkspaceEditor.ToSharedRef(), Asset);
-			if(AssetCompileStatus == UE::UAF::Editor::ECompileStatus::Dirty)
-			{
-				NumAssetsDirty++;
-			}
-
-			if(AssetCompileStatus == UE::UAF::Editor::ECompileStatus::Error)
-			{
-				NumAssetsWithErrors++;
-			}
-
-			if(AssetCompileStatus == UE::UAF::Editor::ECompileStatus::Warning)
-			{
-				NumAssetsWithWarnings++;
-			}
+			CompileStatus = AssetCompileStatus;
 		}
 	}
-
-	if(NumAssetsWithErrors > 0)
-	{
-		CompileStatus = UE::UAF::Editor::ECompileStatus::Error;
-	}
-	else if(NumAssetsWithWarnings > 0)
-	{
-		CompileStatus = UE::UAF::Editor::ECompileStatus::Warning;
-	}
-	else if(NumAssetsDirty > 0)
-	{
-		CompileStatus = UE::UAF::Editor::ECompileStatus::Dirty;
-	}
 	else
 	{
-		CompileStatus = UE::UAF::Editor::ECompileStatus::UpToDate;
+		CompileStatus = UE::UAF::Editor::ECompileStatus::Unknown;
 	}
 }
 
diff --git a/Engine/Plugins/Experimental/UAF/UAF/Source/UAFEditor/Private/Workspace/AnimNextWorkspaceEditorMode.h b/Engine/Plugins/Experimental/UAF/UAF/Source/UAFEditor/Private/Workspace/AnimNextWorkspaceEditorMode.h
--- a/Engine/Plugins/Experimental/UAF/UAF/Source/UAFEditor/Private/Workspace/AnimNextWorkspaceEditorMode.h
+++ b/Engine/Plugins/Experimental/UAF/UAF/Source/UAFEditor/Private/Workspace/AnimNextWorkspaceEditorMode.h
@@ -47,6 +47,8 @@ public:
 
 	TSharedPtr<UE::UAF::Editor::IDebugObjectSelector> GetDebugObjectSelector();
 
+	bool IsAutoCompileChecked() const;
+
 private:
 	// UEdMode interface
 	virtual bool IsCompatibleWith(FEditorModeID OtherModeID) const override { return true; }
@@ -57,11 +59,6 @@ private:
 	void HandleCompile();
 	bool CanCompile() const;
 	void HandleAutoCompile();
-	bool IsAutoCompileChecked() const;
-	void HandleCompileWholeWorkspace();
-	bool IsCompileWholeWorkspaceChecked() const;
-	void HandleCompileDirtyFiles();
-	bool IsCompileDirtyFilesChecked() const;
 
 	// Updates the compile status. Scans all opened assets in the workspace.
 	void UpdateCompileStatus();
@@ -75,9 +72,6 @@ private:
 	// Handle the document focus changing
 	void HandleFocusedDocumentChanged(const UE::Workspace::FWorkspaceDocument& InDocument); 
 
-	// Forward auto-compilation to all assets in the workspace
-	void PropagateAutoCompile(TSharedRef<UE::Workspace::IWorkspaceEditor> InWorkspaceEditor, bool bInAutoCompile);
-
 	TSharedPtr<UE::Workspace::IWorkspaceEditor> GetWorkspaceEditor() const;
 
 	// Lazily construct an asset compiler for the supplied asset
@@ -88,7 +82,6 @@ private:
 	friend UAnimNextWorkspaceSchema;
 	friend UE::UAF::Editor::FAssetWizard;
 
-	// Custom state, persisted via UAnimNextWorkspaceSchema
 	FAnimNextWorkspaceState State;
 
 	// Asset compilers for all current assets 
@@ -100,9 +93,6 @@ private:
 	// Current compilation status
 	UE::UAF::Editor::ECompileStatus CompileStatus = UE::UAF::Editor::ECompileStatus::Unknown;
 
-	// Flag to update the latest compile status
-	bool bUpdateCompileStatus = false;
-
 	TSharedPtr<UE::UAF::Editor::FRewindDebuggerWorkspaceExtension> RewindDebuggerExtension;
 
 	TSharedPtr<UE::UAF::Editor::IDebugObjectSelector> DebugObjectSelector;
diff --git a/Engine/Plugins/Experimental/UAF/UAF/Source/UAFEditor/Private/Workspace/AnimNextWorkspaceSchema.cpp b/Engine/Plugins/Experimental/UAF/UAF/Source/UAFEditor/Private/Workspace/AnimNextWorkspaceSchema.cpp
--- a/Engine/Plugins/Experimental/UAF/UAF/Source/UAFEditor/Private/Workspace/AnimNextWorkspaceSchema.cpp
+++ b/Engine/Plugins/Experimental/UAF/UAF/Source/UAFEditor/Private/Workspace/AnimNextWorkspaceSchema.cpp
@@ -53,7 +53,6 @@ void UAnimNextWorkspaceSchema::OnLoadWorkspaceState(TSharedRef< UE::Workspace::I
 	}
 
 	EditorMode->State = *State;
-	EditorMode->PropagateAutoCompile(InWorkspaceEditor, EditorMode->State.bAutoCompile);
 }
 
 #undef LOCTEXT_NAMESPACE
diff --git a/Engine/Plugins/Experimental/UAF/UAF/Source/UAFEditor/Private/Workspace/AnimNextWorkspaceState.h b/Engine/Plugins/Experimental/UAF/UAF/Source/UAFEditor/Private/Workspace/AnimNextWorkspaceState.h
--- a/Engine/Plugins/Experimental/UAF/UAF/Source/UAFEditor/Private/Workspace/AnimNextWorkspaceState.h
+++ b/Engine/Plugins/Experimental/UAF/UAF/Source/UAFEditor/Private/Workspace/AnimNextWorkspaceState.h
@@ -1,4 +1,4 @@
-﻿// Copyright Epic Games, Inc. All Rights Reserved.
+// Copyright Epic Games, Inc. All Rights Reserved.
 
 #pragma once
 
@@ -9,16 +9,4 @@ USTRUCT()
 struct FAnimNextWorkspaceState
 {
 	GENERATED_BODY()
-
-	// Whether to automatically compile RigVM files after every edit
-	UPROPERTY()
-	bool bAutoCompile = true;
-
-	// When manually compiling, whether to compile the current file or the whole workspace
-	UPROPERTY()
-	bool bCompileWholeWorkspace = true;
-
-	// When manually compiling, whether to compile only dirty files. Files that are not dirty will get skipped.
-	UPROPERTY()
-	bool bCompileDirtyFiles = true;
 };
\ No newline at end of file
diff --git a/Engine/Plugins/Experimental/UAF/UAF/Source/UAFUncookedOnly/Internal/AnimNextRigVMAssetEditorData.h b/Engine/Plugins/Experimental/UAF/UAF/Source/UAFUncookedOnly/Internal/AnimNextRigVMAssetEditorData.h
--- a/Engine/Plugins/Experimental/UAF/UAF/Source/UAFUncookedOnly/Internal/AnimNextRigVMAssetEditorData.h
+++ b/Engine/Plugins/Experimental/UAF/UAF/Source/UAFUncookedOnly/Internal/AnimNextRigVMAssetEditorData.h
@@ -553,6 +553,8 @@ protected:
 
 	// Called from PostLoad to load external packages
 	UE_API void PostLoadExternalPackages();
+	
+	UE_API void RequestVMRecompilation();
 
 	// Find an entry that corresponds to the specified RigVMGraph. This uses the name of the graph to match the entry 
 	UE_API UAnimNextRigVMAssetEntry* FindEntryForRigVMGraph(const URigVMGraph* InRigVMGraph) const;
@@ -663,6 +665,8 @@ protected:
 	// Serialization/post-load fixup for old data
 	void UpgradeDataInterfaces();
 
+	void SetVMRecompilationRequired(bool bInIsRequired);
+
 	/** All entries in this asset - not saved, discovered at load time and also contains InternalEntries */
 	UPROPERTY(Transient, DuplicateTransient)
 	TArray<TObjectPtr<UAnimNextRigVMAssetEntry>> Entries;
@@ -738,6 +742,9 @@ protected:
 
 	FRigVMGraphModifiedEvent RigVMGraphModifiedEvent;
 
+	// Compile status depends upon bVMRecompilationRequired so it useful to know when this value is changed
+	FSimpleMulticastDelegate RecompileRequiredChangedEvent;
+
 	// Delegate to subscribe to modifications to this editor data
 	UE::UAF::UncookedOnly::FOnEditorDataModified ModifiedDelegate;
 
@@ -758,7 +765,10 @@ protected:
 	UPROPERTY()
 	TObjectPtr<UAnimNextEdGraph> FunctionLibraryEdGraph;
 
+	UPROPERTY()
 	bool bAutoRecompileVM = true;
+
+	bool bVMRecompilationRequested = false;
 	mutable bool bErrorsDuringCompilation = false;
 	mutable bool bWarningsDuringCompilation = false;
 	bool bSuspendModelNotificationsForSelf = false;
diff --git a/Engine/Plugins/Experimental/UAF/UAF/Source/UAFUncookedOnly/Private/AnimNextRigVMAssetEditorData.cpp b/Engine/Plugins/Experimental/UAF/UAF/Source/UAFUncookedOnly/Private/AnimNextRigVMAssetEditorData.cpp
--- a/Engine/Plugins/Experimental/UAF/UAF/Source/UAFUncookedOnly/Private/AnimNextRigVMAssetEditorData.cpp
+++ b/Engine/Plugins/Experimental/UAF/UAF/Source/UAFUncookedOnly/Private/AnimNextRigVMAssetEditorData.cpp
@@ -264,7 +264,7 @@ void UAnimNextRigVMAssetEditorData::PostLoad()
 	}
 
 	// Mark this as being dirty so that we recompile when needed
-	bVMRecompilationRequired = true;
+	SetVMRecompilationRequired(true);
 
 	// Queue compilation once the package has been fully loaded
 	// This is necessary in case we have external packages that haven't post-loaded yet
@@ -291,6 +291,19 @@ void UAnimNextRigVMAssetEditorData::PostLoadExternalPackages()
 	Entries.Append(InternalEntries);
 }
 
+void UAnimNextRigVMAssetEditorData::RequestVMRecompilation()
+{
+	if (VMRecompilationBracket == 0)
+	{
+		RecompileVM();
+	}
+    else
+    {
+		SetVMRecompilationRequired(true);
+		bVMRecompilationRequested = true;
+	}
+}
+
 void UAnimNextRigVMAssetEditorData::PostEditChangeProperty(FPropertyChangedEvent& PropertyChangedEvent)
 {
 	Super::PostEditChangeProperty(PropertyChangedEvent);
@@ -1142,7 +1155,8 @@ void UAnimNextRigVMAssetEditorData::RecompileVM()
 			}
 		}
 
-		bVMRecompilationRequired = false;
+		SetVMRecompilationRequired(false);
+		bVMRecompilationRequested = false;
 
 		if(Asset->VM)
 		{
@@ -1401,16 +1415,29 @@ void UAnimNextRigVMAssetEditorData::RecompileVMIfRequired()
 
 void UAnimNextRigVMAssetEditorData::RequestAutoVMRecompilation()
 {
-	bVMRecompilationRequired = true;
+	SetVMRecompilationRequired(true);
 	if (bAutoRecompileVM && VMRecompilationBracket == 0)
 	{
 		RecompileVMIfRequired();
 	}
 }
 
+void UAnimNextRigVMAssetEditorData::SetVMRecompilationRequired(bool bInIsRequired)
+{
+	if (bInIsRequired != bVMRecompilationRequested)
+	{
+		bVMRecompilationRequested = bInIsRequired;
+		RecompileRequiredChangedEvent.Broadcast();
+	}
+}
+
 void UAnimNextRigVMAssetEditorData::SetAutoVMRecompile(bool bAutoRecompile)
 {
-	bAutoRecompileVM = bAutoRecompile;
+	if (bAutoRecompileVM != bAutoRecompile)
+	{
+		bAutoRecompileVM = bAutoRecompile;
+		MarkPackageDirty();
+	}
 }
 
 bool UAnimNextRigVMAssetEditorData::GetAutoVMRecompile() const
@@ -1427,10 +1454,13 @@ void UAnimNextRigVMAssetEditorData::DecrementVMRecompileBracket()
 {
 	if (VMRecompilationBracket == 1)
 	{
-		if (bAutoRecompileVM)
+		// We should only recompile when the asset is dirty and we have auto-compile enabled
+		// or when a recompilation has been explicitly requested (e.g. clicking compile when auto-compile is disabled)
+		if ((bAutoRecompileVM && bVMRecompilationRequired) || bVMRecompilationRequested)
 		{
-			RecompileVMIfRequired();
+			RecompileVM();
 		}
+
 		VMRecompilationBracket = 0;
 
 		if (InteractionBracketFinished.IsBound())
diff --git a/Engine/Plugins/Experimental/UAF/UAF/Source/UAFUncookedOnly/Private/UAFCompilationScope.cpp b/Engine/Plugins/Experimental/UAF/UAF/Source/UAFUncookedOnly/Private/UAFCompilationScope.cpp
--- a/Engine/Plugins/Experimental/UAF/UAF/Source/UAFUncookedOnly/Private/UAFCompilationScope.cpp
+++ b/Engine/Plugins/Experimental/UAF/UAF/Source/UAFUncookedOnly/Private/UAFCompilationScope.cpp
@@ -31,7 +31,7 @@ void Compilation::RequestAssetCompilation(UAnimNextRigVMAsset* InAsset)
 {
 	FCompilationScope Scope(InAsset);
 	UAnimNextRigVMAssetEditorData* EditorData = FUtils::GetEditorData<UAnimNextRigVMAssetEditorData>(InAsset);
-	EditorData->RequestAutoVMRecompilation();
+	EditorData->RequestVMRecompilation();
 }
 
 void Compilation::RequestAssetCompilation(TConstArrayView<UAnimNextRigVMAsset*> InAssets)
@@ -44,7 +44,7 @@ void Compilation::RequestAssetCompilation(TConstArrayView<UAnimNextRigVMAsset*>
 		{
 			check(Asset);
 			UAnimNextRigVMAssetEditorData* EditorData = FUtils::GetEditorData<UAnimNextRigVMAssetEditorData>(Asset);
-			EditorData->RequestAutoVMRecompilation();
+			EditorData->RequestVMRecompilation();
 		}
 	}
 }
diff --git a/Engine/Plugins/Experimental/Workspace/Source/WorkspaceEditor/Private/WorkspaceEditor.h b/Engine/Plugins/Experimental/Workspace/Source/WorkspaceEditor/Private/WorkspaceEditor.h
--- a/Engine/Plugins/Experimental/Workspace/Source/WorkspaceEditor/Private/WorkspaceEditor.h
+++ b/Engine/Plugins/Experimental/Workspace/Source/WorkspaceEditor/Private/WorkspaceEditor.h
@@ -116,6 +116,7 @@ private:
 	virtual const FWorkspaceDocument& GetFocusedWorkspaceDocument() const override;
 	virtual void SetSceneDescription(TObjectPtr<UWorkspaceViewportSceneDescription> InSceneDescription) override;
 	virtual UWorkspaceViewportSceneDescription* GetSceneDescription() const override;
+	virtual UObject* GetFocusedAsset() const override;
 	
 	void ReinitViewportScene();
 
@@ -149,8 +150,6 @@ private:
 	bool AreAssetEntriesModified() const;
 
 	bool IsViewportPinned() const;
-	
-	UObject* GetFocusedAsset() const;
 
 	/** The asset being edited */
 	TObjectPtr<UWorkspace> Workspace = nullptr;
diff --git a/Engine/Plugins/Experimental/Workspace/Source/WorkspaceEditor/Public/IWorkspaceEditor.h b/Engine/Plugins/Experimental/Workspace/Source/WorkspaceEditor/Public/IWorkspaceEditor.h
--- a/Engine/Plugins/Experimental/Workspace/Source/WorkspaceEditor/Public/IWorkspaceEditor.h
+++ b/Engine/Plugins/Experimental/Workspace/Source/WorkspaceEditor/Public/IWorkspaceEditor.h
@@ -108,6 +108,9 @@ public:
 
 	// Get the current scene description config
 	virtual UWorkspaceViewportSceneDescription* GetSceneDescription() const = 0;
+
+	// Gets the currently focused asset's UObject
+	virtual UObject* GetFocusedAsset() const = 0;
 };
 
 }
