Add missing file from UAF LLM tags

diff --git a/Engine/Plugins/Experimental/UAF/UAF/Source/UAF/Private/AnimNextDebugDraw.cpp b/Engine/Plugins/Experimental/UAF/UAF/Source/UAF/Private/AnimNextDebugDraw.cpp
--- a/Engine/Plugins/Experimental/UAF/UAF/Source/UAF/Private/AnimNextDebugDraw.cpp
+++ b/Engine/Plugins/Experimental/UAF/UAF/Source/UAF/Private/AnimNextDebugDraw.cpp
@@ -144,6 +144,8 @@ FDebugDraw::FDebugDraw(UObject* InOwner)
 	: SceneProxyDesc(InOwner)
 	, SceneProxy(new FAnimNextDebugSceneProxy(SceneProxyDesc))
 {
+	LLM_SCOPE_BYNAME(TEXT("UAF/Debug"));
+
 	check(IsInGameThread());
 
 	SceneInfoData = new FPrimitiveSceneInfoData();
