UAF translational retargeting node

* Created retarget pose node
* Created sharable retargeting tools and modularized translational retargeting code

#Synced-CL 48139024

diff --git a/Engine/Plugins/Experimental/UAF/UAF/Source/UAF/Private/DecompressionTools.cpp b/Engine/Plugins/Experimental/UAF/UAF/Source/UAF/Private/DecompressionTools.cpp
--- a/Engine/Plugins/Experimental/UAF/UAF/Source/UAF/Private/DecompressionTools.cpp
+++ b/Engine/Plugins/Experimental/UAF/UAF/Source/UAF/Private/DecompressionTools.cpp
@@ -18,6 +18,7 @@
 #include "UAF/Attributes/AttributeCollection.h"
 #include "UAF/Attributes/AttributeEvaluationContext.h"
 #include "UAF/Attributes/EngineAttributes.h"
+#include "RetargetingTools.h"
 
 CSV_DECLARE_CATEGORY_MODULE_EXTERN(ENGINE_API, Animation);
 
@@ -1201,7 +1202,7 @@ void FDecompressionTools::DecompressPose(FLODPose& OutAnimationPoseData,
 			}
 
 			// @laurent - we should look into splitting rotation and translation tracks, so we don't have to process translation twice.
-			RetargetBoneTransform(OutAnimationPoseData.GetRefPose()
+			FRetargetingTools::RetargetBoneTransform(OutAnimationPoseData.GetRefPose()
 				, DecompressionContext.GetSourceSkeleton()
 				, OutAnimationPoseData.GetSkeletonAsset()
 				, DecompressionContext.AnimName
@@ -1243,20 +1244,15 @@ void FDecompressionTools::DecompressPose(FLODPose& OutAnimationPoseData,
 					OutAnimationPoseData.LocalTransformsView[LODBoneIndex].SetRotation(SkeletonRemapping.RetargetAdditiveRotationToTargetSkeleton(TargetSkeletonBoneIndex, OutAnimationPoseData.LocalTransformsView[LODBoneIndex].GetRotation()));
 				}
 
-				// Check what retarget mode to use for the translational retargeting for this specific bone.
-				const int32 SourceSkeletonBoneIndex = SkeletonRemapping.GetSourceSkeletonBoneIndex(TargetSkeletonBoneIndex);
-				const EBoneTranslationRetargetingMode::Type RetargetMode = FAnimationRuntime::GetBoneTranslationRetargetingMode(
-					bUseSourceRetargetModes,
-					SourceSkeletonBoneIndex,
-					TargetSkeletonBoneIndex,
+				FRetargetingTools::TranslationRetargetReferencePoseBone</*bIsAdditive=*/true>(
 					SourceSkeleton,
 					TargetSkeleton,
-					bDisableRetargeting);
-
-				if (RetargetMode != EBoneTranslationRetargetingMode::Skeleton)
-				{
-					OutAnimationPoseData.LocalTransformsView[LODBoneIndex].SetTranslation(SkeletonRemapping.RetargetAdditiveTranslationToTargetSkeleton(TargetSkeletonBoneIndex, OutAnimationPoseData.LocalTransformsView[LODBoneIndex].GetTranslation()));
-				}
+					SkeletonRemapping,
+					bUseSourceRetargetModes,
+					bDisableRetargeting,
+					LODBoneIndex,
+					TargetSkeletonBoneIndex,
+					OutAnimationPoseData);
 			}
 		}
 		else
@@ -1266,20 +1262,15 @@ void FDecompressionTools::DecompressPose(FLODPose& OutAnimationPoseData,
 				const int32 TargetSkeletonBoneIndex = LODBoneIndexToSkeletonBoneIndexMap[LODBoneIndex];
 				OutAnimationPoseData.LocalTransformsView[LODBoneIndex].SetRotation(SkeletonRemapping.RetargetBoneRotationToTargetSkeleton(TargetSkeletonBoneIndex, OutAnimationPoseData.LocalTransformsView[LODBoneIndex].GetRotation()));
 
-				// Check what retarget mode to use for the translational retargeting for this specific bone.
-				const int32 SourceSkeletonBoneIndex = SkeletonRemapping.GetSourceSkeletonBoneIndex(TargetSkeletonBoneIndex);
-				const EBoneTranslationRetargetingMode::Type RetargetMode = FAnimationRuntime::GetBoneTranslationRetargetingMode(
-					bUseSourceRetargetModes,
-					SourceSkeletonBoneIndex,
-					TargetSkeletonBoneIndex,
+				FRetargetingTools::TranslationRetargetReferencePoseBone</*bIsAdditive=*/false>(
 					SourceSkeleton,
 					TargetSkeleton,
-					bDisableRetargeting);
-
-				if (RetargetMode != EBoneTranslationRetargetingMode::Skeleton)
-				{
-					OutAnimationPoseData.LocalTransformsView[LODBoneIndex].SetTranslation(SkeletonRemapping.RetargetBoneTranslationToTargetSkeleton(TargetSkeletonBoneIndex, OutAnimationPoseData.LocalTransformsView[LODBoneIndex].GetTranslation()));
-				}
+					SkeletonRemapping,
+					bUseSourceRetargetModes,
+					bDisableRetargeting,
+					LODBoneIndex,
+					TargetSkeletonBoneIndex,
+					OutAnimationPoseData);				
 			}
 		}
 	}
@@ -1291,102 +1282,15 @@ void FDecompressionTools::DecompressPose(FLODPose& OutAnimationPoseData,
 		RootMotionReset.ResetRootBoneForRootMotion(RootTransform, ReferencePose.GetRefPoseTransform(0));
 		OutAnimationPoseData.LocalTransformsView[0] = RootTransform;
 	}
-
-	// Anim Scale Retargeting
-	int32 const NumBonesToScaleRetarget = AnimScaleRetargetingPairs.Num();
-	if (NumBonesToScaleRetarget > 0)
-	{
-		TArray<FTransform> const& AuthoredOnRefSkeleton = RetargetTransforms;
-
-		for (const BoneTrackPair& BonePair : AnimScaleRetargetingPairs)
-		{
-			const int32 LODBoneIndex(BonePair.AtomIndex);
-			const int32 SourceSkeletonBoneIndex = BonePair.TrackIndex;
-			if (!AuthoredOnRefSkeleton.IsValidIndex(SourceSkeletonBoneIndex))
-			{
-				continue;
-			}
-
-			// @todo - precache that in FBoneContainer when we have SkeletonIndex->TrackIndex mapping. So we can just apply scale right away.
-			float const SourceTranslationLength = AuthoredOnRefSkeleton[SourceSkeletonBoneIndex].GetTranslation().Size();
-			if (SourceTranslationLength > UE_KINDA_SMALL_NUMBER)
-			{
-				float const TargetTranslationLength = ReferencePose.GetRefPoseTranslation(LODBoneIndex).Size();
-				OutAnimationPoseData.LocalTransformsView[LODBoneIndex].ScaleTranslation(TargetTranslationLength / SourceTranslationLength);
-			}
-		}
-	}
-
-	// Anim Relative Retargeting
-	int32 const NumBonesToRelativeRetarget = AnimRelativeRetargetingPairs.Num();
-	if (NumBonesToRelativeRetarget > 0)
-	{
-		TArray<FTransform> const& AuthoredOnRefSkeleton = RetargetTransforms;
-
-		for (const BoneTrackPair& BonePair : AnimRelativeRetargetingPairs)
-		{
-			const int32 LODBoneIndex(BonePair.AtomIndex);
-			const int32 SourceSkeletonBoneIndex = BonePair.TrackIndex;
-			if (!AuthoredOnRefSkeleton.IsValidIndex(SourceSkeletonBoneIndex))
-			{
-				continue;
-			}
-
-			const FTransform& RefPoseTransform = ReferencePose.GetRefPoseTransform(LODBoneIndex);
-
-			// Remap the base pose onto the target skeleton so that we are working entirely in target space
-			FTransform BaseTransform = AuthoredOnRefSkeleton[SourceSkeletonBoneIndex];
-			if (SkeletonRemapping.RequiresReferencePoseRetarget())
-			{
-				const int32 TargetSkeletonBoneIndex = SkeletonRemapping.GetTargetSkeletonBoneIndex(SourceSkeletonBoneIndex);
-				BaseTransform = SkeletonRemapping.RetargetBoneTransformToTargetSkeleton(TargetSkeletonBoneIndex, BaseTransform);
-			}
-
-			// Apply the retargeting as if it were an additive difference between the current skeleton and the retarget skeleton. 
-			OutAnimationPoseData.LocalTransformsView[LODBoneIndex].SetRotation(OutAnimationPoseData.LocalTransformsView[LODBoneIndex].GetRotation() * BaseTransform.GetRotation().Inverse() * RefPoseTransform.GetRotation());
-			OutAnimationPoseData.LocalTransformsView[LODBoneIndex].SetTranslation(OutAnimationPoseData.LocalTransformsView[LODBoneIndex].GetTranslation() + (RefPoseTransform.GetTranslation() - BaseTransform.GetTranslation()));
-			OutAnimationPoseData.LocalTransformsView[LODBoneIndex].SetScale3D(OutAnimationPoseData.LocalTransformsView[LODBoneIndex].GetScale3D() * (RefPoseTransform.GetScale3D() * BaseTransform.GetSafeScaleReciprocal(BaseTransform.GetScale3D())));
-			OutAnimationPoseData.LocalTransformsView[LODBoneIndex].NormalizeRotation();
-		}
-	}
-
-	// TODO : Have to recreate GetRetargetSourceCachedData
-	// Translation 'Orient and Scale' Translation Retargeting
-	//const int32 NumBonesToOrientAndScaleRetarget = OrientAndScaleRetargetingPairs.Num();
-	//if (NumBonesToOrientAndScaleRetarget > 0)
-	//{
-	//	const FRetargetSourceCachedData& RetargetSourceCachedData = RequiredBones.GetRetargetSourceCachedData(DecompressionContext.AnimName, SkeletonRemapping, RetargetTransforms);
-	//	const TArray<FOrientAndScaleRetargetingCachedData>& OrientAndScaleDataArray = RetargetSourceCachedData.OrientAndScaleData;
-	//	const TArray<int32>& CompactPoseIndexToOrientAndScaleIndex = RetargetSourceCachedData.CompactPoseIndexToOrientAndScaleIndex;
-
-	//	// If we have any cached retargeting data.
-	//	if ((OrientAndScaleDataArray.Num() > 0) && (CompactPoseIndexToOrientAndScaleIndex.Num() == RequiredBones.GetCompactPoseNumBones()))
-	//	{
-	//		for (int32 Index = 0; Index < NumBonesToOrientAndScaleRetarget; Index++)
-	//		{
-	//			const BoneTrackPair& BonePair = OrientAndScaleRetargetingPairs[Index];
-	//			const FCompactPoseBoneIndex CompactPoseBoneIndex(BonePair.AtomIndex);
-	//			const int32 OrientAndScaleIndex = CompactPoseIndexToOrientAndScaleIndex[CompactPoseBoneIndex.GetInt()];
-	//			if (OrientAndScaleIndex != INDEX_NONE)
-	//			{
-	//				const FOrientAndScaleRetargetingCachedData& OrientAndScaleData = OrientAndScaleDataArray[OrientAndScaleIndex];
-	//				FTransform& BoneTransform = OutAnimationPoseData.LocalTransforms[CompactPoseBoneIndex];
-	//				const FVector AnimatedTranslation = BoneTransform.GetTranslation();
-
-	//				// If Translation is not animated, we can just copy the TargetTranslation. No retargeting needs to be done.
-	//				const FVector NewTranslation = (AnimatedTranslation - OrientAndScaleData.SourceTranslation).IsNearlyZero(BONE_TRANS_RT_ORIENT_AND_SCALE_PRECISION) ?
-	//					OrientAndScaleData.TargetTranslation :
-	//					OrientAndScaleData.TranslationDeltaOrient.RotateVector(AnimatedTranslation) * OrientAndScaleData.TranslationScale;
-
-	//				BoneTransform.SetTranslation(NewTranslation);
-	//			}
-	//		}
-	//	}
-	//}
-	UE_CLOGFMT(OrientAndScaleRetargetingPairs.Num() != 0,
-		LogAnimation, Verbose,
-		"OrientAndScaleRetargetingPairs found for {AnimSequence} ({SourceSkeleton} -> {TargetSkeleton}), Orient and Scale retargeting is not yet supported in UAF",
-		*DecompressionContext.AnimName.ToString(), SourceSkeleton ? *SourceSkeleton->GetName() : TEXT("None"), TargetSkeleton ? *TargetSkeleton->GetName() : TEXT("None"));
+	
+	FRetargetingTools::RetargetPose(
+		ReferencePose,
+		SkeletonRemapping,
+		AnimScaleRetargetingPairs,
+		AnimRelativeRetargetingPairs,
+		OrientAndScaleRetargetingPairs,
+		RetargetTransforms,
+		OutAnimationPoseData);
 
 	if (bIsMeshSpaceAdditive)
 	{
@@ -1414,131 +1318,4 @@ void FDecompressionTools::DecompressPose(FLODPose& OutAnimationPoseData,
 	}
 }
 
-void FDecompressionTools::RetargetBoneTransform(const FReferencePose& ReferencePose
-	, const USkeleton* SourceSkeleton
-	, const USkeleton* TargetSkeleton
-	, const FName& RetargetSource
-	, FTransform& BoneTransform
-	, const int32 SkeletonBoneIndex
-	, const int32 LODBoneIndex
-	, const bool bIsBakedAdditive
-	, const bool bDisableRetargeting)
-{
-	if (SourceSkeleton)
-	{
-		const TArray<FTransform>& RetargetTransforms = SourceSkeleton->GetRefLocalPoses(RetargetSource);
-		RetargetBoneTransform(ReferencePose, SourceSkeleton, TargetSkeleton, RetargetSource, RetargetTransforms, BoneTransform, SkeletonBoneIndex, LODBoneIndex, bIsBakedAdditive, bDisableRetargeting);
-	}
-}
-
-void FDecompressionTools::RetargetBoneTransform(const FReferencePose& ReferencePose
-	, const USkeleton* SourceSkeleton
-	, const USkeleton* TargetSkeleton
-	, const FName& SourceName
-	, const TArray<FTransform>& RetargetTransforms
-	, FTransform& BoneTransform
-	, const int32 SkeletonBoneIndex
-	, const int32 LODBoneIndex
-	, const bool bIsBakedAdditive
-	, const bool bDisableRetargeting)
-{
-	check(!RetargetTransforms.IsEmpty());
-	if (SourceSkeleton)
-	{
-		//const USkeleton* TargetSkeleton = RequiredBones.GetSkeletonAsset();
-		const FSkeletonRemapping& SkeletonRemapping = UE::Anim::FSkeletonRemappingRegistry::Get().GetRemapping(SourceSkeleton, TargetSkeleton);
-
-		// TODO this assumes ReferencePose to have single LOD entry, or target LOD index being 0 
-		const int32 TargetSkeletonBoneIndex = ReferencePose.GetSkeletonBoneIndexFromLODBoneIndex(LODBoneIndex);
-		const int32 SourceSkeletonBoneIndex = SkeletonRemapping.IsValid() ? SkeletonRemapping.GetSourceSkeletonBoneIndex(TargetSkeletonBoneIndex) : SkeletonBoneIndex;
-
-		const bool UseSourceRetargetModes = TargetSkeleton->GetUseRetargetModesFromCompatibleSkeleton();
-		const EBoneTranslationRetargetingMode::Type RetargetMode = UseSourceRetargetModes 
-			? SourceSkeleton->GetBoneTranslationRetargetingMode(SourceSkeletonBoneIndex, bDisableRetargeting) 
-			: TargetSkeleton->GetBoneTranslationRetargetingMode(TargetSkeletonBoneIndex, bDisableRetargeting);
-
-		switch (RetargetMode)
-		{
-		case EBoneTranslationRetargetingMode::AnimationScaled:
-		{
-			// @todo - precache that in FBoneContainer when we have SkeletonIndex->TrackIndex mapping. So we can just apply scale right away.
-			if (RetargetTransforms.IsValidIndex(SourceSkeletonBoneIndex))
-			{
-				const TArray<FTransform>& SkeletonRefPoseArray = RetargetTransforms;
-				const float SourceTranslationLength = SkeletonRefPoseArray[SourceSkeletonBoneIndex].GetTranslation().Size();
-				if (SourceTranslationLength > UE_KINDA_SMALL_NUMBER)
-				{
-					const float TargetTranslationLength = ReferencePose.GetRefPoseTranslation(LODBoneIndex).Size();
-					BoneTransform.ScaleTranslation(TargetTranslationLength / SourceTranslationLength);
-				}
-			}
-			break;
-		}
-
-		case EBoneTranslationRetargetingMode::Skeleton:
-		{
-			BoneTransform.SetTranslation(bIsBakedAdditive ? FVector::ZeroVector : ReferencePose.GetRefPoseTranslation(LODBoneIndex));
-			break;
-		}
-
-		case EBoneTranslationRetargetingMode::AnimationRelative:
-		{
-			// With baked additive animations, Animation Relative delta gets canceled out, so we can skip it.
-			// (A1 + Rel) - (A2 + Rel) = A1 - A2.
-			if (!bIsBakedAdditive)
-			{
-				const TArray<FTransform>& AuthoredOnRefSkeleton = RetargetTransforms;
-				const FTransform& RefPoseTransform = ReferencePose.GetRefPoseTransform(LODBoneIndex);
-
-				// Remap the base pose onto the target skeleton so that we are working entirely in target space
-				if (AuthoredOnRefSkeleton.IsValidIndex(SourceSkeletonBoneIndex))
-				{
-					FTransform BaseTransform = AuthoredOnRefSkeleton[SourceSkeletonBoneIndex];
-					if (SkeletonRemapping.RequiresReferencePoseRetarget())
-					{
-						BaseTransform = SkeletonRemapping.RetargetBoneTransformToTargetSkeleton(TargetSkeletonBoneIndex, BaseTransform);
-					}
-
-					// Apply the retargeting as if it were an additive difference between the current skeleton and the retarget skeleton. 
-					BoneTransform.SetRotation(BoneTransform.GetRotation() * BaseTransform.GetRotation().Inverse() * RefPoseTransform.GetRotation());
-					BoneTransform.SetTranslation(BoneTransform.GetTranslation() + (RefPoseTransform.GetTranslation() - BaseTransform.GetTranslation()));
-					BoneTransform.SetScale3D(BoneTransform.GetScale3D() * (RefPoseTransform.GetScale3D() * BaseTransform.GetSafeScaleReciprocal(BaseTransform.GetScale3D())));
-					BoneTransform.NormalizeRotation();
-				}
-			}
-			break;
-		}
-
-		case EBoneTranslationRetargetingMode::OrientAndScale:
-		{
-			if (!bIsBakedAdditive)
-			{
-				//const FRetargetSourceCachedData& RetargetSourceCachedData = RequiredBones.GetRetargetSourceCachedData(SourceName, SkeletonRemapping, RetargetTransforms);
-				//const TArray<FOrientAndScaleRetargetingCachedData>& OrientAndScaleDataArray = RetargetSourceCachedData.OrientAndScaleData;
-				//const TArray<int32>& CompactPoseIndexToOrientAndScaleIndex = RetargetSourceCachedData.CompactPoseIndexToOrientAndScaleIndex;
-
-				//// If we have any cached retargeting data.
-				//if ((OrientAndScaleDataArray.Num() > 0) && (CompactPoseIndexToOrientAndScaleIndex.Num() == RequiredBones.GetCompactPoseNumBones()))
-				//{
-				//	const int32 OrientAndScaleIndex = CompactPoseIndexToOrientAndScaleIndex[BoneIndex.GetInt()];
-				//	if (OrientAndScaleIndex != INDEX_NONE)
-				//	{
-				//		const FOrientAndScaleRetargetingCachedData& OrientAndScaleData = OrientAndScaleDataArray[OrientAndScaleIndex];
-				//		const FVector AnimatedTranslation = BoneTransform.GetTranslation();
-
-				//		// If Translation is not animated, we can just copy the TargetTranslation. No retargeting needs to be done.
-				//		const FVector NewTranslation = (AnimatedTranslation - OrientAndScaleData.SourceTranslation).IsNearlyZero(BONE_TRANS_RT_ORIENT_AND_SCALE_PRECISION) ?
-				//			OrientAndScaleData.TargetTranslation :
-				//			OrientAndScaleData.TranslationDeltaOrient.RotateVector(AnimatedTranslation) * OrientAndScaleData.TranslationScale;
-
-				//		BoneTransform.SetTranslation(NewTranslation);
-				//	}
-				//}
-			}
-			break;
-		}
-		}
-	}
-}
-
 } // end namespace UE::UAF
diff --git a/Engine/Plugins/Experimental/UAF/UAF/Source/UAF/Private/Graph/RemapPoseData.cpp b/Engine/Plugins/Experimental/UAF/UAF/Source/UAF/Private/Graph/RemapPoseData.cpp
deleted file mode 100644
diff --git a/Engine/Plugins/Experimental/UAF/UAF/Source/UAF/Private/Graph/RemapPoseData.h b/Engine/Plugins/Experimental/UAF/UAF/Source/UAF/Private/Graph/RemapPoseData.h
deleted file mode 100644
diff --git a/Engine/Plugins/Experimental/UAF/UAF/Source/UAF/Private/Graph/RigUnit_AnimNextRetargetPose.cpp b/Engine/Plugins/Experimental/UAF/UAF/Source/UAF/Private/Graph/RigUnit_AnimNextRetargetPose.cpp
new file mode 100644
--- /dev/null
+++ b/Engine/Plugins/Experimental/UAF/UAF/Source/UAF/Private/Graph/RigUnit_AnimNextRetargetPose.cpp
@@ -0,0 +1,99 @@
+// Copyright Epic Games, Inc. All Rights Reserved.
+
+#include "RigUnit_AnimNextRetargetPose.h"
+
+#include "GenerationTools.h"
+#include "Components/SkeletalMeshComponent.h"
+#include "Engine/SkeletalMesh.h"
+#include "AnimationRuntime.h"
+#include "RetargetingTools.h"
+#include "Animation/SkeletonRemappingRegistry.h"
+#include "Animation/Skeleton.h"
+#include "Animation/SkeletonRemapping.h"
+
+#include UE_INLINE_GENERATED_CPP_BY_NAME(RigUnit_AnimNextRetargetPose)
+DEFINE_STAT(STAT_AnimNext_RigUnit_RetargetPose);
+
+FRigUnit_AnimNextRetargetPose_Execute()
+{
+	SCOPE_CYCLE_COUNTER(STAT_AnimNext_RigUnit_RetargetPose);
+
+	using namespace UE::UAF;
+
+	if(!TargetAnimGraphRefPose.ReferencePose.IsValid() || !Pose.LODPose.IsValid())
+	{
+		return;
+	}
+
+	FMemMark MemMark(FMemStack::Get());
+
+	const UE::UAF::FLODPoseHeap& SourcePose = Pose.LODPose;
+	const UE::UAF::FReferencePose& SourceRefPose = SourcePose.GetRefPose();
+	const USkeletalMesh* SourceMesh = SourceRefPose.SkeletalMesh.Get();
+
+	const UE::UAF::FReferencePose& TargetRefPose = TargetAnimGraphRefPose.ReferencePose.GetRef<UE::UAF::FReferencePose>();
+	UE::UAF::FLODPoseHeap& TargetPose = Result.LODPose;
+	const USkeletalMesh* TargetMesh = TargetRefPose.SkeletalMesh.Get();
+	FRemapPoseData& RemapPoseData = RetargetPoseData.RemapPoseData;
+
+	if (SourceMesh == TargetMesh)
+	{
+		TargetPose.PrepareForLOD(TargetRefPose, SourcePose.LODLevel, /*bSetRefPose=*/false, SourcePose.IsAdditive());
+		// Just copy the pose to the target in case there is nothing to convert.
+		TargetPose.CopyFrom(SourcePose);
+
+		// No need to retarget, early out.
+		return;
+	}
+
+	const USkeleton* SourceSkeleton = SourceMesh->GetSkeleton();
+	const USkeleton* TargetSkeleton = TargetMesh->GetSkeleton();
+	
+	const FSkeletonRemapping& SkeletonRemapping = UE::Anim::FSkeletonRemappingRegistry::Get().GetRemapping(SourceSkeleton, TargetSkeleton);
+	const bool bUseSourceRetargetModes = TargetSkeleton->GetUseRetargetModesFromCompatibleSkeleton();
+	const bool bDisableRetargeting = TargetPose.GetDisableRetargeting();
+	
+	// 1) Initialize or refresh cached source target bone mappings to avoid runtime lookups.
+	if (RemapPoseData.ShouldReinit(SourceRefPose, TargetRefPose))
+	{
+		RemapPoseData.Reinit(SourceRefPose, TargetRefPose);
+	}
+
+	// 2) Copy curves and remap bone transforms and attributes to the target pose.
+	Result.Curves.CopyFrom(Pose.Curves);
+	RemapPoseData.RemapPose(SourcePose, TargetPose);
+	RemapPoseData.RemapAttributes(Pose.LODPose, Pose.Attributes, Result.LODPose, Result.Attributes);
+
+	// 3) Build or update per-bone retargeting pairs for current source/target LODs.
+	if (RetargetPoseData.PairsInitForSourcePoseLOD != SourcePose.LODLevel ||
+		RetargetPoseData.PairsInitForTargetPoseLOD != TargetPose.LODLevel)
+	{
+		FRetargetingTools::BuildRetargetingPairs(
+			SourcePose, TargetPose,
+			RemapPoseData,
+			SourceSkeleton, TargetSkeleton,
+			bUseSourceRetargetModes,
+			bDisableRetargeting,
+			RetargetPoseData.AnimScalePairs, RetargetPoseData.AnimRelativePairs, RetargetPoseData.OrientAndScalePairs);
+
+		RetargetPoseData.PairsInitForSourcePoseLOD = SourcePose.LODLevel;
+		RetargetPoseData.PairsInitForTargetPoseLOD = TargetPose.LODLevel;
+	}
+
+	// 4) Retarget reference pose to align target skeleton base transforms.
+	FRetargetingTools::ReferencePoseRetarget(
+		SourceSkeleton,
+		TargetSkeleton,
+		SkeletonRemapping,
+		bUseSourceRetargetModes,
+		bDisableRetargeting,
+		TargetPose);
+
+	// 5) Apply per-bone translation retargeting using retargeting pairs.
+	FRetargetingTools::RetargetPose(
+		TargetRefPose,
+		SkeletonRemapping,
+		RetargetPoseData.AnimScalePairs, RetargetPoseData.AnimRelativePairs, RetargetPoseData.OrientAndScalePairs,
+		SourceSkeleton->GetRefLocalPoses(),
+		TargetPose);
+}
\ No newline at end of file
diff --git a/Engine/Plugins/Experimental/UAF/UAF/Source/UAF/Private/Graph/RigUnit_AnimNextRetargetPose.h b/Engine/Plugins/Experimental/UAF/UAF/Source/UAF/Private/Graph/RigUnit_AnimNextRetargetPose.h
new file mode 100644
--- /dev/null
+++ b/Engine/Plugins/Experimental/UAF/UAF/Source/UAF/Private/Graph/RigUnit_AnimNextRetargetPose.h
@@ -0,0 +1,57 @@
+// Copyright Epic Games, Inc. All Rights Reserved.
+
+#pragma once
+
+#include "AnimNextExecuteContext.h"
+#include "Graph/AnimNext_LODPose.h"
+#include "Graph/RigUnit_AnimNextBase.h"
+#include "RemapPoseData.h"
+#include "AnimEncoding.h"
+#include "AnimNextStats.h"
+
+#include "RigUnit_AnimNextRetargetPose.generated.h"
+
+DECLARE_CYCLE_STAT_EXTERN(TEXT("UAF RigUnit: Retarget Pose"), STAT_AnimNext_RigUnit_RetargetPose, STATGROUP_AnimNext, UAF_API);
+
+USTRUCT(BlueprintType)
+struct FRetargetPoseData
+{
+	GENERATED_BODY()
+
+	FRemapPoseData RemapPoseData; // TODO: Store this per remap node or inside the data registry? We don't have ref counting ref poses and don't know when to remove these though.
+
+	BoneTrackArray AnimScalePairs;
+	BoneTrackArray AnimRelativePairs;
+	BoneTrackArray OrientAndScalePairs;
+
+	// Pairs initialized for the given source and target pose LODs.
+	// These are used to determine if we need to re-initialize the pairs after a LOD switch in either of the two poses.
+	int32 PairsInitForSourcePoseLOD = -1;
+	int32 PairsInitForTargetPoseLOD = -1;
+};
+
+/** Translational retargeting node. */
+USTRUCT(meta=(DisplayName="Retarget Pose", Category="Animation Graph", NodeColor="1, 0, 1", Keywords="Output,Pose,Port"))
+struct FRigUnit_AnimNextRetargetPose : public FRigUnit_AnimNextBase
+{
+	GENERATED_BODY()
+
+	RIGVM_METHOD()
+	void Execute();
+	
+	UPROPERTY(EditAnywhere, DisplayName = "Source Pose", Category = "Graph", meta = (Input))
+	FAnimNextGraphLODPose Pose;
+
+	UPROPERTY(EditAnywhere, DisplayName = "Retargeted Pose", Category = "Graph", meta = (Output))
+	FAnimNextGraphLODPose Result;
+
+	UPROPERTY(EditAnywhere, DisplayName = "Target Ref Pose", Category = "Graph", meta = (Input))
+	FAnimNextGraphReferencePose TargetAnimGraphRefPose;
+
+	// The execution result
+	UPROPERTY(EditAnywhere, DisplayName = "Execute", Category = "BeginExecution", meta = (Input, Output))
+	FAnimNextExecuteContext ExecuteContext;
+
+	UPROPERTY(meta = (Hidden))
+	FRetargetPoseData RetargetPoseData;
+};
diff --git a/Engine/Plugins/Experimental/UAF/UAF/Source/UAF/Private/RemapPoseData.cpp b/Engine/Plugins/Experimental/UAF/UAF/Source/UAF/Private/RemapPoseData.cpp
new file mode 100644
--- /dev/null
+++ b/Engine/Plugins/Experimental/UAF/UAF/Source/UAF/Private/RemapPoseData.cpp
@@ -0,0 +1,259 @@
+﻿// Copyright Epic Games, Inc. All Rights Reserved.
+
+#include "RemapPoseData.h"
+#include "AnimNextStats.h"
+#include "Animation/Skeleton.h"
+#include "Engine/SkeletalMesh.h"
+
+#include UE_INLINE_GENERATED_CPP_BY_NAME(RemapPoseData)
+
+DEFINE_STAT(STAT_AnimNext_RemapPose_Mesh2Mesh);
+
+bool FRemapPoseData::ShouldReinit(const UE::UAF::FReferencePose& InSourceRefPose,
+	const UE::UAF::FReferencePose& InTargetRefPose) const
+{
+	if (!SourceRefPose || !TargetRefPose)
+	{
+		// First time use.
+		return true;
+	}
+
+	if (InSourceRefPose.SkeletalMesh != SourceRefPose->SkeletalMesh)
+	{
+		// Source mesh changed.
+		return true;
+	}
+
+	if (InTargetRefPose.SkeletalMesh != TargetRefPose->SkeletalMesh)
+	{
+		// Target mesh changed.
+		return true;
+	}
+
+	return false;
+}
+
+void FRemapPoseData::Reinit(const UE::UAF::FReferencePose& InSourceRefPose,
+	const UE::UAF::FReferencePose& InTargetRefPose)
+{
+	SourceRefPose = &InSourceRefPose;
+	TargetRefPose = &InTargetRefPose;
+
+	const USkeletalMesh* SourceSkeletalMesh = InSourceRefPose.SkeletalMesh.Get();
+	const USkeleton* SourceSkeleton = SourceSkeletalMesh->GetSkeleton();
+	const FReferenceSkeleton& SourceRefSkeleton = SourceSkeleton->GetReferenceSkeleton();
+	const int32 SourceNumLODLevels = SourceSkeletalMesh->GetLODNum();
+
+	const USkeletalMesh* TargetSkeletalMesh = InTargetRefPose.SkeletalMesh.Get();
+	const FReferenceSkeleton& TargetRefSkeleton = TargetSkeletalMesh->GetRefSkeleton();
+	const int32 TargetNumLODLevels = TargetSkeletalMesh->GetLODNum();
+
+	// First index is the source LOD
+	SourceToTargetBoneIndexMapPerLOD.Empty();
+	SourceToTargetBoneIndexMapPerLOD.AddDefaulted(SourceNumLODLevels);
+
+	for (int32 SourceLODLevel = 0; SourceLODLevel < SourceNumLODLevels; ++SourceLODLevel)
+	{
+		const TArrayView<const FBoneIndexType> SourceSkeletonIndexToPoseIndex = InSourceRefPose.GetSkeletonBoneIndexToLODBoneIndexMap(); // skeleton -> pose
+		const int32 SourceNumBonesInLOD = SourceRefPose->GetNumBonesForLOD(SourceLODLevel);
+
+		// Second index is the target LOD
+		SourceToTargetBoneIndexMapPerLOD[SourceLODLevel].AddDefaulted(TargetNumLODLevels);
+
+		for (int32 TargetLODLevel = 0; TargetLODLevel < TargetNumLODLevels; ++TargetLODLevel)
+		{
+			const TArrayView<const FBoneIndexType> TargetPoseToMeshBoneIndexMap = InTargetRefPose.GetLODBoneIndexToMeshBoneIndexMap(TargetLODLevel); // pose -> mesh
+			FBoneRemapping& SourceToTargetBoneIndexMap = SourceToTargetBoneIndexMapPerLOD[SourceLODLevel][TargetLODLevel];
+
+			// Iterate over the bones present in the given LOD on the target skeleton.
+			// These bones are candidates for our mapping table, only the bones from the target skeleton that are available in the source skeleton are of interest.
+			for (FBoneIndexType BoneIndex = 0; BoneIndex < TargetPoseToMeshBoneIndexMap.Num(); ++BoneIndex)
+			{
+				const FName BoneName = TargetRefSkeleton.GetBoneName(TargetPoseToMeshBoneIndexMap[BoneIndex]);
+				const int32 SourceSkeletonBoneIndex = SourceRefSkeleton.FindBoneIndex(BoneName);
+
+				// Is the current bone on the target skeleton part of the source skeleton as well? Only add bones that are present on the target.
+				if (SourceSkeletonBoneIndex != INDEX_NONE)
+				{
+					const FBoneIndexType RemappedIndex = SourceSkeletonIndexToPoseIndex[SourceSkeletonBoneIndex];
+
+					// There are cases where the given bone are present in the skeleton, but not part of the actual pose/skeletal mesh,
+					// like e.g. when they are disabled for the given LOD on the target while they are not on the source. Skip these bones.
+					if (RemappedIndex != INDEX_NONE && RemappedIndex < SourceNumBonesInLOD)
+					{
+						FRemappedBone RemappedBone;
+						RemappedBone.SourceBoneIndex = RemappedIndex;
+						RemappedBone.TargetBoneIndex = BoneIndex;
+
+						SourceToTargetBoneIndexMap.BoneIndexMap.Add(RemappedBone);
+					}
+				}
+			}
+
+			for (int32 MappingElementIndex = 0; MappingElementIndex < SourceToTargetBoneIndexMap.BoneIndexMap.Num(); ++MappingElementIndex)
+			{
+				const int32 SourceBoneIndex = SourceToTargetBoneIndexMap.BoneIndexMap[MappingElementIndex].SourceBoneIndex;
+				const int32 TargetBoneIndex = SourceToTargetBoneIndexMap.BoneIndexMap[MappingElementIndex].TargetBoneIndex;
+
+				if (TargetBoneIndex == UE::UAF::FReferencePose::RootBoneIndex)
+				{
+					SourceToTargetBoneIndexMap.TargetRootToSourceBoneIndex = SourceBoneIndex;
+					break;
+				}
+			}
+		}
+	}
+}
+
+FTransform FRemapPoseData::RecursiveCalcModelspaceTransform(const UE::UAF::FLODPoseHeap& Pose, FBoneIndexType BoneIndex) const
+{
+	if (BoneIndex == (uint16)INDEX_NONE)
+	{
+		return FTransform::Identity;
+	}
+
+	const FTransform LocalTransform = Pose.LocalTransforms[BoneIndex];
+
+	const TArrayView<const FBoneIndexType> BoneToParentBoneIndexMap = Pose.GetLODBoneIndexToParentLODBoneIndexMap();
+	FBoneIndexType ParentBoneIndex = BoneToParentBoneIndexMap[BoneIndex];
+	if (ParentBoneIndex == (uint16)INDEX_NONE)
+	{
+		return LocalTransform;
+	}
+
+	const FTransform ParentModelTransform = RecursiveCalcModelspaceTransform(Pose, ParentBoneIndex);
+	return LocalTransform * ParentModelTransform;
+}
+
+// TODO: Source pose dictates the LOD level of the target at the moment. What if the LOD levels of source and target are not in sync? We probably need to use the target skel mesh component LOD level when preparing the target pose.
+void FRemapPoseData::RemapPose(const UE::UAF::FLODPoseHeap& SourcePose,
+	UE::UAF::FLODPoseHeap& OutTargetPose) const
+{
+	SCOPE_CYCLE_COUNTER(STAT_AnimNext_RemapPose_Mesh2Mesh);
+
+	const int32 SourceLODLevel = SourcePose.LODLevel;
+	const int32 TargetLODLevel = TargetRefPose->GetSourceLODLevel();
+
+	const bool bIsAdditive = SourcePose.IsAdditive();
+
+	// Always prepare for LOD as we may not be copying in all of the target bones and it is just cheaper to copy them all in
+	OutTargetPose.PrepareForLOD(*TargetRefPose, TargetLODLevel, /*bSetRefPose=*/true, bIsAdditive);
+
+	OutTargetPose.Flags = SourcePose.Flags;
+
+	const TArray<FRemappedBone>& SourceToTargetBoneIndexMap = SourceToTargetBoneIndexMapPerLOD[SourceLODLevel][TargetLODLevel].BoneIndexMap;
+	const int32 NumToRemap = SourceToTargetBoneIndexMap.Num();
+
+#if DEFAULT_SOA && DEFAULT_SOA_VIEW
+	const double* RESTRICT SrcRotationPtr = reinterpret_cast<const double*>(SourcePose.LocalTransformsView.Rotations.GetData());
+	double* RESTRICT DstRotationPtr = reinterpret_cast<double*>(OutTargetPose.LocalTransformsView.Rotations.GetData());
+
+	// Our SoA buffer is contiguous
+	// Because the translations/scales have the same size (FVector), each entry is a fixed offset apart and we can use a single ptr/offset pair
+	const double* RESTRICT SrcTranslationPtr = reinterpret_cast<const double*>(SourcePose.LocalTransformsView.Translations.GetData());
+	double* RESTRICT DstTranslationPtr = reinterpret_cast<double*>(OutTargetPose.LocalTransformsView.Translations.GetData());
+
+	const int64 SrcScaleOffset = reinterpret_cast<const double*>(SourcePose.LocalTransformsView.Scales3D.GetData()) - SrcTranslationPtr;
+	const int64 DstScaleOffset = reinterpret_cast<double*>(OutTargetPose.LocalTransformsView.Scales3D.GetData()) - DstTranslationPtr;
+
+	const FRemappedBone* RESTRICT RemappedBonePtr = SourceToTargetBoneIndexMap.GetData();
+	const FRemappedBone* RESTRICT RemappedBoneEndPtr = RemappedBonePtr + NumToRemap;
+
+	while (RemappedBonePtr < RemappedBoneEndPtr)
+	{
+		const int32 SourceBoneIndex = RemappedBonePtr->SourceBoneIndex;
+		const int32 TargetBoneIndex = RemappedBonePtr->TargetBoneIndex;
+
+		SourcePose.LocalTransformsView.Rotations.RangeCheck(SourceBoneIndex);
+		OutTargetPose.LocalTransformsView.Rotations.RangeCheck(TargetBoneIndex);
+
+		VectorRegister4Double Rotation = VectorLoadAligned(SrcRotationPtr + SourceBoneIndex * 4);
+		VectorRegister4Double Translation = VectorLoad(SrcTranslationPtr + SourceBoneIndex * 3);
+		VectorRegister4Double Scale = VectorLoad(SrcTranslationPtr + SourceBoneIndex * 3 + SrcScaleOffset);
+
+		VectorStoreAligned(Rotation, DstRotationPtr + TargetBoneIndex * 4);
+		VectorStoreFloat3(Translation, DstTranslationPtr + TargetBoneIndex * 3);
+		VectorStoreFloat3(Scale, DstTranslationPtr + TargetBoneIndex * 3 + DstScaleOffset);
+
+		RemappedBonePtr++;
+	}
+#else
+	for (int32 MappingElementIndex = 0; MappingElementIndex < NumToRemap; ++MappingElementIndex)
+	{
+		const int32 SourceBoneIndex = SourceToTargetBoneIndexMap[MappingElementIndex].SourceBoneIndex;
+		const int32 TargetBoneIndex = SourceToTargetBoneIndexMap[MappingElementIndex].TargetBoneIndex;
+
+		const auto& TargetTransform = SourcePose.LocalTransformsView[SourceBoneIndex];
+		OutTargetPose.LocalTransforms[TargetBoneIndex] = TargetTransform;
+	}
+#endif
+
+	// For the cases where the target skeleton does not share the same root bone with the source, try to find the corresponding
+	// bone on the source skeleton for the target's root bone and sync the skeletons up from there by calculating the delta transform
+	// between the two and move it along with the source. This allows us to remap and attach skeletal meshes that only contain the
+	// bones it needs, like e.g. only skinned bones.
+	const int32 RootBoneOnSourceMapped = SourceToTargetBoneIndexMapPerLOD[SourceLODLevel][TargetLODLevel].TargetRootToSourceBoneIndex;
+	if (RootBoneOnSourceMapped != INDEX_NONE)
+	{
+		// TODO: Replace with the generic method once we have access to the modelspace transform via the LOD pose.
+		const FTransform RootModelTransform = RecursiveCalcModelspaceTransform(SourcePose, RootBoneOnSourceMapped);
+		OutTargetPose.LocalTransforms[UE::UAF::FReferencePose::RootBoneIndex] = RootModelTransform;
+	}
+}
+
+void FRemapPoseData::RemapAttributes(const UE::UAF::FLODPose& SourceLODPose,
+	const UE::Anim::FHeapAttributeContainer& InAttributes,
+	const UE::UAF::FLODPose& TargetLODPose,
+	UE::Anim::FHeapAttributeContainer& OutAttributes)
+{
+	const USkeletalMesh* SourceSkeletalMesh = SourceRefPose->SkeletalMesh.Get();
+	if (!SourceSkeletalMesh)
+	{
+		ensureMsgf(false, TEXT("FRemapPoseData::RemapAttributes(): Source skeletal mesh is not valid anymore."));
+		return;
+	}
+	const FReferenceSkeleton& SourceRefSkeleton = SourceSkeletalMesh->GetRefSkeleton();
+	const TArrayView<const FBoneIndexType> SourceLODBoneIndexToSkeletonBoneIndexMap = SourceLODPose.GetLODBoneIndexToSkeletonBoneIndexMap();
+
+	const USkeletalMesh* TargetSkeletalMesh = TargetRefPose->SkeletalMesh.Get();
+	if (!TargetSkeletalMesh)
+	{
+		ensureMsgf(false, TEXT("FRemapPoseData::RemapAttributes(): Target skeletal mesh is not valid anymore."));
+		return;
+	}
+	const FReferenceSkeleton& TargetRefSkeleton = TargetSkeletalMesh->GetRefSkeleton();
+	const TArrayView<const FBoneIndexType> TargetSkeletonBoneToLODBoneIndexMap = TargetLODPose.GetSkeletonBoneIndexToLODBoneIndexMap();
+
+	for (const TWeakObjectPtr<UScriptStruct> WeakScriptStruct : InAttributes.GetUniqueTypes())
+	{
+		const UScriptStruct* ScriptStruct = WeakScriptStruct.Get();
+		const int32 TypeIndex = InAttributes.FindTypeIndex(ScriptStruct);
+		if (TypeIndex != INDEX_NONE)
+		{
+			const TArray<UE::Anim::TWrappedAttribute<FDefaultAllocator>, FDefaultAllocator>& SourceValues = InAttributes.GetValues(TypeIndex);
+			const TArray<UE::Anim::FAttributeId, FDefaultAllocator>& AttributeIds = InAttributes.GetKeys(TypeIndex);
+
+			// Try and remap all the source attributes to their respective new bone indices
+			for (int32 EntryIndex = 0; EntryIndex < AttributeIds.Num(); ++EntryIndex)
+			{
+				const UE::Anim::FAttributeId& AttributeId = AttributeIds[EntryIndex];
+
+				// Remap the source bone from LOD pose bone index to a skeleton bone index and get the bone name.
+				const int32 SourceLODBoneIndex = AttributeId.GetIndex();
+				const int32 SourceSkeletonBoneIndex = SourceLODBoneIndexToSkeletonBoneIndexMap[SourceLODBoneIndex];
+				const FName BoneName = SourceRefSkeleton.GetBoneName(SourceSkeletonBoneIndex);
+				
+				// Find the given bone inside the target skeleton, if it is in there, remap the attribute.
+				const int32 TargetSkeletonBoneIndex = TargetRefSkeleton.FindBoneIndex(BoneName);
+				if (TargetSkeletonBoneIndex != INDEX_NONE)
+				{
+					const int32 TargetLODBoneIndex = TargetSkeletonBoneToLODBoneIndexMap[TargetSkeletonBoneIndex];
+
+					const UE::Anim::FAttributeId NewInfo(AttributeId.GetName(), TargetLODBoneIndex, AttributeId.GetNamespace());
+					uint8* NewAttribute = OutAttributes.FindOrAdd(ScriptStruct, NewInfo);
+					ScriptStruct->CopyScriptStruct(NewAttribute, SourceValues[EntryIndex].GetPtr<void>());
+				}
+			}
+		}
+	}
+}
diff --git a/Engine/Plugins/Experimental/UAF/UAF/Source/UAF/Private/RetargetingTools.cpp b/Engine/Plugins/Experimental/UAF/UAF/Source/UAF/Private/RetargetingTools.cpp
new file mode 100644
--- /dev/null
+++ b/Engine/Plugins/Experimental/UAF/UAF/Source/UAF/Private/RetargetingTools.cpp
@@ -0,0 +1,368 @@
+// Copyright Epic Games, Inc. All Rights Reserved.
+
+#include "RetargetingTools.h"
+#include "AnimationRuntime.h"
+#include "Animation/SkeletonRemappingRegistry.h"
+#include "Animation/SkeletonRemapping.h"
+
+namespace UE::UAF
+{
+	void FRetargetingTools::BuildRetargetingPairs(
+	    const UE::UAF::FLODPose& SourcePose,
+	    const UE::UAF::FLODPose& TargetPose,
+	    const FRemapPoseData& RemapPoseData,
+	    const USkeleton* SourceSkeleton,
+	    const USkeleton* TargetSkeleton,
+	    bool bUseSourceRetargetModes,
+	    bool bDisableRetargeting,
+	    BoneTrackArray& OutAnimScalePairs,
+	    BoneTrackArray& OutAnimRelativePairs,
+	    BoneTrackArray& OutOrientAndScalePairs)
+	{
+	    OutAnimScalePairs.Reset();
+	    OutAnimRelativePairs.Reset();
+	    OutOrientAndScalePairs.Reset();
+
+		const int32 SourceLODLevel = SourcePose.LODLevel;
+		const int32 TargetLODLevel = TargetPose.LODLevel;
+
+	    const TArrayView<const FBoneIndexType> SourcePoseToSkeletonBoneIndexMap = SourcePose.GetLODBoneIndexToSkeletonBoneIndexMap();
+	    const TArrayView<const FBoneIndexType> TargetPoseToSkeletonBoneIndexMap = TargetPose.GetLODBoneIndexToSkeletonBoneIndexMap();
+		const TArray<FRemappedBone>& SourceToTargetPoseBoneIndexMap = RemapPoseData.SourceToTargetBoneIndexMapPerLOD[SourceLODLevel][TargetLODLevel].BoneIndexMap; // Source LODPose -> target LODPose
+
+		// Iterate through all bones available in the source as well as the target skeleton/pose.
+	    for (const FRemappedBone& RemappedBone : SourceToTargetPoseBoneIndexMap)
+	    {
+	        const int32 SourcePoseBoneIndex = RemappedBone.SourceBoneIndex;
+	        const int32 TargetPoseBoneIndex = RemappedBone.TargetBoneIndex;
+	        if (!SourcePoseToSkeletonBoneIndexMap.IsValidIndex(SourcePoseBoneIndex) || !TargetPoseToSkeletonBoneIndexMap.IsValidIndex(TargetPoseBoneIndex))
+	        {
+	            continue;
+	        }
+
+	        const int32 SourceSkeletonBoneIndex = SourcePoseToSkeletonBoneIndexMap[SourcePoseBoneIndex];
+	        const int32 TargetSkeletonBoneIndex = TargetPoseToSkeletonBoneIndexMap[TargetPoseBoneIndex];
+	        if (SourceSkeletonBoneIndex == INDEX_NONE || TargetSkeletonBoneIndex == INDEX_NONE)
+	        {
+		        continue;
+	        }
+
+	    	const EBoneTranslationRetargetingMode::Type Mode = FAnimationRuntime::GetBoneTranslationRetargetingMode(
+				bUseSourceRetargetModes,
+				SourceSkeletonBoneIndex,
+				TargetSkeletonBoneIndex,
+				SourceSkeleton,
+				TargetSkeleton,
+				bDisableRetargeting);
+
+	        switch (Mode)
+	        {
+				case EBoneTranslationRetargetingMode::AnimationScaled:
+	        	{
+					// (AtomIndex = Target LODPose, TrackIndex = Source Skeleton)
+	        		OutAnimScalePairs.Add(BoneTrackPair(TargetPoseBoneIndex, SourceSkeletonBoneIndex));
+	        		break;
+	        	}
+
+				case EBoneTranslationRetargetingMode::AnimationRelative:
+	        	{
+	        		if (!SourcePose.IsAdditive())
+	        		{
+	        			OutAnimRelativePairs.Add(BoneTrackPair(TargetPoseBoneIndex, SourceSkeletonBoneIndex));
+	        		}
+	        		break;
+				}
+
+	            case EBoneTranslationRetargetingMode::OrientAndScale:
+		        {
+			        if (!SourcePose.IsAdditive())
+			        {
+				        OutOrientAndScalePairs.Add(BoneTrackPair(TargetPoseBoneIndex, SourceSkeletonBoneIndex));
+			        }
+	        		break;
+		        }
+
+	            default:
+		        {
+			        break;
+		        }
+	        }
+	    }
+	}
+
+	void FRetargetingTools::ReferencePoseRetarget(
+		const USkeleton* SourceSkeleton,
+		const USkeleton* TargetSkeleton,
+		const FSkeletonRemapping& SkeletonRemapping,
+		bool bUseSourceRetargetModes,
+		bool bDisableRetargeting,
+		UE::UAF::FLODPose& OutPose)
+	{
+		if (!SkeletonRemapping.RequiresReferencePoseRetarget())
+		{
+			return;
+		}
+	    
+		const TArrayView<const FBoneIndexType> PoseToSkeletonBoneIndexMap = OutPose.GetLODBoneIndexToSkeletonBoneIndexMap();
+		const int32 NumLODBones = PoseToSkeletonBoneIndexMap.Num();
+	    
+		for (int32 LODBoneIndex = 0; LODBoneIndex < NumLODBones; ++LODBoneIndex)
+		{
+			const int32 TargetSkeletonBoneIndex = PoseToSkeletonBoneIndexMap[LODBoneIndex];
+			if (TargetSkeletonBoneIndex == INDEX_NONE)
+			{
+				continue;
+			}
+    
+			OutPose.LocalTransformsView[LODBoneIndex].SetRotation(
+				SkeletonRemapping.RetargetBoneRotationToTargetSkeleton(
+					TargetSkeletonBoneIndex,
+					OutPose.LocalTransformsView[LODBoneIndex].GetRotation()));
+    
+			TranslationRetargetReferencePoseBone<false>(
+				SourceSkeleton,
+				TargetSkeleton,
+				SkeletonRemapping,
+				bUseSourceRetargetModes,
+				bDisableRetargeting,
+				LODBoneIndex,
+				TargetSkeletonBoneIndex,
+				OutPose);
+		}
+	}
+	
+	void FRetargetingTools::RetargetBoneTransform(
+			const FReferencePose& ReferencePose,
+			const USkeleton* SourceSkeleton,
+			const USkeleton* TargetSkeleton,
+			const FName& SourceName,
+			const TArray<FTransform>& RetargetTransforms,
+			FTransform& BoneTransform,
+			const int32 SkeletonBoneIndex,
+			const int32 LODBoneIndex,
+			const bool bIsBakedAdditive,
+			const bool bDisableRetargeting)
+	{
+		check(!RetargetTransforms.IsEmpty());
+		if (!SourceSkeleton)
+		{
+			return;
+		}
+
+		const FSkeletonRemapping& SkeletonRemapping = UE::Anim::FSkeletonRemappingRegistry::Get().GetRemapping(SourceSkeleton, TargetSkeleton);
+		const int32 TargetSkeletonBoneIndex = ReferencePose.GetSkeletonBoneIndexFromLODBoneIndex(LODBoneIndex);
+		const int32 SourceSkeletonBoneIndex = SkeletonRemapping.IsValid() ? SkeletonRemapping.GetSourceSkeletonBoneIndex(TargetSkeletonBoneIndex) : SkeletonBoneIndex;
+
+		const bool UseSourceRetargetModes = TargetSkeleton->GetUseRetargetModesFromCompatibleSkeleton();
+		const EBoneTranslationRetargetingMode::Type RetargetMode = UseSourceRetargetModes 
+			? SourceSkeleton->GetBoneTranslationRetargetingMode(SourceSkeletonBoneIndex, bDisableRetargeting) 
+			: TargetSkeleton->GetBoneTranslationRetargetingMode(TargetSkeletonBoneIndex, bDisableRetargeting);
+
+		switch (RetargetMode)
+		{
+		case EBoneTranslationRetargetingMode::AnimationScaled:
+			{
+				// @todo - precache that in FBoneContainer when we have SkeletonIndex->TrackIndex mapping. So we can just apply scale right away.
+				if (RetargetTransforms.IsValidIndex(SourceSkeletonBoneIndex))
+				{
+					const TArray<FTransform>& SkeletonRefPoseArray = RetargetTransforms;
+					const float SourceTranslationLength = SkeletonRefPoseArray[SourceSkeletonBoneIndex].GetTranslation().Size();
+					if (SourceTranslationLength > UE_KINDA_SMALL_NUMBER)
+					{
+						const float TargetTranslationLength = ReferencePose.GetRefPoseTranslation(LODBoneIndex).Size();
+						BoneTransform.ScaleTranslation(TargetTranslationLength / SourceTranslationLength);
+					}
+				}
+				break;
+			}
+
+		case EBoneTranslationRetargetingMode::Skeleton:
+			{
+				BoneTransform.SetTranslation(bIsBakedAdditive ? FVector::ZeroVector : ReferencePose.GetRefPoseTranslation(LODBoneIndex));
+				break;
+			}
+
+		case EBoneTranslationRetargetingMode::AnimationRelative:
+			{
+				// With baked additive animations, Animation Relative delta gets canceled out, so we can skip it.
+				// (A1 + Rel) - (A2 + Rel) = A1 - A2.
+				if (!bIsBakedAdditive)
+				{
+					const TArray<FTransform>& AuthoredOnRefSkeleton = RetargetTransforms;
+
+					// Remap the base pose onto the target skeleton so that we are working entirely in target space
+					if (AuthoredOnRefSkeleton.IsValidIndex(SourceSkeletonBoneIndex))
+					{
+						FTransform BaseTransform = AuthoredOnRefSkeleton[SourceSkeletonBoneIndex];
+						RetargetBone_AnimRelative(LODBoneIndex,
+							TargetSkeletonBoneIndex,
+							ReferencePose,
+							SkeletonRemapping,
+							BaseTransform,
+							BoneTransform);
+					}
+				}
+				break;
+			}
+
+		case EBoneTranslationRetargetingMode::OrientAndScale:
+			{
+				if (!bIsBakedAdditive)
+				{
+					//const FRetargetSourceCachedData& RetargetSourceCachedData = RequiredBones.GetRetargetSourceCachedData(SourceName, SkeletonRemapping, RetargetTransforms);
+					//const TArray<FOrientAndScaleRetargetingCachedData>& OrientAndScaleDataArray = RetargetSourceCachedData.OrientAndScaleData;
+					//const TArray<int32>& CompactPoseIndexToOrientAndScaleIndex = RetargetSourceCachedData.CompactPoseIndexToOrientAndScaleIndex;
+
+					//// If we have any cached retargeting data.
+					//if ((OrientAndScaleDataArray.Num() > 0) && (CompactPoseIndexToOrientAndScaleIndex.Num() == RequiredBones.GetCompactPoseNumBones()))
+					//{
+					//	const int32 OrientAndScaleIndex = CompactPoseIndexToOrientAndScaleIndex[BoneIndex.GetInt()];
+					//	if (OrientAndScaleIndex != INDEX_NONE)
+					//	{
+					//		const FOrientAndScaleRetargetingCachedData& OrientAndScaleData = OrientAndScaleDataArray[OrientAndScaleIndex];
+					//		const FVector AnimatedTranslation = BoneTransform.GetTranslation();
+
+					//		// If Translation is not animated, we can just copy the TargetTranslation. No retargeting needs to be done.
+					//		const FVector NewTranslation = (AnimatedTranslation - OrientAndScaleData.SourceTranslation).IsNearlyZero(BONE_TRANS_RT_ORIENT_AND_SCALE_PRECISION) ?
+					//			OrientAndScaleData.TargetTranslation :
+					//			OrientAndScaleData.TranslationDeltaOrient.RotateVector(AnimatedTranslation) * OrientAndScaleData.TranslationScale;
+
+					//		BoneTransform.SetTranslation(NewTranslation);
+					//	}
+					//}
+				}
+				break;
+			}
+		}
+	}
+
+	void FRetargetingTools::RetargetBone_AnimRelative(
+		const int32 TargetPoseBoneIndex,
+		const int32 TargetSkeletonBoneIndex,
+		const FReferencePose& TargetReferencePose,
+		const FSkeletonRemapping& SourceToTargetRemapping,
+		FTransform BaseTransform,
+		FTransform& InOutBoneTransform)
+	{
+		const FTransform& RefPoseTransform = TargetReferencePose.GetRefPoseTransform(TargetPoseBoneIndex);
+
+		// Remap the base pose onto the target skeleton so that we are working entirely in target space
+		if (SourceToTargetRemapping.RequiresReferencePoseRetarget())
+		{
+			if (TargetSkeletonBoneIndex != INDEX_NONE)
+			{
+				BaseTransform = SourceToTargetRemapping.RetargetBoneTransformToTargetSkeleton(TargetSkeletonBoneIndex, BaseTransform);
+			}
+		}
+
+		// Apply the retargeting as if it were an additive difference between the current skeleton and the retarget skeleton. 
+		InOutBoneTransform.SetRotation(InOutBoneTransform.GetRotation() * BaseTransform.GetRotation().Inverse() * RefPoseTransform.GetRotation());
+		InOutBoneTransform.SetTranslation(InOutBoneTransform.GetTranslation() + (RefPoseTransform.GetTranslation() - BaseTransform.GetTranslation()));
+		InOutBoneTransform.SetScale3D(InOutBoneTransform.GetScale3D() * (RefPoseTransform.GetScale3D() * BaseTransform.GetSafeScaleReciprocal(BaseTransform.GetScale3D())));
+		InOutBoneTransform.NormalizeRotation();
+	}
+
+	void FRetargetingTools::RetargetPose(
+		const FReferencePose& TargetReferencePose,
+		const FSkeletonRemapping& SourceToTargetRemapping,
+		BoneTrackArray& AnimScaleRetargetingPairs,
+		BoneTrackArray& AnimRelativeRetargetingPairs,
+		BoneTrackArray& OrientAndScaleRetargetingPairs,
+		const TArray<FTransform>& SourceReferencePoseTransforms,
+		FLODPose& OutPose)
+	{
+		// Anim Scale Retargeting
+		int32 const NumBonesToScaleRetarget = AnimScaleRetargetingPairs.Num();
+		if (NumBonesToScaleRetarget > 0)
+		{
+			TArray<FTransform> const& AuthoredOnRefSkeleton = SourceReferencePoseTransforms;
+
+			for (const BoneTrackPair& BonePair : AnimScaleRetargetingPairs)
+			{
+				const int32 LODBoneIndex(BonePair.AtomIndex);
+				const int32 SourceSkeletonBoneIndex = BonePair.TrackIndex;
+				if (!AuthoredOnRefSkeleton.IsValidIndex(SourceSkeletonBoneIndex))
+				{
+					continue;
+				}
+
+				// @todo - precache that in FBoneContainer when we have SkeletonIndex->TrackIndex mapping. So we can just apply scale right away.
+				float const SourceTranslationLength = AuthoredOnRefSkeleton[SourceSkeletonBoneIndex].GetTranslation().Size();
+				if (SourceTranslationLength > UE_KINDA_SMALL_NUMBER)
+				{
+					float const TargetTranslationLength = TargetReferencePose.GetRefPoseTranslation(LODBoneIndex).Size();
+					OutPose.LocalTransformsView[LODBoneIndex].ScaleTranslation(TargetTranslationLength / SourceTranslationLength);
+				}
+			}
+		}
+
+		// Anim Relative Retargeting
+		int32 const NumBonesToRelativeRetarget = AnimRelativeRetargetingPairs.Num();
+		if (NumBonesToRelativeRetarget > 0)
+		{
+			TArray<FTransform> const& AuthoredOnRefSkeleton = SourceReferencePoseTransforms;
+
+			for (const BoneTrackPair& BonePair : AnimRelativeRetargetingPairs)
+			{
+				const int32 LODBoneIndex(BonePair.AtomIndex);
+				const int32 SourceSkeletonBoneIndex = BonePair.TrackIndex;
+				if (!AuthoredOnRefSkeleton.IsValidIndex(SourceSkeletonBoneIndex))
+				{
+					continue;
+				}
+
+				const FTransform& BaseTransform = AuthoredOnRefSkeleton[SourceSkeletonBoneIndex];
+				FTransform RetargetTransform = OutPose.LocalTransformsView[LODBoneIndex];
+				const int32 TargetSkeletonBoneIndex = SourceToTargetRemapping.GetTargetSkeletonBoneIndex(SourceSkeletonBoneIndex);
+
+				RetargetBone_AnimRelative(LODBoneIndex,
+					TargetSkeletonBoneIndex,
+					TargetReferencePose,
+					SourceToTargetRemapping,
+					BaseTransform,
+					RetargetTransform);
+
+				OutPose.LocalTransformsView[LODBoneIndex] = RetargetTransform;
+			}
+		}
+
+		// TODO : Have to recreate GetRetargetSourceCachedData
+		// Translation 'Orient and Scale' Translation Retargeting
+		//const int32 NumBonesToOrientAndScaleRetarget = OrientAndScaleRetargetingPairs.Num();
+		//if (NumBonesToOrientAndScaleRetarget > 0)
+		//{
+		//	const FRetargetSourceCachedData& RetargetSourceCachedData = RequiredBones.GetRetargetSourceCachedData(DecompressionContext.AnimName, SkeletonRemapping, RetargetTransforms);
+		//	const TArray<FOrientAndScaleRetargetingCachedData>& OrientAndScaleDataArray = RetargetSourceCachedData.OrientAndScaleData;
+		//	const TArray<int32>& CompactPoseIndexToOrientAndScaleIndex = RetargetSourceCachedData.CompactPoseIndexToOrientAndScaleIndex;
+
+		//	// If we have any cached retargeting data.
+		//	if ((OrientAndScaleDataArray.Num() > 0) && (CompactPoseIndexToOrientAndScaleIndex.Num() == RequiredBones.GetCompactPoseNumBones()))
+		//	{
+		//		for (int32 Index = 0; Index < NumBonesToOrientAndScaleRetarget; Index++)
+		//		{
+		//			const BoneTrackPair& BonePair = OrientAndScaleRetargetingPairs[Index];
+		//			const FCompactPoseBoneIndex CompactPoseBoneIndex(BonePair.AtomIndex);
+		//			const int32 OrientAndScaleIndex = CompactPoseIndexToOrientAndScaleIndex[CompactPoseBoneIndex.GetInt()];
+		//			if (OrientAndScaleIndex != INDEX_NONE)
+		//			{
+		//				const FOrientAndScaleRetargetingCachedData& OrientAndScaleData = OrientAndScaleDataArray[OrientAndScaleIndex];
+		//				FTransform& BoneTransform = OutAnimationPoseData.LocalTransforms[CompactPoseBoneIndex];
+		//				const FVector AnimatedTranslation = BoneTransform.GetTranslation();
+
+		//				// If Translation is not animated, we can just copy the TargetTranslation. No retargeting needs to be done.
+		//				const FVector NewTranslation = (AnimatedTranslation - OrientAndScaleData.SourceTranslation).IsNearlyZero(BONE_TRANS_RT_ORIENT_AND_SCALE_PRECISION) ?
+		//					OrientAndScaleData.TargetTranslation :
+		//					OrientAndScaleData.TranslationDeltaOrient.RotateVector(AnimatedTranslation) * OrientAndScaleData.TranslationScale;
+
+		//				BoneTransform.SetTranslation(NewTranslation);
+		//			}
+		//		}
+		//	}
+		//}
+		
+		// UE_CLOGFMT(OrientAndScaleRetargetingPairs.Num() != 0,
+		// 	LogAnimation, Verbose,
+		// 	"OrientAndScaleRetargetingPairs found for {AnimSequence} ({SourceSkeleton} -> {TargetSkeleton}), Orient and Scale retargeting is not yet supported in UAF",
+		// 	*DecompressionContext.AnimName.ToString(), SourceSkeleton ? *SourceSkeleton->GetName() : TEXT("None"), TargetSkeleton ? *TargetSkeleton->GetName() : TEXT("None"));
+	}
+} // end namespace UE::UAF
\ No newline at end of file
diff --git a/Engine/Plugins/Experimental/UAF/UAF/Source/UAF/Public/DecompressionTools.h b/Engine/Plugins/Experimental/UAF/UAF/Source/UAF/Public/DecompressionTools.h
--- a/Engine/Plugins/Experimental/UAF/UAF/Source/UAF/Public/DecompressionTools.h
+++ b/Engine/Plugins/Experimental/UAF/UAF/Source/UAF/Public/DecompressionTools.h
@@ -63,27 +63,6 @@ public:
 		FAnimSequenceDecompressionContext& DecompressionContext,
 		FName RetargetSource,
 		const FRootMotionReset& RootMotionReset);
-
-	static UE_API void RetargetBoneTransform(const FReferencePose& ReferencePose
-		, const USkeleton* SourceSkeleton
-		, const USkeleton* TargetSkeleton
-		, const FName& RetargetSource
-		, FTransform& BoneTransform
-		, const int32 SkeletonBoneIndex
-		, const int32 LODBoneIndex
-		, const bool bIsBakedAdditive
-		, const bool bDisableRetargeting);
-
-	static UE_API void RetargetBoneTransform(const FReferencePose& ReferencePose
-		, const USkeleton* SourceSkeleton
-		, const USkeleton* TargetSkeleton
-		, const FName& SourceName
-		, const TArray<FTransform>& RetargetTransforms
-		, FTransform& BoneTransform
-		, const int32 SkeletonBoneIndex
-		, const int32 LODBoneIndex
-		, const bool bIsBakedAdditive
-		, const bool bDisableRetargeting);
 };
 
 } // end namespace UE::UAF
diff --git a/Engine/Plugins/Experimental/UAF/UAF/Source/UAF/Public/RemapPoseData.h b/Engine/Plugins/Experimental/UAF/UAF/Source/UAF/Public/RemapPoseData.h
new file mode 100644
--- /dev/null
+++ b/Engine/Plugins/Experimental/UAF/UAF/Source/UAF/Public/RemapPoseData.h
@@ -0,0 +1,86 @@
+// Copyright Epic Games, Inc. All Rights Reserved.
+
+#pragma once
+
+#include "Animation/AttributesRuntime.h"
+#include "ReferencePose.h"
+#include "LODPose.h"
+#include "BoneIndices.h"
+
+#include "RemapPoseData.generated.h"
+
+USTRUCT()
+struct FRemappedBone
+{
+	GENERATED_BODY()
+
+	FBoneIndexType SourceBoneIndex = INDEX_NONE;
+	FBoneIndexType TargetBoneIndex = INDEX_NONE;
+};
+
+USTRUCT()
+struct FBoneRemapping
+{
+	GENERATED_BODY()
+
+	/**
+	 * Source to target bone index array for bones from the target pose that are also present in the source pose are stored.
+	 * This is a subset of the bones from the target pose. The smaller the subset, the less bones to remap.
+	 */
+	TArray<FRemappedBone> BoneIndexMap;
+
+	/**
+	 * Target root bone with index 0 mapped to the source skeletal mesh. The bone index stored here represents the bone index on the source
+	 * skeletal mesh of the bone with the same name as the root bone on the target. This is used as bone attachment point.
+	 */
+	int32 TargetRootToSourceBoneIndex = INDEX_NONE;
+};
+
+USTRUCT()
+struct FRemapPoseData
+{
+	GENERATED_BODY()
+
+	/**
+	 * Check if we need to re-initialize the mapping. This will return true in case of first time usage of the data,
+	 * or in the case that either the source or target pose changed.
+	 */
+	bool ShouldReinit(const UE::UAF::FReferencePose& InSourceRefPose,
+		const UE::UAF::FReferencePose& InTargetRefPose) const;
+
+	/**
+	 * Re-initializes the mapping for each LOD. Call ShouldReinit() before to avoid unnecessary compute.
+	 * Compare the bones from the source and target pose and find the subset of the target that is also present in the source.
+	 */
+	void Reinit(const UE::UAF::FReferencePose& InSourceRefPose,
+		const UE::UAF::FReferencePose& InTargetRefPose);
+
+	/**
+	 * Remap a pose using the previously cached mapping. Call this at runtime.
+	 */
+	void RemapPose(const UE::UAF::FLODPoseHeap& SourcePose,
+		UE::UAF::FLODPoseHeap& OutTargetPose) const;
+
+	void RemapAttributes(const UE::UAF::FLODPose& SourceLODPose,
+		const UE::Anim::FHeapAttributeContainer& InAttributes,
+		const UE::UAF::FLODPose& TargetLODPose,
+		UE::Anim::FHeapAttributeContainer& OutAttributes);
+
+	/**
+	 * Temporary calculate modelspace transform, until the LODPose is able to do that.
+	 */
+	FTransform RecursiveCalcModelspaceTransform(const UE::UAF::FLODPoseHeap& Pose, FBoneIndexType BoneIndex) const;
+
+	const UE::UAF::FReferencePose* SourceRefPose = nullptr;
+	const UE::UAF::FReferencePose* TargetRefPose = nullptr;
+
+	/**
+	 * Bone mapping from the source to the target skeletal mesh per LOD level combination.
+	 * The array's length equals the number of source skeletal mesh LODs available. Each of the elements
+	 * will contain a bone mapping based on the given source LOD for the target skeletal mesh LOD. The number of elements
+	 * of the sub-array equals to the number of LOD levels of the target skeletal mesh.
+	 * Example: SourceToTargetBoneIndexMapPerLOD[4][2] gives a bone mapping from source LOD 4 to target LOD 2 and will usually contain
+	 * less elements than the target skeletal mesh has bones.
+	 */
+	TArray<TArray<FBoneRemapping>> SourceToTargetBoneIndexMapPerLOD;
+};
\ No newline at end of file
diff --git a/Engine/Plugins/Experimental/UAF/UAF/Source/UAF/Public/RetargetingTools.h b/Engine/Plugins/Experimental/UAF/UAF/Source/UAF/Public/RetargetingTools.h
new file mode 100644
--- /dev/null
+++ b/Engine/Plugins/Experimental/UAF/UAF/Source/UAF/Public/RetargetingTools.h
@@ -0,0 +1,213 @@
+// Copyright Epic Games, Inc. All Rights Reserved.
+
+#pragma once
+
+#include "CoreMinimal.h"
+#include "ReferencePose.h"
+#include "LODPose.h"
+#include "AnimEncoding.h"
+#include "AnimationRuntime.h"
+#include "Animation/Skeleton.h"
+#include "Animation/SkeletonRemapping.h"
+#include "RemapPoseData.h"
+
+namespace UE::UAF
+{
+	class FRetargetingTools
+	{
+	public:
+		/**
+		 * Build sparse bone-track pairs for each translation retargeting mode.
+		 * Pre-filters Source to target bone relations into three lists:
+		 * AnimationScaled, AnimationRelative, and OrientAndScale.
+		 * @param SourcePose                Source pose to retarget from.
+		 * @param TargetPose                Retargeting target pose.
+		 * @param RemapPoseData             Cached source to target bone mapping per LOD.
+		 * @param SourceSkeleton            Source skeleton to retarget from.
+		 * @param TargetSkeleton            Retargeting target skeleton.
+		 * @param bUseSourceRetargetModes   Set to true when you wish to extract the retarget mode from the source skeleton instead of the target skeleton.
+		 * @param bDisableRetargeting       Is translational retargeting disabled?
+		 * @param OutAnimScalePairs         Output pairs for AnimationScaled.
+		 * @param OutAnimRelativePairs      Output pairs for AnimationRelative.
+		 * @param OutOrientAndScalePairs    Output pairs for OrientAndScale.
+		 */
+		static UAF_API void BuildRetargetingPairs(
+			const FLODPose& SourcePose,
+			const FLODPose& TargetPose,
+			const FRemapPoseData& RemapPoseData,
+			const USkeleton* SourceSkeleton,
+			const USkeleton* TargetSkeleton,
+			bool bUseSourceRetargetModes,
+			bool bDisableRetargeting,
+			BoneTrackArray& OutAnimScalePairs,
+			BoneTrackArray& OutAnimRelativePairs,
+			BoneTrackArray& OutOrientAndScalePairs);
+
+		/**
+		 * Per-bone reference-pose translation correction.
+		 * This is the step that moves authored offsets across differing proportions without touching rotation.
+		 * @tparam bIsAdditive              True if the source pose is additive (deltas relative to a reference pose),
+		 *                                  in which case additive-space retargeting is applied. False for absolute poses.
+		 * @param SourceSkeleton            Source skeleton to retarget from.
+		 * @param TargetSkeleton            Retargeting target skeleton.
+		 * @param SkeletonRemapping         Maps source and target skeleton bone indices.
+		 * @param bUseSourceRetargetModes   Set to true when you wish to extract the retarget mode from the source skeleton instead of the target skeleton.
+		 * @param bDisableRetargeting       Is translational retargeting disabled?
+		 * @param PoseBoneIndex             Target pose bone index.
+		 * @param TargetSkeletonBoneIndex   Target skeleton bone index.
+		 * @param OutPose                   Target pose receiving the translation fix.
+		 */
+		template<bool bIsAdditive>
+		static FORCEINLINE void TranslationRetargetReferencePoseBone(
+			const USkeleton* SourceSkeleton,
+			const USkeleton* TargetSkeleton,
+			const FSkeletonRemapping& SkeletonRemapping,
+			bool bUseSourceRetargetModes,
+			bool bDisableRetargeting,
+			const int32 PoseBoneIndex,
+			const int32 TargetSkeletonBoneIndex,
+			FLODPose& OutPose)
+		{
+			const int32 SourceSkeletonBoneIndex = SkeletonRemapping.GetSourceSkeletonBoneIndex(TargetSkeletonBoneIndex);
+			const EBoneTranslationRetargetingMode::Type RetargetMode = FAnimationRuntime::GetBoneTranslationRetargetingMode(
+				bUseSourceRetargetModes,
+				SourceSkeletonBoneIndex,
+				TargetSkeletonBoneIndex,
+				SourceSkeleton,
+				TargetSkeleton,
+				bDisableRetargeting);
+
+			if (RetargetMode != EBoneTranslationRetargetingMode::Skeleton)
+			{
+				FVector BoneTranslation = OutPose.LocalTransformsView[PoseBoneIndex].GetTranslation();
+
+				if constexpr (bIsAdditive)
+				{
+					BoneTranslation = SkeletonRemapping.RetargetAdditiveTranslationToTargetSkeleton(TargetSkeletonBoneIndex, BoneTranslation);
+				}
+				else
+				{
+					BoneTranslation = SkeletonRemapping.RetargetBoneTranslationToTargetSkeleton(TargetSkeletonBoneIndex, BoneTranslation);
+				}
+
+				OutPose.LocalTransformsView[PoseBoneIndex].SetTranslation(BoneTranslation);
+			}
+		}
+
+		/**
+		 * Reference-pose correction per bone.
+		 * @param SourceSkeleton            Source skeleton to retarget from.
+		 * @param TargetSkeleton            Retargeting target skeleton.
+		 * @param SkeletonRemapping         Maps source and target skeleton bone indices.
+		 * @param bUseSourceRetargetModes   Set to true when you wish to extract the retarget mode from the source skeleton instead of the target skeleton.
+		 * @param bDisableRetargeting       Is translational retargeting disabled?
+		 * @param OutPose                   Target pose to work on.
+		 */
+		static UAF_API void ReferencePoseRetarget(
+			const USkeleton* SourceSkeleton,
+			const USkeleton* TargetSkeleton,
+			const FSkeletonRemapping& SkeletonRemapping,
+			bool bUseSourceRetargetModes,
+			bool bDisableRetargeting,
+			FLODPose& OutPose);
+
+		/**
+		 * Translation retarget a pose.
+		 * Applies the bone’s EBoneTranslationRetargetingMode:
+		 *  - AnimationScaled: scale local translation by source/target ref-length ratio.
+		 *  - AnimationRelative: transfer delta from authored base into target basis.
+		 *    and add to current transform (skip if baked additive upstream).
+		 *  - OrientAndScale: orientation-aware scaling when enabled.
+		 * @param TargetReferencePose       Target reference pose.
+		 * @param SourceSkeleton            Source skeleton to retarget from.
+		 * @param TargetSkeleton            Retargeting target skeleton.
+		 * @param RetargetTransforms        The name of the retarget source to find.
+		 * @param BoneTransform             .
+		 * @param SkeletonBoneIndex         Target skeleton bone index.
+		 * @param LODBoneIndex              Bone index in current LOD view.
+		 * @param bIsBakedAdditive          Treat input as baked additive.
+		 * @param bDisableRetargeting       Early-out passthrough.
+		 */
+		static UAF_API void RetargetBoneTransform(
+			const FReferencePose& TargetReferencePose,
+			const USkeleton* SourceSkeleton,
+			const USkeleton* TargetSkeleton,
+			const FName& SourceName,
+			const TArray<FTransform>& RetargetTransforms,
+			FTransform& BoneTransform,
+			const int32 SkeletonBoneIndex,
+			const int32 LODBoneIndex,
+			const bool bIsBakedAdditive,
+			const bool bDisableRetargeting);
+
+		/**
+		 * Translation retarget a pose.
+		 * Convenience path that pulls the source’s ref local poses (RetargetSource)
+		 * and forwards to the explicit-transform overload.
+		 * @param TargetReferencePose       Target reference pose.
+		 * @param SourceSkeleton            Source skeleton to retarget from.
+		 * @param TargetSkeleton            Retargeting target skeleton.
+		 * @param RetargetTransforms        The name of the retarget source to find.
+		 * @param BoneTransform             .
+		 * @param SkeletonBoneIndex         Target skeleton bone index.
+		 * @param LODBoneIndex              Bone index in current LOD view.
+		 * @param bIsBakedAdditive          Treat input as baked additive.
+		 * @param bDisableRetargeting       Early-out passthrough.
+		 */
+		static FORCEINLINE void RetargetBoneTransform(
+			const FReferencePose& TargetReferencePose,
+    		const USkeleton* SourceSkeleton,
+    		const USkeleton* TargetSkeleton,
+    		const FName& RetargetSource,
+    		FTransform& BoneTransform,
+    		const int32 SkeletonBoneIndex,
+    		const int32 LODBoneIndex,
+    		const bool bIsBakedAdditive,
+    		const bool bDisableRetargeting)
+    	{
+    		if (SourceSkeleton)
+    		{
+    			const TArray<FTransform>& RetargetTransforms = SourceSkeleton->GetRefLocalPoses(RetargetSource);
+    			RetargetBoneTransform(TargetReferencePose, SourceSkeleton, TargetSkeleton, RetargetSource, RetargetTransforms, BoneTransform, SkeletonBoneIndex, LODBoneIndex, bIsBakedAdditive, bDisableRetargeting);
+    		}
+    	}
+
+		/**
+		 * AnimationRelative path: apply authored-base delta in target basis.
+		 * @param TargetPoseBoneIndex       Bone index of the bone on the target pose.
+		 * @param TargetSkeletonBoneIndex   Corresponding bone index in the target's skeleton.
+		 * @param TargetReferencePose       Target reference pose.
+		 * @param SourceToTargetRemapping   Maps source and target skeleton bone indices.
+		 * @param BaseTransform             Authored base (source) local-space transform.
+		 * @param InOutBoneTransform        In/out current local transform to receive delta.
+		 */
+		static UAF_API void RetargetBone_AnimRelative(
+			const int32 TargetPoseBoneIndex,
+			const int32 TargetSkeletonBoneIndex,
+			const FReferencePose& TargetReferencePose,
+			const FSkeletonRemapping& SourceToTargetRemapping,
+			FTransform BaseTransform,
+			FTransform& InOutBoneTransform);
+
+		/**
+		 * Translation retarget a pose.
+		 * Runs the three per-mode passes (Scaled, Relative, Orient+Scale) over OutPose.
+		 * Work scales with number of pairs, not bones.
+		 * @param TargetReferencePose			Target reference pose.
+		 * @param SourceToTargetRemapping		Maps source and target skeleton bone indices.
+		 * @param AnimScaleRetargetingPairs		Pairs for AnimationScaled pass.
+		 * @param AnimRelativeRetargetingPairs	Pairs for AnimationRelative pass.
+		 * @param OrientAndScaleRetargetingPairs Pairs for OrientAndScale pass.
+		 * @param SourceReferencePoseTransforms  Source reference pose transforms.
+		 * @param OutPose						Target pose to work on in-place.
+		 */
+		static UAF_API void RetargetPose(
+			const FReferencePose& TargetReferencePose,
+			const FSkeletonRemapping& SourceToTargetRemapping,
+			BoneTrackArray& AnimScaleRetargetingPairs,
+			BoneTrackArray& AnimRelativeRetargetingPairs,
+			BoneTrackArray& OrientAndScaleRetargetingPairs,
+			const TArray<FTransform>& SourceReferencePoseTransforms,
+			FLODPose& OutPose);
+	};
+} // end namespace UE::UAF
