Fix for when replacing variable references incorrectly creating references to unrelated assets

#Synced-CL 48085336

diff --git a/Engine/Plugins/Experimental/UAF/UAF/Source/UAFUncookedOnly/Private/UncookedOnlyUtils.cpp b/Engine/Plugins/Experimental/UAF/UAF/Source/UAFUncookedOnly/Private/UncookedOnlyUtils.cpp
--- a/Engine/Plugins/Experimental/UAF/UAF/Source/UAFUncookedOnly/Private/UncookedOnlyUtils.cpp
+++ b/Engine/Plugins/Experimental/UAF/UAF/Source/UAFUncookedOnly/Private/UncookedOnlyUtils.cpp
@@ -1777,43 +1777,47 @@ void FUtils::ReplaceVariableReferences(UAnimNextRigVMAssetEditorData* InEditorDa
 	const bool bVariableTypeChanged = FindReferenceType != ReplaceReferenceType;
 	const bool bVariableNameChanged = ReferenceToFind.GetName() != ReferenceToReplaceWith.GetName();
 
-	// Ensure that SharedVariables have been included for replace variable (if from another asset)
-	const UAnimNextSharedVariables* SharedVariablesReplaceOuter = Cast<UAnimNextSharedVariables>(ReplaceOuterAsset);
-	if (SharedVariablesReplaceOuter && ThisOuterAsset != SharedVariablesReplaceOuter)
+	// Conditionally adds a SharedVariables dependency when required (no reference exists yet, and asset reference requiring one is introduced)
+	auto EnsureSharedVariableDependencyExists = [ReplaceOuterAsset, ThisOuterAsset, InEditorData, bSetupUndoRedo, bPrintPythonCommands]()
 	{
-		bool bHasEntry = false;
-		InEditorData->ForEachEntryOfType<UAnimNextSharedVariablesEntry>([&bHasEntry, SharedVariablesReplaceOuter](const UAnimNextSharedVariablesEntry* Entry)
+		// Ensure that SharedVariables have been included for replace variable (if from another asset)
+		const UAnimNextSharedVariables* SharedVariablesReplaceOuter = Cast<UAnimNextSharedVariables>(ReplaceOuterAsset);
+		if (SharedVariablesReplaceOuter && ThisOuterAsset != SharedVariablesReplaceOuter)
 		{
-			if (Entry->GetType() == EAnimNextSharedVariablesType::Asset)
+			bool bHasEntry = false;
+			InEditorData->ForEachEntryOfType<UAnimNextSharedVariablesEntry>([&bHasEntry, SharedVariablesReplaceOuter](const UAnimNextSharedVariablesEntry* Entry)
 			{
-				if (Entry->GetAsset() == SharedVariablesReplaceOuter)
+				if (Entry->GetType() == EAnimNextSharedVariablesType::Asset)
 				{
-					bHasEntry = true;
-					return false;
+					if (Entry->GetAsset() == SharedVariablesReplaceOuter)
+					{
+						bHasEntry = true;
+						return false;
+					}
 				}
-			}
 
-			return true;
-		});
+				return true;
+			});
 
-		if (!bHasEntry)
-		{
-			FString ErrorMessage;
-			if (CanAddSharedVariablesReference(InEditorData, SharedVariablesReplaceOuter, &ErrorMessage))
-			{
-				FCompilationScope CompilerResults(LOCTEXT("AddSharedAssetJobName", "Add Shared Asset dependency"), ThisOuterAsset);
-				InEditorData->AddSharedVariables(SharedVariablesReplaceOuter, bSetupUndoRedo, bPrintPythonCommands);
-			}
-			else
+			if (!bHasEntry)
 			{
-				FText NotificationText = FText::FromString(ErrorMessage);
-				FNotificationInfo Notification(NotificationText);
-				Notification.ExpireDuration = 3.0f;
-				FSlateNotificationManager::Get().AddNotification(Notification);
-				UE_LOGFMT(LogAnimation, Error, "{0}", ErrorMessage);
+				FString ErrorMessage;
+				if (CanAddSharedVariablesReference(InEditorData, SharedVariablesReplaceOuter, &ErrorMessage))
+				{
+					FCompilationScope CompilerResults(LOCTEXT("AddSharedAssetJobName", "Add Shared Asset dependency"), ThisOuterAsset);
+					InEditorData->AddSharedVariables(SharedVariablesReplaceOuter, bSetupUndoRedo, bPrintPythonCommands);
+				}
+				else
+				{
+					FText NotificationText = FText::FromString(ErrorMessage);
+					FNotificationInfo Notification(NotificationText);
+					Notification.ExpireDuration = 3.0f;
+					FSlateNotificationManager::Get().AddNotification(Notification);
+					UE_LOGFMT(LogAnimation, Error, "{0}", ErrorMessage);
+				}
 			}
 		}
-	}			
+	};	
 	
 	TArray<URigVMEdGraphNode*> EdNodes;
 	for(const UAnimNextRigVMAssetEntry* Entry : InEditorData->Entries)
@@ -1883,13 +1887,16 @@ void FUtils::ReplaceVariableReferences(UAnimNextRigVMAssetEditorData* InEditorDa
 										{
 											const FRigVMTemplateArgumentType NewRigVMType = ReplaceReferenceType.ToRigVMTemplateArgument();
 
+											EnsureSharedVariableDependencyExists();
+
 											constexpr bool bSetupOrphanPins = true;
 											Controller->RefreshSharedVariableNode(SharedVariableNode->GetFName(), ReferenceToReplaceWith.GetObject().GetPath(), ReferenceToReplaceWith.GetName(), NewRigVMType.CPPType.ToString(), NewRigVMType.CPPTypeObject.Get(), bSetupUndoRedo, bSetupOrphanPins, bPrintPythonCommands);											
 										}
 									}
 								}
 							}
-							else
+							// For regular variable nodes, the FindOuterAsset has to match the asset being processed
+							else if (FindOuterAsset == ThisOuterAsset)
 							{
 								if (bReplacingWithNone)
 								{
@@ -1902,6 +1909,7 @@ void FUtils::ReplaceVariableReferences(UAnimNextRigVMAssetEditorData* InEditorDa
 									{
 										if (FindOuterAsset == ThisOuterAsset && ReplaceOuterAsset != ThisOuterAsset)
 										{
+											EnsureSharedVariableDependencyExists();
 											Controller->ReplaceVariableNodeWithSharedVariableNode(VariableNode, ReferenceToReplaceWith.GetName(), ReferenceToReplaceWith.GetObject(), bSetupUndoRedo, bPrintPythonCommands);
 										}
 										else
