Only show modified (*) for asset entries in functions/variables outliners
#fix remove widget from regular outliner entries (only shown in asset entries now)
#Synced-CL 48141268

diff --git a/Engine/Plugins/Experimental/UAF/UAF/Source/UAFEditor/Private/Common/Outliner/OutlinerEntryItemWidget.cpp b/Engine/Plugins/Experimental/UAF/UAF/Source/UAFEditor/Private/Common/Outliner/OutlinerEntryItemWidget.cpp
--- a/Engine/Plugins/Experimental/UAF/UAF/Source/UAFEditor/Private/Common/Outliner/OutlinerEntryItemWidget.cpp
+++ b/Engine/Plugins/Experimental/UAF/UAF/Source/UAFEditor/Private/Common/Outliner/OutlinerEntryItemWidget.cpp
@@ -74,15 +74,6 @@ void SOutlinerEntryLabel::Construct(const FArguments& InArgs, FOutlinerEntryItem
 			.OnTextCommitted(this, &SOutlinerEntryLabel::OnTextCommited)
 			.OnVerifyTextChanged(this, &SOutlinerEntryLabel::OnVerifyTextChanged)
 		]
-		+SHorizontalBox::Slot()
-		.VAlign(VAlign_Center)
-		.Padding(2.0f, 0.0f, 2.0f, 3.0f)
-		.AutoWidth()
-		[
-			SNew(SImage)
-			.ToolTipText(this, &SOutlinerEntryLabel::GetDirtyTooltipText)
-			.Image(this, &SOutlinerEntryLabel::GetDirtyImageBrush)
-		]
 	];
 }
 
