Fix variable binding data and accessibility not persisting when moving variables between UAF assets

#Synced-CL 48085336

diff --git a/Engine/Plugins/Experimental/UAF/UAF/Source/UAFUncookedOnly/Private/UncookedOnlyUtils.cpp b/Engine/Plugins/Experimental/UAF/UAF/Source/UAFUncookedOnly/Private/UncookedOnlyUtils.cpp
--- a/Engine/Plugins/Experimental/UAF/UAF/Source/UAFUncookedOnly/Private/UncookedOnlyUtils.cpp
+++ b/Engine/Plugins/Experimental/UAF/UAF/Source/UAFUncookedOnly/Private/UncookedOnlyUtils.cpp
@@ -1684,6 +1684,13 @@ UAnimNextVariableEntry* FUtils::MoveVariableToAsset(UAnimNextVariableEntry* Vari
 			UAnimNextVariableEntry* NewEntry = NewEditorData->AddVariable(VariableEntry->GetVariableName(), VariableEntry->GetType(), DefaultValueString, bSetupUndoRedo, bPrintPythonCommands);
 			check(NewEntry);
 
+			TInstancedStruct<FAnimNextVariableBindingData> BindingCopy = VariableEntry->Binding.BindingData;
+			NewEntry->SetBinding(MoveTemp(BindingCopy));
+
+			NewEntry->SetExportAccessSpecifier(NewOuter->IsA<UAnimNextSharedVariables>()
+				? EAnimNextExportAccessSpecifier::Public
+				: VariableEntry->GetExportAccessSpecifier());
+
 			const FAnimNextSoftVariableReference FindReference = FAnimNextSoftVariableReference::FromName(VariableEntry->GetEntryName(), CurrentOuter);
 			const FAnimNextSoftVariableReference ReplaceReference = FAnimNextSoftVariableReference::FromName(VariableEntry->GetEntryName(), NewOuter);
 
