Enable support for Orient/Scale in UAF (retargeting data calculated on the fly, not cached yet as in FBoneContainer)

#Synced-CL 48351162

diff --git a/Engine/Plugins/Experimental/UAF/UAF/Source/UAF/Private/DecompressionTools.cpp b/Engine/Plugins/Experimental/UAF/UAF/Source/UAF/Private/DecompressionTools.cpp
--- a/Engine/Plugins/Experimental/UAF/UAF/Source/UAF/Private/DecompressionTools.cpp
+++ b/Engine/Plugins/Experimental/UAF/UAF/Source/UAF/Private/DecompressionTools.cpp
@@ -370,43 +370,51 @@ namespace Private
 			}
 		}
 
-		// TODO : Have to recreate GetRetargetSourceCachedData
-		// Translation 'Orient and Scale' Translation Retargeting
-		//const int32 NumBonesToOrientAndScaleRetarget = OrientAndScaleRetargetingPairs.Num();
-		//if (NumBonesToOrientAndScaleRetarget > 0)
-		//{
-		//	const FRetargetSourceCachedData& RetargetSourceCachedData = RequiredBones.GetRetargetSourceCachedData(DecompressionContext.AnimName, SkeletonRemapping, RetargetTransforms);
-		//	const TArray<FOrientAndScaleRetargetingCachedData>& OrientAndScaleDataArray = RetargetSourceCachedData.OrientAndScaleData;
-		//	const TArray<int32>& CompactPoseIndexToOrientAndScaleIndex = RetargetSourceCachedData.CompactPoseIndexToOrientAndScaleIndex;
-
-		//	// If we have any cached retargeting data.
-		//	if ((OrientAndScaleDataArray.Num() > 0) && (CompactPoseIndexToOrientAndScaleIndex.Num() == RequiredBones.GetCompactPoseNumBones()))
-		//	{
-		//		for (int32 Index = 0; Index < NumBonesToOrientAndScaleRetarget; Index++)
-		//		{
-		//			const BoneTrackPair& BonePair = OrientAndScaleRetargetingPairs[Index];
-		//			const FCompactPoseBoneIndex CompactPoseBoneIndex(BonePair.AtomIndex);
-		//			const int32 OrientAndScaleIndex = CompactPoseIndexToOrientAndScaleIndex[CompactPoseBoneIndex.GetInt()];
-		//			if (OrientAndScaleIndex != INDEX_NONE)
-		//			{
-		//				const FOrientAndScaleRetargetingCachedData& OrientAndScaleData = OrientAndScaleDataArray[OrientAndScaleIndex];
-		//				FTransform& BoneTransform = OutAnimationPoseData.LocalTransforms[CompactPoseBoneIndex];
-		//				const FVector AnimatedTranslation = BoneTransform.GetTranslation();
-
-		//				// If Translation is not animated, we can just copy the TargetTranslation. No retargeting needs to be done.
-		//				const FVector NewTranslation = (AnimatedTranslation - OrientAndScaleData.SourceTranslation).IsNearlyZero(BONE_TRANS_RT_ORIENT_AND_SCALE_PRECISION) ?
-		//					OrientAndScaleData.TargetTranslation :
-		//					OrientAndScaleData.TranslationDeltaOrient.RotateVector(AnimatedTranslation) * OrientAndScaleData.TranslationScale;
-
-		//				BoneTransform.SetTranslation(NewTranslation);
-		//			}
-		//		}
-		//	}
-		//}
-		UE_CLOGFMT(OrientAndScaleRetargetingPairs.Num() != 0,
-			LogAnimation, Verbose,
-			"OrientAndScaleRetargetingPairs found for {AnimSequence} ({SourceSkeleton} -> {TargetSkeleton}), Orient and Scale retargeting is not yet supported in UAF",
-			*DecompressionContext.AnimName.ToString(), SourceSkeleton ? *SourceSkeleton->GetName() : TEXT("None"), TargetSkeleton ? *TargetSkeleton->GetName() : TEXT("None"));
+		// TODO : Have to recreate GetRetargetSourceCachedData for optimization
+		const int32 NumBonesToOrientAndScaleRetarget = OrientAndScaleRetargetingPairs.Num();
+		if (NumBonesToOrientAndScaleRetarget > 0)
+		{
+			const TArray<FTransform>& AuthoredOnRefSkeleton = RetargetTransforms;
+			const TArray<FTransform>& TargetRefSkeleton = TargetSkeleton->GetRefLocalPoses();
+			
+			for (const BoneTrackPair& BonePair : OrientAndScaleRetargetingPairs)
+			{
+				const FAttributeSetIndex BoneSetIndex(BonePair.AtomIndex);
+				const int32 SourceSkeletonBoneIndex = BonePair.TrackIndex;
+				const int32 TargetSkeletonBoneIndex = bIsSkeletonRemappingValid ? SkeletonRemapping.GetTargetSkeletonBoneIndex(SourceSkeletonBoneIndex) : SourceSkeletonBoneIndex;
+				if(AuthoredOnRefSkeleton.IsValidIndex(SourceSkeletonBoneIndex) && TargetRefSkeleton.IsValidIndex(TargetSkeletonBoneIndex))
+				{
+					const FVector SourceSkelTrans = AuthoredOnRefSkeleton[SourceSkeletonBoneIndex].GetTranslation();
+					const FVector TargetSkelTrans = TargetRefSkeleton[TargetSkeletonBoneIndex].GetTranslation();
+					
+					// If translations are identical, we don't need to do any retargeting
+					if (!SourceSkelTrans.Equals(TargetSkelTrans, BONE_TRANS_RT_ORIENT_AND_SCALE_PRECISION))
+					{
+						const float SourceSkelTransLength = SourceSkelTrans.Size();
+						const float TargetSkelTransLength = TargetSkelTrans.Size();
+
+						// this only works on non zero vectors.
+						if (!FMath::IsNearlyZero(SourceSkelTransLength * TargetSkelTransLength))
+						{
+							const FVector SourceSkelTransDir = SourceSkelTrans / SourceSkelTransLength;
+							const FVector TargetSkelTransDir = TargetSkelTrans / TargetSkelTransLength;
+
+							const FQuat DeltaRotation = FQuat::FindBetweenNormals(SourceSkelTransDir, TargetSkelTransDir);
+							const float Scale = TargetSkelTransLength / SourceSkelTransLength;
+
+							FTransform& BoneTransform = OutBoneTransforms[BoneSetIndex].Value;
+							const FVector AnimatedTranslation = BoneTransform.GetTranslation();
+							// If Translation is not animated, we can just copy the TargetTranslation. No retargeting needs to be done.
+							const FVector NewTranslation = (AnimatedTranslation - SourceSkelTrans).IsNearlyZero(BONE_TRANS_RT_ORIENT_AND_SCALE_PRECISION) ?
+								TargetSkelTrans :
+								DeltaRotation.RotateVector(AnimatedTranslation) * Scale;
+
+							BoneTransform.SetTranslation(NewTranslation);
+						}
+					}
+				}
+			}
+		}
 
 		if (bIsMeshSpaceAdditive)
 		{
@@ -1282,7 +1290,7 @@ void FDecompressionTools::DecompressPose(FLODPose& OutAnimationPoseData,
 		RootMotionReset.ResetRootBoneForRootMotion(RootTransform, ReferencePose.GetRefPoseTransform(0));
 		OutAnimationPoseData.LocalTransformsView[0] = RootTransform;
 	}
-	
+
 	FRetargetingTools::RetargetPose(
 		ReferencePose,
 		SkeletonRemapping,
diff --git a/Engine/Plugins/Experimental/UAF/UAF/Source/UAF/Private/RetargetingTools.cpp b/Engine/Plugins/Experimental/UAF/UAF/Source/UAF/Private/RetargetingTools.cpp
--- a/Engine/Plugins/Experimental/UAF/UAF/Source/UAF/Private/RetargetingTools.cpp
+++ b/Engine/Plugins/Experimental/UAF/UAF/Source/UAF/Private/RetargetingTools.cpp
@@ -209,27 +209,37 @@ namespace UE::UAF
 			{
 				if (!bIsBakedAdditive)
 				{
-					//const FRetargetSourceCachedData& RetargetSourceCachedData = RequiredBones.GetRetargetSourceCachedData(SourceName, SkeletonRemapping, RetargetTransforms);
-					//const TArray<FOrientAndScaleRetargetingCachedData>& OrientAndScaleDataArray = RetargetSourceCachedData.OrientAndScaleData;
-					//const TArray<int32>& CompactPoseIndexToOrientAndScaleIndex = RetargetSourceCachedData.CompactPoseIndexToOrientAndScaleIndex;
-
-					//// If we have any cached retargeting data.
-					//if ((OrientAndScaleDataArray.Num() > 0) && (CompactPoseIndexToOrientAndScaleIndex.Num() == RequiredBones.GetCompactPoseNumBones()))
-					//{
-					//	const int32 OrientAndScaleIndex = CompactPoseIndexToOrientAndScaleIndex[BoneIndex.GetInt()];
-					//	if (OrientAndScaleIndex != INDEX_NONE)
-					//	{
-					//		const FOrientAndScaleRetargetingCachedData& OrientAndScaleData = OrientAndScaleDataArray[OrientAndScaleIndex];
-					//		const FVector AnimatedTranslation = BoneTransform.GetTranslation();
-
-					//		// If Translation is not animated, we can just copy the TargetTranslation. No retargeting needs to be done.
-					//		const FVector NewTranslation = (AnimatedTranslation - OrientAndScaleData.SourceTranslation).IsNearlyZero(BONE_TRANS_RT_ORIENT_AND_SCALE_PRECISION) ?
-					//			OrientAndScaleData.TargetTranslation :
-					//			OrientAndScaleData.TranslationDeltaOrient.RotateVector(AnimatedTranslation) * OrientAndScaleData.TranslationScale;
-
-					//		BoneTransform.SetTranslation(NewTranslation);
-					//	}
-					//}
+					const TArray<FTransform>& AuthoredOnRefSkeleton = RetargetTransforms;
+					if(AuthoredOnRefSkeleton.IsValidIndex(SourceSkeletonBoneIndex) && TargetSkeleton && TargetSkeleton->GetReferenceSkeleton().IsValidIndex(TargetSkeletonBoneIndex))
+					{
+						const FVector SourceSkelTrans = AuthoredOnRefSkeleton[SourceSkeletonBoneIndex].GetTranslation();
+						const FVector TargetSkelTrans = TargetSkeleton->GetReferenceSkeleton().GetRefBonePose()[TargetSkeletonBoneIndex].GetTranslation();
+
+						// If translations are identical, we don't need to do any retargeting
+						if (!SourceSkelTrans.Equals(TargetSkelTrans, BONE_TRANS_RT_ORIENT_AND_SCALE_PRECISION))
+						{
+							const float SourceSkelTransLength = SourceSkelTrans.Size();
+							const float TargetSkelTransLength = TargetSkelTrans.Size();
+
+							// this only works on non zero vectors.
+							if (!FMath::IsNearlyZero(SourceSkelTransLength * TargetSkelTransLength))
+							{
+								const FVector SourceSkelTransDir = SourceSkelTrans / SourceSkelTransLength;
+								const FVector TargetSkelTransDir = TargetSkelTrans / TargetSkelTransLength;
+
+								const FQuat DeltaRotation = FQuat::FindBetweenNormals(SourceSkelTransDir, TargetSkelTransDir);
+								const float Scale = TargetSkelTransLength / SourceSkelTransLength;
+
+								const FVector AnimatedTranslation = BoneTransform.GetTranslation();
+								// If Translation is not animated, we can just copy the TargetTranslation. No retargeting needs to be done.
+								const FVector NewTranslation = (AnimatedTranslation - SourceSkelTrans).IsNearlyZero(BONE_TRANS_RT_ORIENT_AND_SCALE_PRECISION) ?
+									TargetSkelTrans :
+									DeltaRotation.RotateVector(AnimatedTranslation) * Scale;
+
+								BoneTransform.SetTranslation(NewTranslation);
+							}
+						}
+					}
 				}
 				break;
 			}
@@ -326,43 +336,50 @@ namespace UE::UAF
 			}
 		}
 
-		// TODO : Have to recreate GetRetargetSourceCachedData
-		// Translation 'Orient and Scale' Translation Retargeting
-		//const int32 NumBonesToOrientAndScaleRetarget = OrientAndScaleRetargetingPairs.Num();
-		//if (NumBonesToOrientAndScaleRetarget > 0)
-		//{
-		//	const FRetargetSourceCachedData& RetargetSourceCachedData = RequiredBones.GetRetargetSourceCachedData(DecompressionContext.AnimName, SkeletonRemapping, RetargetTransforms);
-		//	const TArray<FOrientAndScaleRetargetingCachedData>& OrientAndScaleDataArray = RetargetSourceCachedData.OrientAndScaleData;
-		//	const TArray<int32>& CompactPoseIndexToOrientAndScaleIndex = RetargetSourceCachedData.CompactPoseIndexToOrientAndScaleIndex;
-
-		//	// If we have any cached retargeting data.
-		//	if ((OrientAndScaleDataArray.Num() > 0) && (CompactPoseIndexToOrientAndScaleIndex.Num() == RequiredBones.GetCompactPoseNumBones()))
-		//	{
-		//		for (int32 Index = 0; Index < NumBonesToOrientAndScaleRetarget; Index++)
-		//		{
-		//			const BoneTrackPair& BonePair = OrientAndScaleRetargetingPairs[Index];
-		//			const FCompactPoseBoneIndex CompactPoseBoneIndex(BonePair.AtomIndex);
-		//			const int32 OrientAndScaleIndex = CompactPoseIndexToOrientAndScaleIndex[CompactPoseBoneIndex.GetInt()];
-		//			if (OrientAndScaleIndex != INDEX_NONE)
-		//			{
-		//				const FOrientAndScaleRetargetingCachedData& OrientAndScaleData = OrientAndScaleDataArray[OrientAndScaleIndex];
-		//				FTransform& BoneTransform = OutAnimationPoseData.LocalTransforms[CompactPoseBoneIndex];
-		//				const FVector AnimatedTranslation = BoneTransform.GetTranslation();
-
-		//				// If Translation is not animated, we can just copy the TargetTranslation. No retargeting needs to be done.
-		//				const FVector NewTranslation = (AnimatedTranslation - OrientAndScaleData.SourceTranslation).IsNearlyZero(BONE_TRANS_RT_ORIENT_AND_SCALE_PRECISION) ?
-		//					OrientAndScaleData.TargetTranslation :
-		//					OrientAndScaleData.TranslationDeltaOrient.RotateVector(AnimatedTranslation) * OrientAndScaleData.TranslationScale;
-
-		//				BoneTransform.SetTranslation(NewTranslation);
-		//			}
-		//		}
-		//	}
-		//}
-		
-		// UE_CLOGFMT(OrientAndScaleRetargetingPairs.Num() != 0,
-		// 	LogAnimation, Verbose,
-		// 	"OrientAndScaleRetargetingPairs found for {AnimSequence} ({SourceSkeleton} -> {TargetSkeleton}), Orient and Scale retargeting is not yet supported in UAF",
-		// 	*DecompressionContext.AnimName.ToString(), SourceSkeleton ? *SourceSkeleton->GetName() : TEXT("None"), TargetSkeleton ? *TargetSkeleton->GetName() : TEXT("None"));
+		// TODO : Have to recreate GetRetargetSourceCachedData for optimization
+		const int32 NumBonesToOrientAndScaleRetarget = OrientAndScaleRetargetingPairs.Num();
+		if (NumBonesToOrientAndScaleRetarget > 0)
+		{
+			const TArray<FTransform>& AuthoredOnRefSkeleton = SourceReferencePoseTransforms;
+			const TArray<FTransform>& TargetSkeletonRefPose = (SourceToTargetRemapping.IsValid() ? SourceToTargetRemapping.GetTargetSkeleton() : TargetReferencePose.GetSkeletonAsset())->GetReferenceSkeleton().GetRefBonePose();
+			
+			for (const BoneTrackPair& BonePair : OrientAndScaleRetargetingPairs)
+			{
+				const int32 LODBoneIndex(BonePair.AtomIndex);
+				const int32 SourceSkeletonBoneIndex = BonePair.TrackIndex;
+				const int32 TargetSkeletonBoneIndex = SourceToTargetRemapping.IsValid() ? SourceToTargetRemapping.GetTargetSkeletonBoneIndex(SourceSkeletonBoneIndex) : SourceSkeletonBoneIndex;
+				
+				if(AuthoredOnRefSkeleton.IsValidIndex(SourceSkeletonBoneIndex) && TargetSkeletonRefPose.IsValidIndex(TargetSkeletonBoneIndex))
+				{
+					const FVector SourceSkelTrans = AuthoredOnRefSkeleton[SourceSkeletonBoneIndex].GetTranslation();
+					const FVector TargetSkelTrans = TargetSkeletonRefPose[TargetSkeletonBoneIndex].GetTranslation();
+			
+					// If translations are identical, we don't need to do any retargeting
+					if (!SourceSkelTrans.Equals(TargetSkelTrans, BONE_TRANS_RT_ORIENT_AND_SCALE_PRECISION))
+					{
+						const float SourceSkelTransLength = SourceSkelTrans.Size();
+						const float TargetSkelTransLength = TargetSkelTrans.Size();
+
+						// this only works on non zero vectors.
+						if (!FMath::IsNearlyZero(SourceSkelTransLength * TargetSkelTransLength))
+						{
+							const FVector SourceSkelTransDir = SourceSkelTrans / SourceSkelTransLength;
+							const FVector TargetSkelTransDir = TargetSkelTrans / TargetSkelTransLength;
+
+							const FQuat DeltaRotation = FQuat::FindBetweenNormals(SourceSkelTransDir, TargetSkelTransDir);
+							const float Scale = TargetSkelTransLength / SourceSkelTransLength;
+
+							const FVector AnimatedTranslation = OutPose.LocalTransformsView[LODBoneIndex].GetTranslation();
+							// If Translation is not animated, we can just copy the TargetTranslation. No retargeting needs to be done.
+							const FVector NewTranslation = (AnimatedTranslation - SourceSkelTrans).IsNearlyZero(BONE_TRANS_RT_ORIENT_AND_SCALE_PRECISION) ?
+								TargetSkelTrans :
+								DeltaRotation.RotateVector(AnimatedTranslation) * Scale;
+
+							OutPose.LocalTransformsView[LODBoneIndex].SetTranslation(NewTranslation);
+						}
+					}
+				}
+			}
+		}
 	}
-} // end namespace UE::UAF
\ No newline at end of file
+} // end namespace UE::UAF
