UAF: Crash when selecting functions in Call RigVM Function Trait
#fix convert check into compiler error informing user of missing Shared Variables reference (or something went horribly wrong)
#Synced-CL 48351162

diff --git a/Engine/Plugins/Experimental/UAF/UAF/Source/UAFUncookedOnly/Private/AnimNextRigVMAssetEditorData.cpp b/Engine/Plugins/Experimental/UAF/UAF/Source/UAFUncookedOnly/Private/AnimNextRigVMAssetEditorData.cpp
--- a/Engine/Plugins/Experimental/UAF/UAF/Source/UAFUncookedOnly/Private/AnimNextRigVMAssetEditorData.cpp
+++ b/Engine/Plugins/Experimental/UAF/UAF/Source/UAFUncookedOnly/Private/AnimNextRigVMAssetEditorData.cpp
@@ -967,7 +967,21 @@ void UAnimNextRigVMAssetEditorData::BuildFunctionWrapperEvents(FAnimNextRigVMAss
 			{
 				return InVariable.Name == InternalArgName;
 			});
-			check(VariableIndex != INDEX_NONE);
+
+			if (VariableIndex == INDEX_NONE)
+			{
+				if (FunctionHeader.GetFunctionHost() != this)
+				{
+					InSettings.ReportError(FString::Printf(TEXT("Missing Shared Variables Reference to %s required for function %s"), FunctionHeader.GetFunctionHost() ? *FUtils::GetAsset(Cast<UAnimNextRigVMAssetEditorData>(FunctionHeader.GetFunctionHost()))->GetFName().ToString() : TEXT("None"), *FunctionHeader.Name.ToString()));
+				}
+				else
+				{
+					InSettings.ReportError(FString::Printf(TEXT("Failed to link internal variable param node to function: %s -> %s"), *Argument.Name.ToString(), *FunctionHeader.Name.ToString()));
+				}
+				
+				continue;
+			}
+			
 			ArgIndices.Add(VariableIndex);
 
 			if (bIsGetter)
