Ensure trying to delete blank function in UAF
#fix fallback to editor-data on which delegate callback was registered to update compilation status
#Synced-CL 48351162

diff --git a/Engine/Plugins/Experimental/UAF/UAF/Source/UAFEditor/Private/Workspace/AnimNextWorkspaceEditorMode.cpp b/Engine/Plugins/Experimental/UAF/UAF/Source/UAFEditor/Private/Workspace/AnimNextWorkspaceEditorMode.cpp
--- a/Engine/Plugins/Experimental/UAF/UAF/Source/UAFEditor/Private/Workspace/AnimNextWorkspaceEditorMode.cpp
+++ b/Engine/Plugins/Experimental/UAF/UAF/Source/UAFEditor/Private/Workspace/AnimNextWorkspaceEditorMode.cpp
@@ -17,8 +17,6 @@
 #include "IWorkspaceEditorModule.h"
 #include "Modules/ModuleManager.h"
 #include "RigVMCommands.h"
-#include "RigVMCore/RigVMGraphFunctionHost.h"
-#include "RigVMModel/Nodes/RigVMFunctionReferenceNode.h"
 #include "Toolkits/AssetEditorToolkitMenuContext.h"
 #include "Toolkits/BaseToolkit.h"
 #include "ToolMenus.h"
@@ -63,10 +61,11 @@ void UAnimNextWorkspaceEditorMode::Enter()
 		UAnimNextRigVMAsset* AnimNextRigVMAsset = static_cast<UAnimNextRigVMAsset*>(Asset);
 		UAnimNextRigVMAssetEditorData* EditorData = UE::UAF::UncookedOnly::FUtils::GetEditorData<UAnimNextRigVMAssetEditorData>(AnimNextRigVMAsset);
 
+		TWeakObjectPtr<UAnimNextRigVMAssetEditorData> WeakEditorData = EditorData;
 		EditorData->RigVMCompiledEvent.RemoveAll(this);
 		EditorData->RigVMCompiledEvent.AddUObject(this, &UAnimNextWorkspaceEditorMode::HandleRigVMCompiledEvent);
 		EditorData->RigVMGraphModifiedEvent.RemoveAll(this);
-		EditorData->RigVMGraphModifiedEvent.AddUObject(this, &UAnimNextWorkspaceEditorMode::HandleRigVMModifiedEvent);
+		EditorData->RigVMGraphModifiedEvent.AddUObject(this, &UAnimNextWorkspaceEditorMode::HandleRigVMModifiedEvent, WeakEditorData);
 		EditorData->RecompileRequiredChangedEvent.RemoveAll(this);
 		EditorData->RecompileRequiredChangedEvent.AddUObject(this, &UAnimNextWorkspaceEditorMode::UpdateCompileStatus);
 
@@ -314,10 +313,11 @@ void UAnimNextWorkspaceEditorMode::HandleFocusedDocumentChanged(const UE::Worksp
 
 	// Subscribe to asset compilation/modification
 	UAnimNextRigVMAssetEditorData* EditorData = UE::UAF::UncookedOnly::FUtils::GetEditorData<UAnimNextRigVMAssetEditorData>(AnimNextRigVMAsset);
+	TWeakObjectPtr<UAnimNextRigVMAssetEditorData> WeakEditorData = EditorData;
 	EditorData->RigVMCompiledEvent.RemoveAll(this);
 	EditorData->RigVMCompiledEvent.AddUObject(this, &UAnimNextWorkspaceEditorMode::HandleRigVMCompiledEvent);
 	EditorData->RigVMGraphModifiedEvent.RemoveAll(this);
-	EditorData->RigVMGraphModifiedEvent.AddUObject(this, &UAnimNextWorkspaceEditorMode::HandleRigVMModifiedEvent);
+	EditorData->RigVMGraphModifiedEvent.AddUObject(this, &UAnimNextWorkspaceEditorMode::HandleRigVMModifiedEvent, WeakEditorData);
 	EditorData->RecompileRequiredChangedEvent.RemoveAll(this);
 	EditorData->RecompileRequiredChangedEvent.AddUObject(this, &UAnimNextWorkspaceEditorMode::UpdateCompileStatus);
 
@@ -355,7 +355,7 @@ void UAnimNextWorkspaceEditorMode::HandleRigVMCompiledEvent(UObject* InAsset, UR
 	UpdateCompileStatus();
 }
 
-void UAnimNextWorkspaceEditorMode::HandleRigVMModifiedEvent(ERigVMGraphNotifType InType, URigVMGraph* InGraph, UObject* InSubject)
+void UAnimNextWorkspaceEditorMode::HandleRigVMModifiedEvent(ERigVMGraphNotifType InType, URigVMGraph* InGraph, UObject* InSubject, TWeakObjectPtr<UAnimNextRigVMAssetEditorData> WeakEditorData)
 {
 	if(InGraph == nullptr)
 	{
@@ -364,18 +364,10 @@ void UAnimNextWorkspaceEditorMode::HandleRigVMModifiedEvent(ERigVMGraphNotifType
 
 	UAnimNextRigVMAssetEditorData* EditorData = InGraph->GetTypedOuter<UAnimNextRigVMAssetEditorData>();
 
-	// Our RigVM Graph may be a function, in this case the graph has a temp outer not the editor data. 
-	// So search for the host object of the function instead.
+	// Some notifies will not produce a valid outer editor data e.g. deleting or changing a RigVM function - so try to retrieve it from the weak-object ptr (set when registering delegate)
 	if (!EditorData)
 	{
-		if (URigVMFunctionReferenceNode* FunctionReferenceNode = Cast<URigVMFunctionReferenceNode>(InSubject))
-		{
-			constexpr bool bLoadIfNecessary = true;
-			if (UAnimNextRigVMAssetEditorData* FunctionHost = Cast<UAnimNextRigVMAssetEditorData>(FunctionReferenceNode->GetReferencedFunctionHeader().GetFunctionHost(bLoadIfNecessary)))
-			{
-				EditorData = FunctionHost;
-			}
-		}
+		EditorData = WeakEditorData.Get();
 	}
 	
 	if (!ensureMsgf(EditorData, TEXT("Failed to find UAF asset for RigVM modification, asset may not update correctly. It also may fail to propagate dirty status to other assets.")))
diff --git a/Engine/Plugins/Experimental/UAF/UAF/Source/UAFEditor/Private/Workspace/AnimNextWorkspaceEditorMode.h b/Engine/Plugins/Experimental/UAF/UAF/Source/UAFEditor/Private/Workspace/AnimNextWorkspaceEditorMode.h
--- a/Engine/Plugins/Experimental/UAF/UAF/Source/UAFEditor/Private/Workspace/AnimNextWorkspaceEditorMode.h
+++ b/Engine/Plugins/Experimental/UAF/UAF/Source/UAFEditor/Private/Workspace/AnimNextWorkspaceEditorMode.h
@@ -9,6 +9,7 @@
 #include "AnimNextWorkspaceEditorMode.generated.h"
 
 class UAnimNextWorkspaceSchema;
+class UAnimNextRigVMAssetEditorData;
 class URigVM;
 class URigVMGraph;
 struct FRigVMExtendedExecuteContext;
@@ -67,7 +68,7 @@ private:
 	void HandleRigVMCompiledEvent(UObject* InAsset, URigVM* InVM, FRigVMExtendedExecuteContext& InExtendedExecuteContext);
 
 	// Handle an asset getting modified
-	void HandleRigVMModifiedEvent(ERigVMGraphNotifType InType, URigVMGraph* InGraph, UObject* InSubject);
+	void HandleRigVMModifiedEvent(ERigVMGraphNotifType InType, URigVMGraph* InGraph, UObject* InSubject, TWeakObjectPtr<UAnimNextRigVMAssetEditorData> WeakEditorData);
 
 	// Handle the document focus changing
 	void HandleFocusedDocumentChanged(const UE::Workspace::FWorkspaceDocument& InDocument); 
@@ -99,4 +100,4 @@ private:
 
 	// UAF browser status bar widget, used to store a custom drawer next to the content browser
 	TSharedPtr<SWidget> UAFBrowserStatusbarWidget = nullptr;
-};
\ No newline at end of file
+};
