Adding an Input Pin to a UAF function and changing it from value to array crashes the engine.
#fix FRigVMTemplateArgumentType -> FAnimNextParamType parsing of TArray<Type> values
#Synced-CL 48351162

diff --git a/Engine/Plugins/Experimental/UAF/UAF/Source/UAF/Private/Param/ParamType.cpp b/Engine/Plugins/Experimental/UAF/UAF/Source/UAF/Private/Param/ParamType.cpp
--- a/Engine/Plugins/Experimental/UAF/UAF/Source/UAF/Private/Param/ParamType.cpp
+++ b/Engine/Plugins/Experimental/UAF/UAF/Source/UAF/Private/Param/ParamType.cpp
@@ -83,19 +83,17 @@ FRigVMTemplateArgumentType FAnimNextParamType::ToRigVMTemplateArgument() const
 
 FAnimNextParamType FAnimNextParamType::FromRigVMTemplateArgument(const FRigVMTemplateArgumentType& RigVMType)
 {
-	FAnimNextParamType Type;	
-	const FString CPPTypeString = RigVMType.CPPType.ToString();
-	
-	if (RigVMTypeUtils::IsArrayType(CPPTypeString))
+	FAnimNextParamType Type;
+	if (RigVMType.IsArray())
 	{
 		Type.ContainerType = EPropertyBagContainerType::Array;
 	}
-
+	
 	static const FName IntTypeName(TEXT("int")); // type used by some engine tests
 	static const FName Int64TypeName(TEXT("Int64"));
 	static const FName UInt64TypeName(TEXT("UInt64"));
 
-	const FName CPPType = *CPPTypeString;
+	const FName CPPType = *RigVMType.GetBaseCPPType();
 	
 	if (CPPType == RigVMTypeUtils::BoolTypeName)
 	{
@@ -154,7 +152,7 @@ FAnimNextParamType FAnimNextParamType::FromRigVMTemplateArgument(const FRigVMTem
 	}
 	else
 	{
-		ensureMsgf(false, TEXT("Unsupported type : %s"), *CPPTypeString);
+		ensureMsgf(false, TEXT("Unsupported type : %s"), *RigVMType.CPPType.ToString());
 		Type.ValueType = EPropertyBagPropertyType::None;
 	}
 
