- Allow SharedVariables asset to be created when selecting a single outliner asset item (in addition to a set of variables within a single asset).
- Changed user facing label to reflect creating, potential, empty shared variable asset
#Synced-CL 48351162

diff --git a/Engine/Plugins/Experimental/UAF/UAF/Source/UAFEditor/Private/Variables/Outliner/VariablesOutlinerCommands.cpp b/Engine/Plugins/Experimental/UAF/UAF/Source/UAFEditor/Private/Variables/Outliner/VariablesOutlinerCommands.cpp
--- a/Engine/Plugins/Experimental/UAF/UAF/Source/UAFEditor/Private/Variables/Outliner/VariablesOutlinerCommands.cpp
+++ b/Engine/Plugins/Experimental/UAF/UAF/Source/UAFEditor/Private/Variables/Outliner/VariablesOutlinerCommands.cpp
@@ -10,7 +10,7 @@ void FVariablesOutlinerCommands::RegisterCommands()
 	UI_COMMAND(AddNewVariable, "Add Variable", "Adds a single new variable to the selected asset.", EUserInterfaceActionType::Button, FInputChord());
 	UI_COMMAND(AddNewVariables, "Add Variable(s)", "Adds variables to assets.\nIf multiple assets are selected, then variables will be added to each.\nIf no assets are selected and there are multiple assets, variables will be added to all assets.", EUserInterfaceActionType::Button, FInputChord());
 
-	UI_COMMAND(CreateSharedVariablesAssets, "Create Shared Variables from Selection", "Creates a new Shared Variables asset from the selected set of variables.", EUserInterfaceActionType::Button, FInputChord());
+	UI_COMMAND(CreateSharedVariablesAssets, "Create Shared Variables", "Creates a new Shared Variables asset, adding and moving all of the selected set of variables.", EUserInterfaceActionType::Button, FInputChord());
 }
 
-#undef LOCTEXT_NAMESPACE
\ No newline at end of file
+#undef LOCTEXT_NAMESPACE
diff --git a/Engine/Plugins/Experimental/UAF/UAF/Source/UAFEditor/Private/Variables/Outliner/VariablesOutlinerMode.cpp b/Engine/Plugins/Experimental/UAF/UAF/Source/UAFEditor/Private/Variables/Outliner/VariablesOutlinerMode.cpp
--- a/Engine/Plugins/Experimental/UAF/UAF/Source/UAFEditor/Private/Variables/Outliner/VariablesOutlinerMode.cpp
+++ b/Engine/Plugins/Experimental/UAF/UAF/Source/UAFEditor/Private/Variables/Outliner/VariablesOutlinerMode.cpp
@@ -92,16 +92,13 @@ void FVariablesOutlinerMode::RegisterToolMenus()
 				
 				FToolMenuSection& VariablesSection = InMenu->AddSection("Variables", LOCTEXT("VariablesSectionLabel", "Variables"));
 
-				const int32 NumSelectedVariables = MenuContext->WeakEntries.Num();
 				/** Variable specific actions */			
 				{
-					// Populates entries for adding of different variable types/sources					
+					// Populates entries for adding of different variable types/sources	
 					PopulateNewVariableToolMenuEntries(InMenu, false);
 					
-					if (NumSelectedVariables >= 1)
-					{
-						VariablesSection.AddMenuEntryWithCommandList(FVariablesOutlinerCommands::Get().CreateSharedVariablesAssets, OutlinerContext->WeakCommandList.Pin(), TAttribute<FText>(), TAttribute<FText>(), FSlateIcon(FUAFStyle::Get().GetStyleSetName(), "ClassIcon.AnimNextSharedVariables"));
-					}
+					FToolMenuEntry& Entry = VariablesSection.AddMenuEntryWithCommandList(FVariablesOutlinerCommands::Get().CreateSharedVariablesAssets, OutlinerContext->WeakCommandList.Pin(), TAttribute<FText>(), TAttribute<FText>(), FSlateIcon(FUAFStyle::Get().GetStyleSetName(), "ClassIcon.AnimNextSharedVariables"));
+					Entry.InsertPosition = FToolMenuInsert("AddNewVariable", EToolMenuInsertType::After);
 				}
 			}), FToolMenuInsert(FCommonOutlinerMode::ContextToolAssetSectionName, EToolMenuInsertType::After));
 		}
@@ -1848,68 +1845,82 @@ void FVariablesOutlinerMode::CreateSharedVariablesAssets() const
 		}
 	});
 
-	if (EntriesToMove.Num())
+	TArray<UAnimNextRigVMAssetEditorData*> UniqueOuterEditorDatas;
+	
+	TArray<FOutlinerAssetItem*> AssetItems;
+	Selection.Get<FOutlinerAssetItem>(AssetItems);
+	if (AssetItems.Num() == 1)
+	{
+		UAnimNextRigVMAsset* Asset = AssetItems[0]->SoftAsset.LoadSynchronous();
+		UniqueOuterEditorDatas.AddUnique(UncookedOnly::FUtils::GetEditorData(Asset));
+	}
+	
+	Selection.ForEachItem<FVariablesOutlinerEntryItem>([&EntriesToMove](const FVariablesOutlinerEntryItem& VariableItem)
 	{
-		// Prompt user for asset path
-		FString SaveObjectPath;
-		{			
-			FSaveAssetDialogConfig SaveAssetDialogConfig;
+		if (UAnimNextVariableEntry* VariableEntry = VariableItem.WeakEntry.Get())
+		{
+			EntriesToMove.Add(VariableEntry);
+		}
+	});
+
+	// Prompt user for asset path
+	FString SaveObjectPath;
+	{			
+		FSaveAssetDialogConfig SaveAssetDialogConfig;
+		{
+			// Populate default path to match currently focussed assets
+			TSharedPtr<Workspace::IWorkspaceEditor> SharedWorkspaceEditor = WeakWorkspaceEditor.Pin();
+			if (SharedWorkspaceEditor.IsValid())
 			{
-				// Populate default path to match currently focussed assets
-				TSharedPtr<Workspace::IWorkspaceEditor> SharedWorkspaceEditor = WeakWorkspaceEditor.Pin();
-				if (SharedWorkspaceEditor.IsValid())
+				const Workspace::FWorkspaceDocument& Document = SharedWorkspaceEditor->GetFocusedWorkspaceDocument();
+				if (UObject* FocussedObject = Document.GetObject())
 				{
-					const Workspace::FWorkspaceDocument& Document =SharedWorkspaceEditor->GetFocusedWorkspaceDocument();
-					if (UObject* FocussedObject = Document.GetObject())
-					{
-						SaveAssetDialogConfig.DefaultPath = FPaths::GetPath(FocussedObject->GetPackage()->GetPathName());
-					}
+					SaveAssetDialogConfig.DefaultPath = FPaths::GetPath(FocussedObject->GetPackage()->GetPathName());
 				}
-				
-				SaveAssetDialogConfig.DefaultAssetName = TEXT("NewSharedVariables");
-				SaveAssetDialogConfig.AssetClassNames.Add(UAnimNextSharedVariables::StaticClass()->GetClassPathName());
-				SaveAssetDialogConfig.ExistingAssetPolicy = ESaveAssetDialogExistingAssetPolicy::AllowButWarn;
-				SaveAssetDialogConfig.DialogTitleOverride = LOCTEXT("SaveLiveLinkPresetDialogTitle", "Add new Shared Variables Asset");
 			}
 
-			FContentBrowserModule& ContentBrowserModule = FModuleManager::LoadModuleChecked<FContentBrowserModule>("ContentBrowser");
-			SaveObjectPath = ContentBrowserModule.Get().CreateModalSaveAssetDialog(SaveAssetDialogConfig);
+			SaveAssetDialogConfig.DefaultAssetName = TEXT("NewSharedVariables");
+			SaveAssetDialogConfig.AssetClassNames.Add(UAnimNextSharedVariables::StaticClass()->GetClassPathName());
+			SaveAssetDialogConfig.ExistingAssetPolicy = ESaveAssetDialogExistingAssetPolicy::AllowButWarn;
+			SaveAssetDialogConfig.DialogTitleOverride = LOCTEXT("SaveLiveLinkPresetDialogTitle", "Add new Shared Variables Asset");
 		}
 
-		if (!SaveObjectPath.IsEmpty())
-		{
-			const FString PackagePath = FPackageName::ObjectPathToPackageName(SaveObjectPath);
-			const FString AssetName = FPaths::GetBaseFilename(PackagePath, true);
-			
-			const FText FormattedMessage = FText::Format(LOCTEXT("CreateSharedVariablesAsset", "Creating new Shared Variables Asset {0}"), FText::FromString(AssetName));
-			FScopedTransaction Transaction(FormattedMessage);			
-			
-			UPackage* Package = CreatePackage(*PackagePath);
-			Package->Modify();
-				
-			UAnimNextSharedVariablesFactory* Factory = NewObject<UAnimNextSharedVariablesFactory>(GetTransientPackage());
-			UAnimNextSharedVariables* NewSharedVariables = CastChecked<UAnimNextSharedVariables>(Factory->FactoryCreateNew(UAnimNextSharedVariables::StaticClass(), Package, *AssetName, RF_Public | RF_Standalone | RF_Transactional, nullptr, nullptr, NAME_None));
-			check(NewSharedVariables);
+		FContentBrowserModule& ContentBrowserModule = FModuleManager::LoadModuleChecked<FContentBrowserModule>("ContentBrowser");
+		SaveObjectPath = ContentBrowserModule.Get().CreateModalSaveAssetDialog(SaveAssetDialogConfig);
+	}
+
+	if (!SaveObjectPath.IsEmpty())
+	{
+		const FString PackagePath = FPackageName::ObjectPathToPackageName(SaveObjectPath);
+		const FString AssetName = FPaths::GetBaseFilename(PackagePath, true);
+		
+		const FText FormattedMessage = FText::Format(LOCTEXT("CreateSharedVariablesAsset", "Creating new Shared Variables Asset {0}"), FText::FromString(AssetName));
+		FScopedTransaction Transaction(FormattedMessage);			
+		
+		UPackage* Package = CreatePackage(*PackagePath);
+		Package->Modify();
 			
-			TArray<UAnimNextRigVMAssetEditorData*> UniqueOuterEditorDatas;
-			for (UAnimNextVariableEntry* Entry : EntriesToMove)
-			{
-				// Force variable to be public as its being moved to SharedVariables asset
-				Entry->SetExportAccessSpecifier(EAnimNextExportAccessSpecifier::Public);
+		UAnimNextSharedVariablesFactory* Factory = NewObject<UAnimNextSharedVariablesFactory>(GetTransientPackage());
+		UAnimNextSharedVariables* NewSharedVariables = CastChecked<UAnimNextSharedVariables>(Factory->FactoryCreateNew(UAnimNextSharedVariables::StaticClass(), Package, *AssetName, RF_Public | RF_Standalone | RF_Transactional, nullptr, nullptr, NAME_None));
+		check(NewSharedVariables);
+		
+		for (UAnimNextVariableEntry* Entry : EntriesToMove)
+		{
+			// Force variable to be public as its being moved to SharedVariables asset
+			Entry->SetExportAccessSpecifier(EAnimNextExportAccessSpecifier::Public);
 
-				UniqueOuterEditorDatas.AddUnique(Entry->GetTypedOuter<UAnimNextRigVMAssetEditorData>());
-				UncookedOnly::FUtils::MoveVariableToAsset(Entry, NewSharedVariables);
-			}
+			UniqueOuterEditorDatas.AddUnique(Entry->GetTypedOuter<UAnimNextRigVMAssetEditorData>());
+			UncookedOnly::FUtils::MoveVariableToAsset(Entry, NewSharedVariables);
+		}
 
-			for (UAnimNextRigVMAssetEditorData* OuterEditorData : UniqueOuterEditorDatas)
-			{
-				OuterEditorData->AddSharedVariables(NewSharedVariables);
-			}
-			
-			// mark asset dirty 
-			FAssetRegistryModule::AssetCreated(NewSharedVariables);
-			NewSharedVariables->MarkPackageDirty();
+		for (UAnimNextRigVMAssetEditorData* OuterEditorData : UniqueOuterEditorDatas)
+		{
+			OuterEditorData->AddSharedVariables(NewSharedVariables);
 		}
+		
+		// mark asset dirty 
+		FAssetRegistryModule::AssetCreated(NewSharedVariables);
+		NewSharedVariables->MarkPackageDirty();
 	}
 }
 
@@ -1917,6 +1928,7 @@ bool FVariablesOutlinerMode::CanCreateSharedVariablesAssets() const
 {
 	const FSceneOutlinerItemSelection& Selection = SceneOutliner->GetSelection();
 	const bool bSelectionOnlyContainsVariables = Selection.Num() == Selection.Num<FVariablesOutlinerEntryItem>();
+	const bool bSelectionOnlyContainsSingleAsset = Selection.Num() == 1 && Selection.Num<FOutlinerAssetItem>() == 1;
 
 	// Prevent moving from an assets which would leave it empty (removing all variables)
 	TMap<const UAnimNextRigVMAssetEditorData*, int32> EditorDataToNumRemainingVariables;
@@ -1952,7 +1964,7 @@ bool FVariablesOutlinerMode::CanCreateSharedVariablesAssets() const
 	EditorDataToNumRemainingVariables.GenerateValueArray(NumRemainVariablesArray);
 	const bool bMoveWillLeaveEmptyAsset = NumRemainVariablesArray.Contains(0);
 	
-	return bSelectionOnlyContainsVariables && !bMoveWillLeaveEmptyAsset;
+	return (bSelectionOnlyContainsVariables && !bMoveWillLeaveEmptyAsset) || bSelectionOnlyContainsSingleAsset;
 }
 
 }
