UAF: Scene graph support changes

Initial implementation of Pose I/O via FRigUnit_UAFWriteOutputPose and FUAFOutputPoseComponent
Misc fixes to allow SG implementation to function

#Synced-CL 48342630

diff --git a/Engine/Plugins/Experimental/UAF/UAF/Source/UAF/Internal/Graph/RigUnit_UAFWriteOutputPose.h b/Engine/Plugins/Experimental/UAF/UAF/Source/UAF/Internal/Graph/RigUnit_UAFWriteOutputPose.h
new file mode 100644
--- /dev/null
+++ b/Engine/Plugins/Experimental/UAF/UAF/Source/UAF/Internal/Graph/RigUnit_UAFWriteOutputPose.h
@@ -0,0 +1,31 @@
+﻿// Copyright Epic Games, Inc. All Rights Reserved.
+
+#pragma once
+
+#include "AnimNextExecuteContext.h"
+#include "Graph/AnimNext_LODPose.h"
+#include "Graph/RigUnit_AnimNextBase.h"
+
+#include "RigUnit_UAFWriteOutputPose.generated.h"
+
+#define UE_API UAF_API
+
+/** Writes the output pose for a system */
+USTRUCT(meta=(DisplayName="Write Output Pose", Category="Animation Graph", NodeColor="0, 1, 1", Keywords="Output,Pose,Port", RequiredComponents="UAFOutputPoseComponent"))
+struct FRigUnit_UAFWriteOutputPose : public FRigUnit_AnimNextBase
+{
+	GENERATED_BODY()
+
+	RIGVM_METHOD()
+	UE_API void Execute();
+
+	// Pose to write to output
+	UPROPERTY(EditAnywhere, Category = "Graph", meta = (Input))
+	FAnimNextGraphLODPose Pose;
+
+	// The execution result
+	UPROPERTY(EditAnywhere, DisplayName = "Execute", Category = "BeginExecution", meta = (Input, Output))
+	FAnimNextExecuteContext ExecuteContext;
+};
+
+#undef UE_API
\ No newline at end of file
diff --git a/Engine/Plugins/Experimental/UAF/UAF/Source/UAF/Internal/Graph/UAFOutputPoseComponent.h b/Engine/Plugins/Experimental/UAF/UAF/Source/UAF/Internal/Graph/UAFOutputPoseComponent.h
new file mode 100644
--- /dev/null
+++ b/Engine/Plugins/Experimental/UAF/UAF/Source/UAF/Internal/Graph/UAFOutputPoseComponent.h
@@ -0,0 +1,81 @@
+﻿// Copyright Epic Games, Inc. All Rights Reserved.
+
+#pragma once
+
+#include "AnimNext_LODPose.h"
+#include "Module/UAFModuleInstanceComponent.h"
+
+#include "UAFOutputPoseComponent.generated.h"
+
+#define UE_API UAF_API
+
+struct FRigUnit_UAFWriteOutputPose;
+class UMeshEntityComponentImpl_Skinned;
+class FEntitySkeletalAnimationModule;
+struct FUAFOutputPoseComponent;
+
+// Module component used to identify output poses
+USTRUCT()
+struct FUAFOutputPoseComponent : public FUAFModuleInstanceComponent
+{
+	GENERATED_BODY()
+	DECLARE_UAF_ASSET_INSTANCE_COMPONENT()
+
+	// Get the refpose for the mesh we are bound to
+	const FAnimNextGraphReferencePose& GetRefPose() const
+	{
+		ensure(RefPose.ReferencePose.IsValid());
+		return RefPose;
+	}
+
+private:
+	// Write the pose
+	// @param    InPose    The pose to write
+	void WritePose(const FAnimNextGraphLODPose& InPose);
+
+	// Generate the pose etc. for rendering by a mesh.
+	// If the serial numbers differ, a read occurs, otherwise the read is skipped.
+	// If the supplied serial number matches the one held in this component, then the pose is assumed to have not been written, so no copy needs to occur
+	// If the serial number supplied is zero, then the read will always succeed
+	// @param    InSerialNumber    The current serial number that the reader holds
+	// @param    InMeshLOD         The pose LOD to generate
+	// @param    InLocalToWorld    The local to world transform used to generate bounds
+	// @param    InFunction        The function to call on read
+	// @return true if the read was successful (i.e. the serial numbers were mismatched)
+	UE_API bool GenerateMeshRenderData(uint32 InSerialNumber, int32 InMeshLOD, const FTransform& InLocalToWorld, TFunctionRef<void(TConstArrayView<FMatrix44f> InReferenceToLocalMatrices, const FBoxSphereBounds& InWorldBounds, const FBoxSphereBounds& InLocalBounds, uint32 InSerialNumber)> InReadFunction) const;
+
+	// Helper function used to calculate animated mesh bounds.
+	// Note: Adapted from USkinnedMeshComponent::CalcMeshBound
+	void CalcMeshBound(TConstArrayView<FTransform> InComponentTransforms, const FTransform& InLocalToWorld, FBoxSphereBounds& OutWorldBounds, FBoxSphereBounds& OutLocalBounds) const;
+
+	// Helper function used to calculate physics-asset defined bounds.
+	// Note: Adapted from UPhysicsAsset::CalcAABB
+	void CalcBoundsFromPhysicsAsset(const UPhysicsAsset* InPhysicAsset, TConstArrayView<FTransform> InComponentTransforms, const FTransform& InLocalToWorld, FBox& OutWorldBounds, FBox& OutLocalBounds) const;
+
+	// Called to bind this component to a specific mesh host
+	UE_API void BindToHost(USkeletalMesh* InMesh, TFunction<void(void)>&& InSignalFunction);
+
+private:
+	friend FRigUnit_UAFWriteOutputPose;
+	friend UMeshEntityComponentImpl_Skinned;
+	friend FEntitySkeletalAnimationModule;
+
+	// The output pose that was written this time around
+	FAnimNextGraphLODPose Pose;
+
+	// The mesh we use to determine the final pose format
+	UPROPERTY(Transient)
+	TObjectPtr<USkeletalMesh> Mesh;
+
+	// Refpose generated from the mesh
+	UPROPERTY(Transient)
+	FAnimNextGraphReferencePose RefPose;
+
+	// Function to call when a pose is written (on the GT)
+	TFunction<void(void)> SignalFunction;
+
+	// Serial number updated each time we write the pose
+	uint32 SerialNumber = 0;
+};
+
+#undef UE_API
\ No newline at end of file
diff --git a/Engine/Plugins/Experimental/UAF/UAF/Source/UAF/Internal/Module/AnimNextModuleInstance.h b/Engine/Plugins/Experimental/UAF/UAF/Source/UAF/Internal/Module/AnimNextModuleInstance.h
--- a/Engine/Plugins/Experimental/UAF/UAF/Source/UAF/Internal/Module/AnimNextModuleInstance.h
+++ b/Engine/Plugins/Experimental/UAF/UAF/Source/UAF/Internal/Module/AnimNextModuleInstance.h
@@ -173,6 +173,22 @@ private:
 	// Accesses value of the specified variable via the proxy double-buffer. Called from any thread.
 	EPropertyBagResult WriteProxyVariable(const FAnimNextVariableReference& InVariable, const FAnimNextParamType& InType, TFunctionRef<void(TArrayView<uint8>)> InFunction);
 
+	// Access a variable's value given its reference via a raw memory view.
+	// This performs non-recursive locks on the system instance to ensure that concurrent access is respected, so care must be taken calling this from
+	// animation worker tasks as deadlocks can occur.
+	// Type must match strictly, no conversions are performed.
+	// @param	InVariable			The variable to get the value of
+	// @param	InType				The type of the variable
+	// @param	InFunction			Function that will be called if no errors occur
+	// @return see EPropertyBagResult
+	EPropertyBagResult ReadVariableFromExternal(const FAnimNextVariableReference& InVariable, const FAnimNextParamType& InType, TFunctionRef<void(TConstArrayView<uint8>)> InFunction);
+
+	// Access the supplied system component for reading. Component will not be lazily created.
+	// @param	InComponentType		The type of the component we want ot access
+	// @param	InFunction			Function that will be called to access the struct
+	// @return true if the struct was available, otherwise false
+	bool ReadComponentFromExternal(UScriptStruct* InComponentType, TFunctionRef<void(FConstStructView InComponentStruct)> InFunction);
+
 private:
 	friend UE::UAF::FSystemReference;
 	friend FUAFWeakSystemReference;
@@ -208,6 +224,9 @@ private:
 	// Lock to ensure proxy access is thread safe
 	FRWLock ProxyLock;
 
+	// Lock guarding against gross concurrent access via system references (safety only: should rarely contend if dependencies are set correctly)
+	FRWLock SystemReferenceLock;
+
 	// Proxy public variables, double-buffered
 	UPROPERTY(Transient)
 	FUAFInstanceVariableDataProxy ProxyVariables[2];
diff --git a/Engine/Plugins/Experimental/UAF/UAF/Source/UAF/Private/Component/AnimNextComponent.cpp b/Engine/Plugins/Experimental/UAF/UAF/Source/UAF/Private/Component/AnimNextComponent.cpp
--- a/Engine/Plugins/Experimental/UAF/UAF/Source/UAF/Private/Component/AnimNextComponent.cpp
+++ b/Engine/Plugins/Experimental/UAF/UAF/Source/UAF/Private/Component/AnimNextComponent.cpp
@@ -24,6 +24,11 @@ void UAnimNextComponent::OnRegister()
 		check(!SystemReference.IsValid());
 
 		SystemReference = FSystemReference(Asset, this, InitMethod);
+
+#if WITH_EDITOR
+		SystemReference.RunInitialTickInEditor();
+#endif
+
 #if UE_ENABLE_DEBUG_DRAWING
 		SystemReference.ShowDebugDrawing(bShowDebugDrawing);
 #endif
diff --git a/Engine/Plugins/Experimental/UAF/UAF/Source/UAF/Private/DataRegistry.cpp b/Engine/Plugins/Experimental/UAF/UAF/Source/UAF/Private/DataRegistry.cpp
--- a/Engine/Plugins/Experimental/UAF/UAF/Source/UAF/Private/DataRegistry.cpp
+++ b/Engine/Plugins/Experimental/UAF/UAF/Source/UAF/Private/DataRegistry.cpp
@@ -48,7 +48,8 @@ namespace UE::UAF
 		FCoreDelegates::OnEnginePreExit.Remove(DataRegistryPreExitHandle);
 		DataRegistryPreExitHandle.Reset();
 
-		if (!GAnimationDataRegistry->GCKeysToCleanup.IsEmpty())
+		if (!GAnimationDataRegistry->GCComponentKeysToCleanup.IsEmpty() ||
+			!GAnimationDataRegistry->GCMeshKeysToCleanup.IsEmpty())
 		{
 			GAnimationDataRegistry->RunGCCleanup(INDEX_NONE);
 
@@ -69,7 +70,8 @@ void FDataRegistry::OnPreExit()
 {
 	checkf(GAnimationDataRegistry, TEXT("Animation Data Registry is not instanced. It is only valid to access this while the engine module is loaded."));
 
-	if (!GAnimationDataRegistry->GCKeysToCleanup.IsEmpty())
+	if (!GAnimationDataRegistry->GCComponentKeysToCleanup.IsEmpty() ||
+		!GAnimationDataRegistry->GCMeshKeysToCleanup.IsEmpty())
 	{
 		GAnimationDataRegistry->RunGCCleanup(INDEX_NONE);
 
@@ -93,9 +95,11 @@ FDataRegistry* FDataRegistry::Get()
 	// Compact the registry on GC
 	if (GAnimationDataRegistry)
 	{
-		FRWScopeLock Lock(GAnimationDataRegistry->SkeletalMeshReferencePosesLock, SLT_Write);
+		FRWScopeLock ComponentLock(GAnimationDataRegistry->SkeletalMeshComponentReferencePosesLock, SLT_Write);
+		FRWScopeLock MeshLock(GAnimationDataRegistry->SkeletalMeshReferencePosesLock, SLT_Write);
 
-		if (GAnimationDataRegistry->SkeletalMeshReferencePoses.IsEmpty())
+		if (GAnimationDataRegistry->SkeletalMeshComponentReferencePoses.IsEmpty() &&
+			GAnimationDataRegistry->SkeletalMeshReferencePoses.IsEmpty())
 		{
 			// Nothing to clean up
 			return;
@@ -103,7 +107,8 @@ FDataRegistry* FDataRegistry::Get()
 
 		// Start our cleanup
 		// If cleanup was already in progress, we purge whatever remains (shouldn't happen unless we force GC to run)
-		if (!GAnimationDataRegistry->GCKeysToCleanup.IsEmpty())
+		if (!GAnimationDataRegistry->GCComponentKeysToCleanup.IsEmpty() ||
+			!GAnimationDataRegistry->GCMeshKeysToCleanup.IsEmpty())
 		{
 			GAnimationDataRegistry->RunGCCleanup(INDEX_NONE);
 		}
@@ -114,8 +119,10 @@ FDataRegistry* FDataRegistry::Get()
 		}
 
 		// Copy our keys into a temporary array, we'll process them in small batches over time
-		GAnimationDataRegistry->SkeletalMeshReferencePoses.GenerateKeyArray(GAnimationDataRegistry->GCKeysToCleanup);
-		GAnimationDataRegistry->GCKeysProcessedIndex = 0;
+		GAnimationDataRegistry->SkeletalMeshReferencePoses.GenerateKeyArray(GAnimationDataRegistry->GCMeshKeysToCleanup);
+		GAnimationDataRegistry->SkeletalMeshComponentReferencePoses.GenerateKeyArray(GAnimationDataRegistry->GCComponentKeysToCleanup);
+		GAnimationDataRegistry->GCComponentKeysProcessedIndex = 0;
+		GAnimationDataRegistry->GCMeshKeysProcessedIndex = 0;
 
 		const float TickerDelay = 0.0f;
 		GAnimationDataRegistry->GCCleanupTickerHandle = FTSTicker::GetCoreTicker().AddTicker(TEXT("UAF DataRegistry Post-GC Cleanup"), TickerDelay, GCCleanupTicker);
@@ -126,7 +133,7 @@ bool FDataRegistry::GCCleanupTicker(float DeltaTime)
 {
 	bool bIsDone;
 	{
-		FRWScopeLock Lock(GAnimationDataRegistry->SkeletalMeshReferencePosesLock, SLT_Write);
+		FRWScopeLock Lock(GAnimationDataRegistry->SkeletalMeshComponentReferencePosesLock, SLT_Write);
 		bIsDone = GAnimationDataRegistry->RunGCCleanup(CVar_GCCleanupTimeBudget.GetValueOnAnyThread());
 	}
 
@@ -146,15 +153,15 @@ bool FDataRegistry::RunGCCleanup(double TimeBudget)
 
 	const double StartTime = FPlatformTime::Seconds();
 
-	while (GCKeysProcessedIndex < GCKeysToCleanup.Num())
+	while (GCComponentKeysProcessedIndex < GCComponentKeysToCleanup.Num())
 	{
-		const TObjectKey<USkeletalMeshComponent>& SkeletalMeshComponentKey = GCKeysToCleanup[GCKeysProcessedIndex];
+		const TObjectKey<USkeletalMeshComponent>& SkeletalMeshComponentKey = GCComponentKeysToCleanup[GCComponentKeysProcessedIndex];
 		if (SkeletalMeshComponentKey.ResolveObjectPtr() == nullptr)
 		{
-			SkeletalMeshReferencePoses.Remove(SkeletalMeshComponentKey);
+			SkeletalMeshComponentReferencePoses.Remove(SkeletalMeshComponentKey);
 		}
 
-		GCKeysProcessedIndex++;
+		GCComponentKeysProcessedIndex++;
 
 		if ((FPlatformTime::Seconds() - StartTime) > TimeBudget)
 		{
@@ -162,11 +169,30 @@ bool FDataRegistry::RunGCCleanup(double TimeBudget)
 		}
 	}
 
-	if (GCKeysProcessedIndex == GCKeysToCleanup.Num())
+	while (GCMeshKeysProcessedIndex < GCMeshKeysToCleanup.Num())
+	{
+		const TObjectKey<USkeletalMesh>& SkeletalMeshKey = GCMeshKeysToCleanup[GCMeshKeysProcessedIndex];
+		if (SkeletalMeshKey.ResolveObjectPtr() == nullptr)
+		{
+			SkeletalMeshReferencePoses.Remove(SkeletalMeshKey);
+		}
+
+		GCMeshKeysProcessedIndex++;
+
+		if ((FPlatformTime::Seconds() - StartTime) > TimeBudget)
+		{
+			break;
+		}
+	}
+
+	if (GCComponentKeysProcessedIndex == GCComponentKeysToCleanup.Num() &&
+		GCMeshKeysProcessedIndex == GCMeshKeysToCleanup.Num())
 	{
 		// We are done
-		GCKeysToCleanup.Empty();
-		GCKeysProcessedIndex = 0;
+		GCComponentKeysToCleanup.Empty();
+		GCMeshKeysToCleanup.Empty();
+		GCComponentKeysProcessedIndex = 0;
+		GCMeshKeysProcessedIndex = 0;
 		return true;
 	}
 
@@ -188,10 +214,10 @@ FDataHandle FDataRegistry::RegisterReferencePose(USkeletalMeshComponent* Skeleta
 	const uint32 KeyHash = GetTypeHash(Key);
 
 	{
-		FRWScopeLock Lock(SkeletalMeshReferencePosesLock, SLT_Write);
+		FRWScopeLock Lock(SkeletalMeshComponentReferencePosesLock, SLT_Write);
 
 		// First check if another thread grabbed the write lock before us to do the same work
-		if (const FReferencePoseData* ExistingPoseData = SkeletalMeshReferencePoses.FindByHash(KeyHash, Key))
+		if (const FReferencePoseData* ExistingPoseData = SkeletalMeshComponentReferencePoses.FindByHash(KeyHash, Key))
 		{
 			// Another thread beat us to it, discard the work we did and re-use what is cached
 			return ExistingPoseData->AnimationDataHandle;
@@ -200,7 +226,36 @@ FDataHandle FDataRegistry::RegisterReferencePose(USkeletalMeshComponent* Skeleta
 		// Only register the delegate if we are the one to add to the map
 		FDelegateHandle DelegateHandle = SkeletalMeshComponent->RegisterOnLODRequiredBonesUpdate_Member(FOnLODRequiredBonesUpdate::CreateRaw(this, &FDataRegistry::OnLODRequiredBonesUpdate));
 
-		SkeletalMeshReferencePoses.AddByHash(KeyHash, Key, FReferencePoseData(Handle, DelegateHandle));
+		SkeletalMeshComponentReferencePoses.AddByHash(KeyHash, Key, FReferencePoseData(Handle, DelegateHandle));
+	}
+
+	return Handle;
+}
+
+FDataHandle FDataRegistry::RegisterReferencePose(USkeletalMesh* SkeletalMesh)
+{
+	FDataHandle Handle = AllocateData<FAnimNextReferencePose>(1);
+
+	FAnimNextReferencePose& AnimationReferencePose = Handle.GetRef<FAnimNextReferencePose>();
+
+	FGenerationTools::GenerateReferencePose(nullptr, SkeletalMesh, AnimationReferencePose);
+
+	// register even if it fails to generate (register an empty ref pose)
+
+	const TObjectKey<USkeletalMesh> Key(SkeletalMesh);
+	const uint32 KeyHash = GetTypeHash(Key);
+
+	{
+		FRWScopeLock Lock(SkeletalMeshReferencePosesLock, SLT_Write);
+
+		// First check if another thread grabbed the write lock before us to do the same work
+		if (const FReferencePoseData* ExistingPoseData = SkeletalMeshReferencePoses.FindByHash(KeyHash, Key))
+		{
+			// Another thread beat us to it, discard the work we did and re-use what is cached
+			return ExistingPoseData->AnimationDataHandle;
+		}
+		
+		SkeletalMeshReferencePoses.AddByHash(KeyHash, Key, FReferencePoseData(Handle, FDelegateHandle()));
 	}
 
 	return Handle;
@@ -217,9 +272,9 @@ FDataHandle FDataRegistry::GetOrGenerateReferencePose(USkeletalMeshComponent* Sk
 	FDataHandle ReturnHandle;
 
 	{
-		FRWScopeLock Lock(SkeletalMeshReferencePosesLock, SLT_ReadOnly);
+		FRWScopeLock Lock(SkeletalMeshComponentReferencePosesLock, SLT_ReadOnly);
 
-		if (const FReferencePoseData* ReferencePoseData = SkeletalMeshReferencePoses.Find(SkeletalMeshComponent))
+		if (const FReferencePoseData* ReferencePoseData = SkeletalMeshComponentReferencePoses.Find(SkeletalMeshComponent))
 		{
 			ReturnHandle = ReferencePoseData->AnimationDataHandle;
 		}
@@ -233,21 +288,50 @@ FDataHandle FDataRegistry::GetOrGenerateReferencePose(USkeletalMeshComponent* Sk
 	return ReturnHandle;
 }
 
+
+FDataHandle FDataRegistry::GetOrGenerateReferencePose(USkeletalMesh* SkeletalMesh)
+{
+	FDataHandle ReturnHandle;
+
+	{
+		FRWScopeLock Lock(SkeletalMeshReferencePosesLock, SLT_ReadOnly);
+
+		if (const FReferencePoseData* ReferencePoseData = SkeletalMeshReferencePoses.Find(SkeletalMesh))
+		{
+			ReturnHandle = ReferencePoseData->AnimationDataHandle;
+		}
+	}
+	
+	if (ReturnHandle.IsValid() == false)
+	{
+		ReturnHandle = RegisterReferencePose(SkeletalMesh);
+	}
+
+	return ReturnHandle;
+}
+
 void FDataRegistry::RemoveReferencePose(USkeletalMeshComponent* SkeletalMeshComponent)
 {
 	if (SkeletalMeshComponent != nullptr)
 	{
-		FRWScopeLock Lock(SkeletalMeshReferencePosesLock, SLT_Write);
+		FRWScopeLock Lock(SkeletalMeshComponentReferencePosesLock, SLT_Write);
 
 		FReferencePoseData ReferencePoseData;
-		if (SkeletalMeshReferencePoses.RemoveAndCopyValue(SkeletalMeshComponent, ReferencePoseData))
+		if (SkeletalMeshComponentReferencePoses.RemoveAndCopyValue(SkeletalMeshComponent, ReferencePoseData))
 		{
 			SkeletalMeshComponent->UnregisterOnLODRequiredBonesUpdate_Member(ReferencePoseData.DelegateHandle);
 		}
 	}
 }
 
-
+void FDataRegistry::RemoveReferencePose(USkeletalMesh* SkeletalMesh)
+{
+	if (SkeletalMesh != nullptr)
+	{
+		FRWScopeLock Lock(SkeletalMeshReferencePosesLock, SLT_Write);
+		SkeletalMeshReferencePoses.Remove(SkeletalMesh);
+	}
+}
 
 void FDataRegistry::RegisterData(const FName& Id, const FDataHandle& AnimationDataHandle)
 {
@@ -311,20 +395,26 @@ void FDataRegistry::FreeAllocatedBlock(Private::FAllocatedBlock* AllocatedBlock)
 // Remove any ReferencePoses and unregister all the SkeletalMeshComponent delegates (if any still alive)
 void FDataRegistry::ReleaseReferencePoseData()
 {
-	FRWScopeLock Lock(SkeletalMeshReferencePosesLock, SLT_Write);
-
-	for (auto Iter = GAnimationDataRegistry->SkeletalMeshReferencePoses.CreateIterator(); Iter; ++Iter)
 	{
-		const TObjectKey<USkeletalMeshComponent>& SkeletalMeshComponentPtr = Iter.Key();
-		const FReferencePoseData& ReferencePoseData = Iter.Value();
+		FRWScopeLock Lock(SkeletalMeshComponentReferencePosesLock, SLT_Write);
 
-		if (USkeletalMeshComponent* SkeletalMeshComponent = SkeletalMeshComponentPtr.ResolveObjectPtr())
+		for (auto Iter = GAnimationDataRegistry->SkeletalMeshComponentReferencePoses.CreateIterator(); Iter; ++Iter)
 		{
-			SkeletalMeshComponent->UnregisterOnLODRequiredBonesUpdate_Member(ReferencePoseData.DelegateHandle);
+			const TObjectKey<USkeletalMeshComponent>& SkeletalMeshComponentPtr = Iter.Key();
+			const FReferencePoseData& ReferencePoseData = Iter.Value();
+
+			if (USkeletalMeshComponent* SkeletalMeshComponent = SkeletalMeshComponentPtr.ResolveObjectPtr())
+			{
+				SkeletalMeshComponent->UnregisterOnLODRequiredBonesUpdate_Member(ReferencePoseData.DelegateHandle);
+			}
 		}
-	}
 
-	SkeletalMeshReferencePoses.Empty();
+		SkeletalMeshComponentReferencePoses.Empty();
+	}
+	{
+		FRWScopeLock Lock(SkeletalMeshReferencePosesLock, SLT_Write);
+		SkeletalMeshReferencePoses.Empty();
+	}
 }
 
 } // end namespace UE::UAF
diff --git a/Engine/Plugins/Experimental/UAF/UAF/Source/UAF/Private/Factory/SystemFactory.cpp b/Engine/Plugins/Experimental/UAF/UAF/Source/UAF/Private/Factory/SystemFactory.cpp
--- a/Engine/Plugins/Experimental/UAF/UAF/Source/UAF/Private/Factory/SystemFactory.cpp
+++ b/Engine/Plugins/Experimental/UAF/UAF/Source/UAF/Private/Factory/SystemFactory.cpp
@@ -17,7 +17,7 @@ static FDelegateHandle GSystemFactoryPreExitHandle;
 static TStripedMap<32, uint64, TWeakObjectPtr<const UAnimNextModule>> GSystemMap;
 
 // All registered object -> param mappings
-static TMap<const UClass*, TPair<FUAFSystemFactoryParams, FSystemFactory::FParamsInitializer>> GDefaultParamsForObject;
+static TMap<TObjectKey<UClass>, TPair<FUAFSystemFactoryParams, FSystemFactory::FParamsInitializer>> GDefaultParamsForObject;
 static TArray<UClass*> GAllFactoryClasses;
 
 const TPair<FUAFSystemFactoryParams, FSystemFactory::FParamsInitializer>* FindSystemParamsForClass(const UClass* InClass)
diff --git a/Engine/Plugins/Experimental/UAF/UAF/Source/UAF/Private/Graph/RigUnit_UAFWriteOutputPose.cpp b/Engine/Plugins/Experimental/UAF/UAF/Source/UAF/Private/Graph/RigUnit_UAFWriteOutputPose.cpp
new file mode 100644
--- /dev/null
+++ b/Engine/Plugins/Experimental/UAF/UAF/Source/UAF/Private/Graph/RigUnit_UAFWriteOutputPose.cpp
@@ -0,0 +1,21 @@
+﻿// Copyright Epic Games, Inc. All Rights Reserved.
+
+#include "Graph/RigUnit_UAFWriteOutputPose.h"
+
+#include "AnimNextStats.h"
+#include "Graph/UAFOutputPoseComponent.h"
+#include "Module/AnimNextModuleInstance.h"
+
+#include UE_INLINE_GENERATED_CPP_BY_NAME(RigUnit_UAFWriteOutputPose)
+
+FRigUnit_UAFWriteOutputPose_Execute()
+{
+	SCOPE_CYCLE_COUNTER(STAT_AnimNext_Write_Pose);
+
+	using namespace UE::UAF;
+
+	const FAnimNextModuleContextData& ModuleContextData = ExecuteContext.GetContextData<FAnimNextModuleContextData>();
+	FAnimNextModuleInstance& ModuleInstance = ModuleContextData.GetModuleInstance();
+	FUAFOutputPoseComponent& OutputPoseComponent = ModuleInstance.GetComponent<FUAFOutputPoseComponent>();
+	OutputPoseComponent.WritePose(Pose);
+}
diff --git a/Engine/Plugins/Experimental/UAF/UAF/Source/UAF/Private/Graph/UAFOutputPoseComponent.cpp b/Engine/Plugins/Experimental/UAF/UAF/Source/UAF/Private/Graph/UAFOutputPoseComponent.cpp
new file mode 100644
--- /dev/null
+++ b/Engine/Plugins/Experimental/UAF/UAF/Source/UAF/Private/Graph/UAFOutputPoseComponent.cpp
@@ -0,0 +1,263 @@
+﻿// Copyright Epic Games, Inc. All Rights Reserved.
+
+#include "Graph/UAFOutputPoseComponent.h"
+
+#include "AnimationRuntime.h"
+#include "ReferencePose.h"
+#include "Engine/SkeletalMesh.h"
+#include "Graph/AnimNext_LODPose.h"
+#include "Module/AnimNextModuleInstance.h"
+#include "Module/ModuleEvents.h"
+#include "PhysicsEngine/PhysicsAsset.h"
+#include "PhysicsEngine/SkeletalBodySetup.h"
+
+void FUAFOutputPoseComponent::CalcBoundsFromPhysicsAsset(const UPhysicsAsset* InPhysicAsset, TConstArrayView<FTransform> InComponentTransforms, const FTransform& InLocalToWorld, FBox& OutWorldBounds, FBox& OutLocalBounds) const
+{
+	OutLocalBounds = FBox(ForceInit);
+	OutWorldBounds = FBox(ForceInit);
+
+	const FReferenceSkeleton& RefSkeleton = Mesh->GetRefSkeleton();
+	TConstArrayView<int32> BodyIndexRefs = InPhysicAsset->BoundsBodies;
+	TConstArrayView<TObjectPtr<USkeletalBodySetup>> SkeletalBodySetups = InPhysicAsset->SkeletalBodySetups;
+
+	// Then iterate over bodies we want to consider, calculating bounding box for each
+	const int32 BodySetupNum = BodyIndexRefs.Num();
+	for(int32 BodySetupIndex = 0; BodySetupIndex < BodySetupNum; ++BodySetupIndex)
+	{
+		const int32 BodyIndex = BodyIndexRefs[BodySetupIndex];
+		const UBodySetup* BodySetup = SkeletalBodySetups[BodyIndex];
+
+		if (BodySetupIndex + 1 < BodySetupNum)
+		{
+			int32 NextIndex = BodyIndexRefs[BodySetupIndex + 1];
+			FPlatformMisc::Prefetch(SkeletalBodySetups[NextIndex]);
+			FPlatformMisc::Prefetch(SkeletalBodySetups[NextIndex], PLATFORM_CACHE_LINE_SIZE);
+		}
+
+		const int32 BoneIndex = RefSkeleton.FindBoneIndex(BodySetup->BoneName);
+		if(BoneIndex != INDEX_NONE)
+		{
+			FBox BodySetupLocalBounds = BodySetup->AggGeom.CalcAABB(InComponentTransforms[BoneIndex]);
+
+			const FTransform WorldBoneTransform = InComponentTransforms[BoneIndex] * InLocalToWorld;
+			FBox BodySetupWorldBounds = BodySetup->AggGeom.CalcAABB(WorldBoneTransform);
+
+			auto SanitizeScaledBounds = [](FBox& InBounds)
+			{
+				if (InBounds.Min.X > InBounds.Max.X)
+				{
+					Swap(InBounds.Min.X, InBounds.Max.X);
+				}
+
+				if (InBounds.Min.Y > InBounds.Max.Y)
+				{
+					Swap(InBounds.Min.Y, InBounds.Max.Y);
+				}
+
+				if (InBounds.Min.Z > InBounds.Max.Z)
+				{
+					Swap(InBounds.Min.Z, InBounds.Max.Z);
+				}
+			};
+
+			SanitizeScaledBounds(BodySetupLocalBounds);
+			SanitizeScaledBounds(BodySetupWorldBounds);
+
+			OutLocalBounds += BodySetupLocalBounds;
+			OutWorldBounds += BodySetupWorldBounds;
+		}
+	}
+
+	auto SanitizeZeroBounds = [&InLocalToWorld](FBox& InBox)
+	{
+		if(!InBox.IsValid)
+		{
+			InBox = FBox(InLocalToWorld.GetLocation(), InLocalToWorld.GetLocation());
+		}
+
+		const float MinBoundSize = 1.f;
+		const FVector BoxSize = InBox.GetSize();
+
+		if(BoxSize.GetMin() < MinBoundSize)
+		{
+			const FVector ExpandByDelta ( FMath::Max<FVector::FReal>(0, MinBoundSize - BoxSize.X), FMath::Max<FVector::FReal>(0, MinBoundSize - BoxSize.Y), FMath::Max<FVector::FReal>(0, MinBoundSize - BoxSize.Z) );
+			InBox = InBox.ExpandBy(ExpandByDelta * 0.5f);	//expand by applies to both directions with GetSize applies to total size so divide by 2
+		}
+	};
+
+	SanitizeZeroBounds(OutLocalBounds);
+	SanitizeZeroBounds(OutWorldBounds);
+}
+
+void FUAFOutputPoseComponent::CalcMeshBound(TConstArrayView<FTransform> InComponentTransforms, const FTransform& InLocalToWorld, FBoxSphereBounds& OutWorldBounds, FBoxSphereBounds& OutLocalBounds) const
+{
+	const FVector DrawScale = InLocalToWorld.GetScale3D();
+	const UPhysicsAsset* PhysicsAsset = Mesh->GetPhysicsAsset();
+
+	// Can only use the PhysicsAsset to calculate the bounding box if we are not non-uniformly scaling the mesh.
+	if(DrawScale.IsUniform() && PhysicsAsset != nullptr && PhysicsAsset->SkeletalBodySetups.Num())
+	{
+		FBox WorldBox;
+		FBox LocalBox;
+		CalcBoundsFromPhysicsAsset(PhysicsAsset, InComponentTransforms, InLocalToWorld, WorldBox, LocalBox);
+		OutLocalBounds = FBoxSphereBounds(LocalBox);
+		OutWorldBounds = FBoxSphereBounds(WorldBox);
+	}
+	// Fallback is to use the one from the skeletal mesh.
+	// Usually pretty bad in terms of accuracy of where the SkelMesh Bounds are located (i.e. usually bigger than it needs to be)
+	else
+	{
+		OutLocalBounds = Mesh->GetBounds();
+		OutWorldBounds = OutLocalBounds.TransformBy(InLocalToWorld);
+	}
+}
+
+void FUAFOutputPoseComponent::WritePose(const FAnimNextGraphLODPose& InPose)
+{
+	SerialNumber++;
+
+	// If serial number wrapped, skip zero so client serials that have just been created will always mismatch
+	if (SerialNumber == 0)
+	{
+		SerialNumber++;
+	}
+
+	Pose = InPose;
+
+	if (SignalFunction.IsSet())
+	{
+		// Signal on the GT that the pose was written
+		GetModuleInstance().QueueOutputTraitEvent(
+			MakeTraitEvent<FAnimNextModule_ActionEvent>([Function = SignalFunction]()
+			{
+				Function();
+			})
+		);
+	}
+}
+
+bool FUAFOutputPoseComponent::GenerateMeshRenderData(
+	uint32 InSerialNumber,
+	int32 InMeshLOD,
+	const FTransform& InLocalToWorld,
+	TFunctionRef<void(TConstArrayView<FMatrix44f> InReferenceToLocalMatrices, const FBoxSphereBounds& InWorldBounds, const FBoxSphereBounds& InLocalBounds, uint32 InSerialNumber)> InFunction) const
+{
+	// Zero serial indicates no valid pose to output, so early out
+	if (SerialNumber == 0)
+	{
+		return false;
+	}
+
+	// Always call the read function when the supplied serial is zero, otherwise only call it on mismatched serial numbers
+	if (InSerialNumber != 0 && InSerialNumber == SerialNumber)
+	{
+		return false;
+	}
+
+	if (!ensure(Mesh != nullptr))
+	{
+		return false;
+	}
+
+	FMemMark MemMark(FMemStack::Get());
+
+	// Build render matrices from pose
+	const UE::UAF::FReferencePose& ReferencePose = Pose.LODPose.GetRefPose();
+	const TArray<FMatrix44f>& RefBasesInvMatrix = Mesh->GetRefBasesInvMatrix();
+
+	const TArrayView<const FBoneIndexType> MeshBoneIndexToParentMeshBoneIndexMap = ReferencePose.GetMeshBoneIndexToParentMeshBoneIndexMap();
+	const TArrayView<const FBoneIndexType> LODBoneIndexToMeshBoneIndexMap = ReferencePose.GetLODBoneIndexToMeshBoneIndexMap(Pose.LODPose.LODLevel);
+
+	checkf(LODBoneIndexToMeshBoneIndexMap.Num() == Pose.LODPose.GetNumBones(), TEXT("Buffer mismatch: %d:%d"), LODBoneIndexToMeshBoneIndexMap.Num(), Pose.LODPose.GetNumBones());
+
+	TArray<FTransform, TMemStackAllocator<>> ComponentSpacePose;
+	TArray<FMatrix44f, TMemStackAllocator<>> ReferenceToLocalMatrices;
+	ComponentSpacePose.SetNumUninitialized(MeshBoneIndexToParentMeshBoneIndexMap.Num());
+	ReferenceToLocalMatrices.SetNumUninitialized(MeshBoneIndexToParentMeshBoneIndexMap.Num());
+
+	const int32 NumPoseLODBones = Pose.LODPose.GetNumBones();
+	const FBoneIndexType* RESTRICT MeshBoneIndexToParentMeshBoneIndexMapPtr = MeshBoneIndexToParentMeshBoneIndexMap.GetData();
+	const FBoneIndexType* RESTRICT LODBoneIndexToMeshBoneIndexMapPtr = LODBoneIndexToMeshBoneIndexMap.GetData();
+	FTransform* RESTRICT ComponentSpaceData = ComponentSpacePose.GetData();
+
+	// First bone (if we have one) is always root bone, and it doesn't have a parent.
+	{
+		check(LODBoneIndexToMeshBoneIndexMap[0] == 0);
+		ComponentSpacePose[0] = Pose.LODPose.LocalTransformsView[0];
+		checkSlow(ComponentSpacePose[0].IsRotationNormalized());
+		checkSlow(!ComponentSpacePose[0].ContainsNaN());
+		const FMatrix44f ReferenceToLocal = (FMatrix44f)ComponentSpacePose[0].ToMatrixWithScale();
+		VectorMatrixMultiply(&ReferenceToLocalMatrices[0], &RefBasesInvMatrix[0], &ReferenceToLocal);
+	}
+
+	for (int32 LODBoneIndex = 1; LODBoneIndex < NumPoseLODBones; ++LODBoneIndex)
+	{
+		const FTransform LocalSpaceTransform = Pose.LODPose.LocalTransformsView[LODBoneIndex];
+
+		const FBoneIndexType MeshBoneIndex = LODBoneIndexToMeshBoneIndexMapPtr[LODBoneIndex];
+		const FBoneIndexType ParentMeshBoneIndex = MeshBoneIndexToParentMeshBoneIndexMapPtr[MeshBoneIndex];
+
+		const FTransform* RESTRICT ParentComponentSpaceTransform = ComponentSpaceData + ParentMeshBoneIndex;
+
+		FTransform* RESTRICT ComponentSpaceTransform = ComponentSpaceData + MeshBoneIndex;
+
+		FTransform::Multiply(ComponentSpaceTransform, &LocalSpaceTransform, ParentComponentSpaceTransform);
+
+		ComponentSpaceTransform->NormalizeRotation();
+
+		checkSlow(ComponentSpaceTransform->IsRotationNormalized());
+		checkSlow(!ComponentSpaceTransform->ContainsNaN());
+
+		const FMatrix44f ReferenceToLocal = (FMatrix44f)ComponentSpaceTransform->ToMatrixWithScale();
+		VectorMatrixMultiply(&ReferenceToLocalMatrices[MeshBoneIndex], &RefBasesInvMatrix[MeshBoneIndex], &ReferenceToLocal);
+	}
+
+	if (InMeshLOD < Pose.LODPose.LODLevel)
+	{
+		// If we evaluate the animation with lower quality than the visual mesh, we might
+		// have some transforms that we didn't compute. We'll grab the reference pose for them.
+		const TArrayView<const FBoneIndexType> LODBoneIndexToMeshBoneIndexMapForMeshLOD = ReferencePose.GetLODBoneIndexToMeshBoneIndexMap(InMeshLOD);
+
+		const int32 NumMeshLODBones = LODBoneIndexToMeshBoneIndexMapForMeshLOD.Num();
+		check(NumPoseLODBones <= NumMeshLODBones);
+
+		const FBoneIndexType* RESTRICT LODBoneIndexToMeshBoneIndexMapForMeshLODPtr = LODBoneIndexToMeshBoneIndexMapForMeshLOD.GetData();
+
+		for (int32 LODBoneIndex = NumPoseLODBones; LODBoneIndex < NumMeshLODBones; ++LODBoneIndex)
+		{
+			const FTransform LocalSpaceTransform = ReferencePose.GetRefPoseTransform(LODBoneIndex);
+
+			const FBoneIndexType MeshBoneIndex = LODBoneIndexToMeshBoneIndexMapForMeshLODPtr[LODBoneIndex];
+			const FBoneIndexType ParentMeshBoneIndex = MeshBoneIndexToParentMeshBoneIndexMapPtr[MeshBoneIndex];
+
+			const FTransform* RESTRICT ParentComponentSpaceTransform = ComponentSpaceData + ParentMeshBoneIndex;
+
+			FTransform* RESTRICT ComponentSpaceTransform = ComponentSpaceData + MeshBoneIndex;
+
+			FTransform::Multiply(ComponentSpaceTransform, &LocalSpaceTransform, ParentComponentSpaceTransform);
+
+			ComponentSpaceTransform->NormalizeRotation();
+
+			checkSlow(ComponentSpaceTransform->IsRotationNormalized());
+			checkSlow(!ComponentSpaceTransform->ContainsNaN());
+
+			const FMatrix44f ReferenceToLocal = (FMatrix44f)ComponentSpaceTransform->ToMatrixWithScale();
+			VectorMatrixMultiply(&ReferenceToLocalMatrices[MeshBoneIndex], &RefBasesInvMatrix[MeshBoneIndex], &ReferenceToLocal);
+		}
+	}
+
+	// Calc bounds
+	FBoxSphereBounds WorldBounds(ForceInit);
+	FBoxSphereBounds LocalBounds(ForceInit);
+	CalcMeshBound(ComponentSpacePose, InLocalToWorld, WorldBounds, LocalBounds);
+	InFunction(ReferenceToLocalMatrices, WorldBounds, LocalBounds, SerialNumber);
+	return true;
+}
+
+void FUAFOutputPoseComponent::BindToHost(USkeletalMesh* InMesh, TFunction<void(void)>&& InSignalFunction)
+{
+	check(InMesh);
+	Mesh = InMesh;
+	SignalFunction = MoveTemp(InSignalFunction);
+	RefPose.ReferencePose = UE::UAF::FDataRegistry::Get()->GetOrGenerateReferencePose(Mesh);
+}
\ No newline at end of file
diff --git a/Engine/Plugins/Experimental/UAF/UAF/Source/UAF/Private/Module/AnimNextActorComponentReferenceComponent.cpp b/Engine/Plugins/Experimental/UAF/UAF/Source/UAF/Private/Module/AnimNextActorComponentReferenceComponent.cpp
--- a/Engine/Plugins/Experimental/UAF/UAF/Source/UAF/Private/Module/AnimNextActorComponentReferenceComponent.cpp
+++ b/Engine/Plugins/Experimental/UAF/UAF/Source/UAF/Private/Module/AnimNextActorComponentReferenceComponent.cpp
@@ -38,7 +38,7 @@ void FAnimNextActorComponentReferenceComponent::OnInitializeHelper(UScriptStruct
 			UActorComponent* FoundComponent = Owner->FindComponentByClass(ComponentType);
 			Reference.QueueTask(NAME_None, [FoundComponent, InScriptStruct](const UE::UAF::FModuleTaskContext& InContext)
 			{
-				InContext.TryAccessComponent(InScriptStruct, [FoundComponent](FUAFModuleInstanceComponent& InComponent)
+				InContext.TryAccessComponent(InScriptStruct, [FoundComponent](FUAFAssetInstanceComponent& InComponent)
 				{
 					static_cast<FAnimNextActorComponentReferenceComponent&>(InComponent).Component = FoundComponent;
 				});
diff --git a/Engine/Plugins/Experimental/UAF/UAF/Source/UAF/Private/Module/AnimNextModuleInstance.cpp b/Engine/Plugins/Experimental/UAF/UAF/Source/UAF/Private/Module/AnimNextModuleInstance.cpp
--- a/Engine/Plugins/Experimental/UAF/UAF/Source/UAF/Private/Module/AnimNextModuleInstance.cpp
+++ b/Engine/Plugins/Experimental/UAF/UAF/Source/UAF/Private/Module/AnimNextModuleInstance.cpp
@@ -262,17 +262,6 @@ void FAnimNextModuleInstance::Initialize(const FUAFSystemFactoryParams& InFactor
 		{
 			SetEnabled(false);
 		}
-#if WITH_EDITOR
-		else if(World)
-		{
-			// In editor worlds we run a linearized 'initial tick' to ensure we generate an initial output pose, as these worlds dont always tick
-			if( World->WorldType == EWorldType::Editor ||
-				World->WorldType == EWorldType::EditorPreview)
-			{
-				FModuleEventTickFunction::InitializeAndRunModule(*this);
-			}
-		}
-#endif
 	}
 }
 
@@ -555,6 +544,31 @@ EPropertyBagResult FAnimNextModuleInstance::WriteProxyVariable(const FAnimNextVa
 	return Result;
 }
 
+EPropertyBagResult FAnimNextModuleInstance::ReadVariableFromExternal(const FAnimNextVariableReference& InVariable, const FAnimNextParamType& InType, TFunctionRef<void(TConstArrayView<uint8>)> InFunction)
+{
+	FReadScopeLock Lock(SystemReferenceLock);
+
+	TArrayView<uint8> Data;
+	EPropertyBagResult Result = Variables.AccessVariable(InVariable, InType, Data);
+	if (Result == EPropertyBagResult::Success)
+	{
+		InFunction(Data);
+	}
+	return Result;
+}
+
+bool FAnimNextModuleInstance::ReadComponentFromExternal(UScriptStruct* InComponentType, TFunctionRef<void(FConstStructView InComponentStruct)> InFunction)
+{
+	FReadScopeLock Lock(SystemReferenceLock);
+
+	if (const FUAFAssetInstanceComponent* Component = TryGetComponent(InComponentType))
+	{
+		InFunction(FConstStructView(InComponentType, (uint8*)Component));
+		return true;
+	}
+	return false;
+}
+
 void FAnimNextModuleInstance::RunScriptEvent(const UE::UAF::FScriptEvent& InEvent, float InDeltaTime)
 {
 	if (!InEvent.IsCallable())
diff --git a/Engine/Plugins/Experimental/UAF/UAF/Source/UAF/Private/Module/ModuleTaskContext.cpp b/Engine/Plugins/Experimental/UAF/UAF/Source/UAF/Private/Module/ModuleTaskContext.cpp
--- a/Engine/Plugins/Experimental/UAF/UAF/Source/UAF/Private/Module/ModuleTaskContext.cpp
+++ b/Engine/Plugins/Experimental/UAF/UAF/Source/UAF/Private/Module/ModuleTaskContext.cpp
@@ -18,9 +18,9 @@ void FModuleTaskContext::QueueInputTraitEvent(FAnimNextTraitEventPtr Event) cons
 	ModuleInstance.QueueInputTraitEvent(MoveTemp(Event));
 }
 
-void FModuleTaskContext::TryAccessComponent(UScriptStruct* InComponentType, TFunctionRef<void(FUAFModuleInstanceComponent&)> InFunction) const
+void FModuleTaskContext::TryAccessComponent(UScriptStruct* InComponentType, TFunctionRef<void(FUAFAssetInstanceComponent&)> InFunction) const
 {
-	FUAFModuleInstanceComponent* Component = static_cast<FUAFModuleInstanceComponent*>(ModuleInstance.TryGetComponent(InComponentType));
+	FUAFAssetInstanceComponent* Component = static_cast<FUAFAssetInstanceComponent*>(ModuleInstance.TryGetComponent(InComponentType));
 	if(Component == nullptr)
 	{
 		return;
diff --git a/Engine/Plugins/Experimental/UAF/UAF/Source/UAF/Private/Module/SystemReference.cpp b/Engine/Plugins/Experimental/UAF/UAF/Source/UAF/Private/Module/SystemReference.cpp
--- a/Engine/Plugins/Experimental/UAF/UAF/Source/UAF/Private/Module/SystemReference.cpp
+++ b/Engine/Plugins/Experimental/UAF/UAF/Source/UAF/Private/Module/SystemReference.cpp
@@ -436,6 +436,30 @@ void FSystemReference::RunEvent(FName InEventName, float InDeltaTime) const
 	TickFunction->Run(InDeltaTime);
 }
 
+#if WITH_EDITOR
+void FSystemReference::RunInitialTickInEditor() const
+{
+	if (IsValid())
+	{
+		// In editor worlds we run a linearized 'initial tick' to ensure we generate an initial output pose, as these worlds dont always tick
+		if( Ptr->GetWorldType() == EWorldType::Editor ||
+			Ptr->GetWorldType() == EWorldType::EditorPreview)
+		{
+			FModuleEventTickFunction::InitializeAndRunModule(*Ptr.Get());
+		}
+	}
+}
+#endif
+
+bool FSystemReference::ReadComponent(UScriptStruct* InComponentType, TFunctionRef<void(FConstStructView InComponentStruct)> InFunction) const
+{
+	if (IsValid())
+	{
+		return Ptr->ReadComponentFromExternal(InComponentType, InFunction);
+	}
+	return false;
+}
+
 } //end namespace UE::UAF
 
 UUAFEngineSubsystem::UUAFEngineSubsystem()
diff --git a/Engine/Plugins/Experimental/UAF/UAF/Source/UAF/Private/Module/UAFWeakSystemReference.cpp b/Engine/Plugins/Experimental/UAF/UAF/Source/UAF/Private/Module/UAFWeakSystemReference.cpp
--- a/Engine/Plugins/Experimental/UAF/UAF/Source/UAF/Private/Module/UAFWeakSystemReference.cpp
+++ b/Engine/Plugins/Experimental/UAF/UAF/Source/UAF/Private/Module/UAFWeakSystemReference.cpp
@@ -46,3 +46,13 @@ EPropertyBagResult FUAFWeakSystemReference::WriteVariable(const FAnimNextVariabl
 	}
 	return EPropertyBagResult::PropertyNotFound;
 }
+
+EPropertyBagResult FUAFWeakSystemReference::ReadVariable(const FAnimNextVariableReference& InVariable, const FAnimNextParamType& InType, TFunctionRef<void(TConstArrayView<uint8>)> InFunction) const
+{
+	TSharedPtr<FAnimNextModuleInstance> Instance = WeakPtr.Pin();
+	if (ensure(Instance.IsValid()))
+	{
+		return Instance->ReadVariableFromExternal(InVariable, InType, InFunction);
+	}
+	return EPropertyBagResult::PropertyNotFound;
+}
diff --git a/Engine/Plugins/Experimental/UAF/UAF/Source/UAF/Private/UAFAssetInstance.cpp b/Engine/Plugins/Experimental/UAF/UAF/Source/UAF/Private/UAFAssetInstance.cpp
--- a/Engine/Plugins/Experimental/UAF/UAF/Source/UAF/Private/UAFAssetInstance.cpp
+++ b/Engine/Plugins/Experimental/UAF/UAF/Source/UAF/Private/UAFAssetInstance.cpp
@@ -32,7 +32,6 @@ void FUAFAssetInstance::CopyDefaultComponents()
 
 void FUAFAssetInstance::BindDefaultComponents()
 {
-	Components = GetAsset<UAnimNextRigVMAsset>()->Components;
 	for (TInstancedStruct<FUAFAssetInstanceComponent>& InstancedStruct : Components)
 	{
 		FUAFAssetInstanceComponent& Component = InstancedStruct.GetMutable<FUAFAssetInstanceComponent>();
diff --git a/Engine/Plugins/Experimental/UAF/UAF/Source/UAF/Public/DataRegistry.h b/Engine/Plugins/Experimental/UAF/UAF/Source/UAF/Public/DataRegistry.h
--- a/Engine/Plugins/Experimental/UAF/UAF/Source/UAF/Public/DataRegistry.h
+++ b/Engine/Plugins/Experimental/UAF/UAF/Source/UAF/Public/DataRegistry.h
@@ -43,8 +43,16 @@ public:
 
 	// Removes a previously registered reference pose for the given SkeletalMeshComponent
 	UE_API void RemoveReferencePose(USkeletalMeshComponent* SkeletalMeshComponent);
-
 	
+	// Generates and registers a reference pose for a SkeletalMesh asset
+	UE_API FDataHandle RegisterReferencePose(USkeletalMesh* SkeletalMesh);
+
+	// Returns a ref counted handle to the reference pose of the given skeletal mesh
+	UE_API FDataHandle GetOrGenerateReferencePose(USkeletalMesh* SkeletalMesh);
+
+	// Removes a previously registered reference pose for the given SkeletalMesh
+	UE_API void RemoveReferencePose(USkeletalMesh* SkeletalMesh);
+
 	// --- AnimationData Storage / Retrieval  --- 
 
 	// Registers an anim data handle with arbitrary data using an FName
@@ -183,18 +191,23 @@ private:
 	TMap <FName, FDataHandle> StoredData;
 	mutable FRWLock StoredDataLock;
 
-	// Map holding reference poses for SkeletalMeshes
-	TMap <TObjectKey<USkeletalMeshComponent>, FReferencePoseData> SkeletalMeshReferencePoses;
+	// Map holding reference poses for Skeletal Mesh Components
+	TMap <TObjectKey<USkeletalMeshComponent>, FReferencePoseData> SkeletalMeshComponentReferencePoses;
+	mutable FRWLock SkeletalMeshComponentReferencePosesLock;
+
+	TMap <TObjectKey<USkeletalMesh>, FReferencePoseData> SkeletalMeshReferencePoses;
 	mutable FRWLock SkeletalMeshReferencePosesLock;
 
 	// A handle to our ticker when GC runs and we need to clean things up
 	FTSTicker::FDelegateHandle GCCleanupTickerHandle;
 
 	// Array of keys to clean up post-GC
-	TArray<TObjectKey<USkeletalMeshComponent>> GCKeysToCleanup;
+	TArray<TObjectKey<USkeletalMeshComponent>> GCComponentKeysToCleanup;
+	TArray<TObjectKey<USkeletalMesh>> GCMeshKeysToCleanup;
 
 	// The current index of GC keys to process
-	int32 GCKeysProcessedIndex = 0;
+	int32 GCComponentKeysProcessedIndex = 0;
+	int32 GCMeshKeysProcessedIndex = 0;
 
 	// Registers a type and sets the allocation block size
 	template<typename DataType>
diff --git a/Engine/Plugins/Experimental/UAF/UAF/Source/UAF/Public/Module/ModuleEvents.h b/Engine/Plugins/Experimental/UAF/UAF/Source/UAF/Public/Module/ModuleEvents.h
--- a/Engine/Plugins/Experimental/UAF/UAF/Source/UAF/Public/Module/ModuleEvents.h
+++ b/Engine/Plugins/Experimental/UAF/UAF/Source/UAF/Public/Module/ModuleEvents.h
@@ -14,13 +14,25 @@ struct FAnimNextModule_ActionEvent : public FAnimNextTraitEvent
 {
 	DECLARE_ANIM_TRAIT_EVENT(FAnimNextModule_ActionEvent, FAnimNextTraitEvent)
 
+	FAnimNextModule_ActionEvent() = default;
+
+	FAnimNextModule_ActionEvent(TUniqueFunction<void(void)>&& InActionFunction)
+		: ActionFunction(MoveTemp(InActionFunction))
+	{}
+
 	// Whether or not this event is thread-safe and can execute on any thread
 	// Events that are not thread-safe will execute on the main thread
 	virtual bool IsThreadSafe() const { return false; }
 
 	// Executes the schedule action
 	// Derived types can override this and implement whatever they wish instead of using a lambda
-	virtual void Execute() const { ActionFunction(); }
+	virtual void Execute() const
+	{
+		if (ActionFunction.IsSet())
+		{
+			ActionFunction();
+		}
+	}
 
 	// The optional action to execute
 	TUniqueFunction<void(void)> ActionFunction;
diff --git a/Engine/Plugins/Experimental/UAF/UAF/Source/UAF/Public/Module/ModuleTaskContext.h b/Engine/Plugins/Experimental/UAF/UAF/Source/UAF/Public/Module/ModuleTaskContext.h
--- a/Engine/Plugins/Experimental/UAF/UAF/Source/UAF/Public/Module/ModuleTaskContext.h
+++ b/Engine/Plugins/Experimental/UAF/UAF/Source/UAF/Public/Module/ModuleTaskContext.h
@@ -7,7 +7,7 @@
 
 #include "TraitCore/TraitEvent.h"
 
-struct FUAFModuleInstanceComponent;
+struct FUAFAssetInstanceComponent;
 struct FAnimNextModuleInstance;
 struct FInstancedPropertyBag;
 struct FAnimNextModuleInjectionComponent;
@@ -31,12 +31,12 @@ public:
 	UAF_API void QueueInputTraitEvent(FAnimNextTraitEventPtr Event) const;
 
 	// Access a module instance component of the specified type. If the component exists, then InFunction will be called
-	UAF_API void TryAccessComponent(UScriptStruct* InComponentType, TFunctionRef<void(FUAFModuleInstanceComponent&)> InFunction) const;
+	UAF_API void TryAccessComponent(UScriptStruct* InComponentType, TFunctionRef<void(FUAFAssetInstanceComponent&)> InFunction) const;
 
 	template<typename ComponentType>
 	void TryAccessComponent(TFunctionRef<void(ComponentType&)> InFunction) const
 	{
-		TryAccessComponent(TBaseStructure<ComponentType>::Get(), [&InFunction](FUAFModuleInstanceComponent& InComponent)
+		TryAccessComponent(TBaseStructure<ComponentType>::Get(), [&InFunction](FUAFAssetInstanceComponent& InComponent)
 		{
 			InFunction(static_cast<ComponentType&>(InComponent));
 		});
diff --git a/Engine/Plugins/Experimental/UAF/UAF/Source/UAF/Public/Module/SystemReference.h b/Engine/Plugins/Experimental/UAF/UAF/Source/UAF/Public/Module/SystemReference.h
--- a/Engine/Plugins/Experimental/UAF/UAF/Source/UAF/Public/Module/SystemReference.h
+++ b/Engine/Plugins/Experimental/UAF/UAF/Source/UAF/Public/Module/SystemReference.h
@@ -3,7 +3,9 @@
 #pragma once
 
 #include "EngineDefines.h"
+#include "UAFModuleInstanceComponent.h"
 #include "Module/TaskRunLocation.h"
+#include "StructUtils/StructView.h"
 #include "Subsystems/EngineSubsystem.h"
 #include "SystemReference.generated.h"
 
@@ -23,6 +25,8 @@ class UAnimNextRigVMAsset;
 struct FMassUAFFragment;
 class UMassUAFProcessor;
 class UMassUAFInitializer;
+class UMeshEntityComponentImpl_Skinned;
+class UEntitySkeletonModifierStackImpl;
 
 namespace UE::UAF
 {
@@ -47,20 +51,22 @@ namespace UE::UAF
 struct FSystemReference
 {
 	// Default-constructable only because it needs to be embedded in UObjects/UStructs
-	UE_API FSystemReference() = default;
+	FSystemReference() = default;
 
 	// Move only, non-copyable
 	FSystemReference(const FSystemReference&) = delete;
 	FSystemReference& operator=(const FSystemReference&) = delete;
 	UE_API FSystemReference(FSystemReference&&);
 	UE_API FSystemReference& operator=(FSystemReference&&);
-	
+
 private:
 	friend UAnimNextComponent;
 	friend FUAFWeakSystemReference;
 	friend UE::UAF::FModuleEventTickFunction;
 	friend UE::UAF::FAnimNextModuleImpl;
 	friend UE::UAF::Tests::FSystemReference_Creation;
+	friend UEntitySkeletonModifierStackImpl;
+	friend UMeshEntityComponentImpl_Skinned;
 
 	// @TODO: Temporarily friending Mass types, so that Mass can act as a system host similar to UAnimNextComponent.
 	friend FMassUAFFragment;
@@ -174,7 +180,29 @@ private:
 	// @return see EPropertyBagResult
 	UE_API EPropertyBagResult WriteVariable(const FAnimNextVariableReference& InVariable, const FAnimNextParamType& InType, TFunctionRef<void(TArrayView<uint8>)> InFunction) const;
 
+	// Access the supplied system component for reading.
+	// @param	InFunction			Function that will be called to access the struct
+	// @return true if the struct was available, otherwise false
+	template<typename ComponentType>
+	bool ReadComponent(TFunctionRef<void(TConstStructView<ComponentType> InComponentStruct)> InFunction) const
+	{
+		static_assert(std::is_base_of_v<FUAFModuleInstanceComponent, ComponentType>, "Component struct must be a child of FUAFModuleInstanceComponent");
+		return ReadComponent(ComponentType::StaticStruct(), [&InFunction](FConstStructView InStructView)
+			{
+				InFunction(TConstStructView<ComponentType>(InStructView.Get<ComponentType>()));
+			});
+	}
+
+#if WITH_EDITOR
+	// Runs a linearized version of the initial tick of the system to support generating an output pose for editor worlds that dont tick
+	// Should be called only from the game thread
+	UE_API void RunInitialTickInEditor() const;
+#endif
+
 private:
+	// Access the supplied system component for reading.
+	UE_API bool ReadComponent(UScriptStruct* InComponentType, TFunctionRef<void(FConstStructView InComponentStruct)> InFunction) const;
+
 	// Internal only constructor for selective access to owned-ref APIs
 	explicit FSystemReference(FAnimNextModuleInstance& InInstance);
 
diff --git a/Engine/Plugins/Experimental/UAF/UAF/Source/UAF/Public/Module/UAFModuleInstanceComponent.h b/Engine/Plugins/Experimental/UAF/UAF/Source/UAF/Public/Module/UAFModuleInstanceComponent.h
--- a/Engine/Plugins/Experimental/UAF/UAF/Source/UAF/Public/Module/UAFModuleInstanceComponent.h
+++ b/Engine/Plugins/Experimental/UAF/UAF/Source/UAF/Public/Module/UAFModuleInstanceComponent.h
@@ -12,7 +12,7 @@ struct FAnimNextTraitEvent;
 
 #define UE_API UAF_API
 
-/** A module instance component is attached and owned by a module instance. */
+/** A system instance component is attached to and owned by a system instance. */
 USTRUCT(meta=(Hidden, Abstract))
 struct FUAFModuleInstanceComponent : public FUAFAssetInstanceComponent
 {
@@ -23,22 +23,22 @@ struct FUAFModuleInstanceComponent : public FUAFAssetInstanceComponent
 
 	virtual ~FUAFModuleInstanceComponent() override = default;
 
-	// Returns the owning module instance this component lives on.
-	// Can return nullptr if an owning module instance has not been created (such as a default component on an asset)
+	// Returns the owning system instance this component lives on.
+	// Can return nullptr if an owning system instance has not been created (such as a default component on an asset)
 	UE_API FAnimNextModuleInstance* GetModuleInstancePtr();
 
-	// Returns the owning module instance this component lives on
-	// Will assert if an owning module instance has not been created (such as a default component on an asset)
+	// Returns the owning system instance this component lives on
+	// Will assert if an owning system instance has not been created (such as a default component on an asset)
 	UE_API FAnimNextModuleInstance& GetModuleInstance();
 
-	// Returns the owning module instance this component lives on
-	// Will assert if an owning module instance has not been created (such as a default component on an asset)
+	// Returns the owning system instance this component lives on
+	// Will assert if an owning system instance has not been created (such as a default component on an asset)
 	UE_API const FAnimNextModuleInstance& GetModuleInstance() const;
 
-	// Called during module execution for any events to be handled
+	// Called during system execution for any events to be handled
 	virtual void OnTraitEvent(FAnimNextTraitEvent& Event) {}
 
-	// Called at end of module execution each frame
+	// Called at end of system execution each frame
 	virtual void OnEndExecution(float InDeltaTime) {}
 
 private:
diff --git a/Engine/Plugins/Experimental/UAF/UAF/Source/UAF/Public/Module/UAFWeakSystemReference.h b/Engine/Plugins/Experimental/UAF/UAF/Source/UAF/Public/Module/UAFWeakSystemReference.h
--- a/Engine/Plugins/Experimental/UAF/UAF/Source/UAF/Public/Module/UAFWeakSystemReference.h
+++ b/Engine/Plugins/Experimental/UAF/UAF/Source/UAF/Public/Module/UAFWeakSystemReference.h
@@ -75,6 +75,16 @@ struct FUAFWeakSystemReference
 	// @return see EPropertyBagResult
 	UE_API EPropertyBagResult WriteVariable(const FAnimNextVariableReference& InVariable, const FAnimNextParamType& InType, TFunctionRef<void(TArrayView<uint8>)> InFunction) const;
 
+	// Access a variable's value for read only.
+	// This performs non-recursive locks on the system instance to ensure that concurrent access is respected, so care must be taken calling this from
+	// animation worker tasks as deadlocks can occur.
+	// Type must match strictly, no conversions are performed.
+	// @param	InVariable			The variable to get the value of
+	// @param	InType				The type of the variable
+	// @param	InFunction			Function that will be called if no errors occur
+	// @return see EPropertyBagResult
+	UE_API EPropertyBagResult ReadVariable(const FAnimNextVariableReference& InVariable, const FAnimNextParamType& InType, TFunctionRef<void(TConstArrayView<uint8>)> InFunction) const;
+
 private:
 	friend UE::UAF::FSystemReference;
 	friend FAnimNextModuleInstance;
diff --git a/Engine/Plugins/Experimental/UAF/UAF/Source/UAF/Public/UAFAssetInstance.h b/Engine/Plugins/Experimental/UAF/UAF/Source/UAF/Public/UAFAssetInstance.h
--- a/Engine/Plugins/Experimental/UAF/UAF/Source/UAF/Public/UAFAssetInstance.h
+++ b/Engine/Plugins/Experimental/UAF/UAF/Source/UAF/Public/UAFAssetInstance.h
@@ -28,6 +28,7 @@ namespace UE::UAF
 	struct FUAFTrace;
 	struct FInstanceTaskContext;
 	struct FVariableOverridesCollection;
+	struct FEntitySkeletalAnimationStackTrait;
 }
 
 // Base struct for instances of UAF assets
@@ -326,6 +327,7 @@ protected:
 	friend UE::UAF::FUAFTrace;
 	friend UE::UAF::FInstanceTaskContext;
 	friend FUAFRigVMComponent;
+	friend UE::UAF::FEntitySkeletalAnimationStackTrait;
 
 	void SetDebugName(FName Name)
 	{
diff --git a/Engine/Plugins/Experimental/UAF/UAFAnimGraph/Source/UAFAnimGraph/Private/Traits/SequencePlayer.cpp b/Engine/Plugins/Experimental/UAF/UAFAnimGraph/Source/UAFAnimGraph/Private/Traits/SequencePlayer.cpp
--- a/Engine/Plugins/Experimental/UAF/UAFAnimGraph/Source/UAFAnimGraph/Private/Traits/SequencePlayer.cpp
+++ b/Engine/Plugins/Experimental/UAF/UAFAnimGraph/Source/UAFAnimGraph/Private/Traits/SequencePlayer.cpp
@@ -39,6 +39,20 @@ namespace UE::UAF
 		}
 	}
 
+	const UAnimSequence* FSequencePlayerTrait::EvaluateAnimSequence(const FTraitBinding& InBinding) const
+	{
+		FInstanceData* InstanceData = InBinding.GetInstanceData<FInstanceData>();
+		if (!InstanceData->bCachedAnimSequence)
+		{
+			// Only cache this once for each constructed (relevant) sequence player.
+			// Usually this is when OnBecomeRelevant is called, but can be earlier for e.g. timeline queries prior to initial relevancy
+			const FSharedData* SharedData = InBinding.GetSharedData<FSharedData>();
+			InstanceData->AnimSequence = SharedData->GetAnimSequence(InBinding);
+			InstanceData->bCachedAnimSequence = true;
+		}
+		return InstanceData->AnimSequence;
+	}
+
 	void FSequencePlayerTrait::PreEvaluate(FEvaluateTraversalContext& Context, const TTraitBinding<IEvaluate>& Binding) const
 	{
 		const FSharedData* SharedData = Binding.GetSharedData<FSharedData>();
@@ -61,7 +75,7 @@ namespace UE::UAF
 	{
 		const FInstanceData* InstanceData = Binding.GetInstanceData<FInstanceData>();
 
-		if (InstanceData->AnimSequence)
+		if (const UAnimSequence* AnimSeq = EvaluateAnimSequence(Binding))
 		{
 			TTraitBinding<ITimeline> TimelineTrait;
 			Binding.GetStackInterface(TimelineTrait);
@@ -70,7 +84,7 @@ namespace UE::UAF
 			const float PlayRate = State.GetPlayRate();
 			const bool bIsLooping = State.IsLooping();
 
-			auto ExtractRootMotionAttribute = [AnimSequence = InstanceData->AnimSequence, PlayRate, bIsLooping](float StartTime, float DeltaTime, bool bAllowLooping)
+			auto ExtractRootMotionAttribute = [AnimSequence = AnimSeq, PlayRate, bIsLooping](float StartTime, float DeltaTime, bool bAllowLooping)
 			{
 				// We do not check for lifetimes, assume the sequence is alive during pose list execution.
 				check(AnimSequence->IsValidLowLevel());
@@ -93,7 +107,7 @@ namespace UE::UAF
 	void FSequencePlayerTrait::GetSyncMarkers(const FExecutionContext& Context, const TTraitBinding<ITimeline>& Binding, FTimelineSyncMarkerArray& OutSyncMarkers) const
 	{
 		const FInstanceData* InstanceData = Binding.GetInstanceData<FInstanceData>();
-		if (const UAnimSequence* AnimSeq = InstanceData->AnimSequence.Get())
+		if (const UAnimSequence* AnimSeq = EvaluateAnimSequence(Binding))
 		{
 			OutSyncMarkers.Reserve(AnimSeq->AuthoredSyncMarkers.Num());
 
@@ -107,7 +121,7 @@ namespace UE::UAF
 	FTimelineState FSequencePlayerTrait::GetState(const FExecutionContext& Context, const TTraitBinding<ITimeline>& Binding) const
 	{
 		const FInstanceData* InstanceData = Binding.GetInstanceData<FInstanceData>();
-		if (const UAnimSequence* AnimSeq = InstanceData->AnimSequence.Get())
+		if (const UAnimSequence* AnimSeq = EvaluateAnimSequence(Binding))
 		{
 			const FSharedData* SharedData = Binding.GetSharedData<FSharedData>();
 
@@ -124,7 +138,7 @@ namespace UE::UAF
 	FTimelineDelta FSequencePlayerTrait::GetDelta(const FExecutionContext& Context, const TTraitBinding<ITimeline>& Binding) const
 	{
 		const FInstanceData* InstanceData = Binding.GetInstanceData<FInstanceData>();
-		if (const UAnimSequence* AnimSeq = InstanceData->AnimSequence.Get())
+		if (const UAnimSequence* AnimSeq = EvaluateAnimSequence(Binding))
 		{
 			return FTimelineDelta(
 				InstanceData->DeltaTimeRecord.Delta,
@@ -139,7 +153,7 @@ namespace UE::UAF
 		FInstanceData* InstanceData = Binding.GetInstanceData<FInstanceData>();
 
 		// We only advance if we have a valid anim sequence
-		if (const UAnimSequence* AnimSeq = InstanceData->AnimSequence.Get())
+		if (const UAnimSequence* AnimSeq = EvaluateAnimSequence(Binding))
 		{
 			TTraitBinding<ITimeline> TimelineTrait;
 			Binding.GetStackInterface(TimelineTrait);
@@ -165,7 +179,7 @@ namespace UE::UAF
 		FInstanceData* InstanceData = Binding.GetInstanceData<FInstanceData>();
 
 		// Cache the anim sequence we'll play during construction, we don't allow it to change afterwards
-		InstanceData->AnimSequence = SharedData->GetAnimSequence(Binding);
+		InstanceData->AnimSequence = EvaluateAnimSequence(Binding);
 
 		float InternalTimeAccumulator = 0.0f;
 		if (const UAnimSequence* AnimSeq = InstanceData->AnimSequence.Get())
@@ -199,7 +213,7 @@ namespace UE::UAF
 	void FSequencePlayerTrait::GetNotifies(FExecutionContext& Context, const TTraitBinding<INotifySource>& Binding, float StartPosition, float Duration, bool bLooping, TArray<FAnimNotifyEventReference>& OutNotifies) const
 	{
 		const FInstanceData* InstanceData = Binding.GetInstanceData<FInstanceData>();
-		if (const UAnimSequence* AnimSeq = InstanceData->AnimSequence.Get())
+		if (const UAnimSequence* AnimSeq = EvaluateAnimSequence(Binding))
 		{
 			FAnimTickRecord TickRecord;
 			TickRecord.TimeAccumulator = &StartPosition;
diff --git a/Engine/Plugins/Experimental/UAF/UAFAnimGraph/Source/UAFAnimGraph/Private/Traits/SequencePlayer.h b/Engine/Plugins/Experimental/UAF/UAFAnimGraph/Source/UAFAnimGraph/Private/Traits/SequencePlayer.h
--- a/Engine/Plugins/Experimental/UAF/UAFAnimGraph/Source/UAFAnimGraph/Private/Traits/SequencePlayer.h
+++ b/Engine/Plugins/Experimental/UAF/UAFAnimGraph/Source/UAFAnimGraph/Private/Traits/SequencePlayer.h
@@ -40,6 +40,9 @@ namespace UE::UAF
 
 			// The last advance type when FAnimationRuntime::AdvanceTime was called
 			ETypeAdvanceAnim LastAdvanceType = ETAA_Default;
+
+			// Whether we have cached our AnimSequence yet
+			bool bCachedAnimSequence = false;
 		};
 
 		// IEvaluate impl
@@ -65,5 +68,7 @@ namespace UE::UAF
 
 		// INotifySource impl
 		virtual void GetNotifies(FExecutionContext& Context, const TTraitBinding<INotifySource>& Binding, float StartPosition, float Duration, bool bLooping, TArray<FAnimNotifyEventReference>& OutNotifies) const override;
+
+		const UAnimSequence* EvaluateAnimSequence(const FTraitBinding& InBinding) const;
 	};
 }
diff --git a/Engine/Plugins/Experimental/UAF/UAFAnimGraph/Source/UAFAnimGraph/Public/Graph/AnimNextAnimGraph.h b/Engine/Plugins/Experimental/UAF/UAFAnimGraph/Source/UAFAnimGraph/Public/Graph/AnimNextAnimGraph.h
--- a/Engine/Plugins/Experimental/UAF/UAFAnimGraph/Source/UAFAnimGraph/Public/Graph/AnimNextAnimGraph.h
+++ b/Engine/Plugins/Experimental/UAF/UAFAnimGraph/Source/UAFAnimGraph/Public/Graph/AnimNextAnimGraph.h
@@ -24,6 +24,7 @@ struct FAnimNextModuleAnimGraphComponent;
 struct FStructView;
 struct FAnimNextModuleInjectionComponent;
 struct FRigUnit_GetPostProcessAnimation;
+struct FEntitySkeletalAnimationStackComponent;
 
 namespace UE::UAF
 {
@@ -105,6 +106,7 @@ private:
 	friend UE::UAF::Editor::FAnimNextGraphDetails;
 	friend FRigUnit_GetPostProcessAnimation;
 	friend UE::UAF::AnimGraph::FAnimNextAnimGraphModule;
+	friend FEntitySkeletalAnimationStackComponent;
 };
 
 #undef UE_API
diff --git a/Engine/Plugins/Experimental/UAF/UAFAnimGraph/Source/UAFAnimGraph/Public/TraitInterfaces/IUpdate.h b/Engine/Plugins/Experimental/UAF/UAFAnimGraph/Source/UAFAnimGraph/Public/TraitInterfaces/IUpdate.h
--- a/Engine/Plugins/Experimental/UAF/UAFAnimGraph/Source/UAFAnimGraph/Public/TraitInterfaces/IUpdate.h
+++ b/Engine/Plugins/Experimental/UAF/UAFAnimGraph/Source/UAFAnimGraph/Public/TraitInterfaces/IUpdate.h
@@ -180,11 +180,11 @@ namespace UE::UAF
 		// Raises an input trait event
 		// The raised event will be seen by every child of the currently updating trait stack
 		// Input trait events can only be raised during PreUpdate
-		virtual void RaiseInputTraitEvent(FAnimNextTraitEventPtr Event) override;
+		UE_API virtual void RaiseInputTraitEvent(FAnimNextTraitEventPtr Event) override;
 
 		// Raises an output trait event
 		// The raised event will be seen by every parent of the currently updating trait stack
-		virtual void RaiseOutputTraitEvent(FAnimNextTraitEventPtr Event) override;
+		UE_API virtual void RaiseOutputTraitEvent(FAnimNextTraitEventPtr Event) override;
 
 	private:
 		// Constructs a new traversal context
