UAF: Add open in new tab functionality for UAF functions

#Synced-CL 48387151

diff --git a/Engine/Plugins/Experimental/UAF/UAF/Source/UAFEditor/Private/Common/Outliner/OutlinerCommands.cpp b/Engine/Plugins/Experimental/UAF/UAF/Source/UAFEditor/Private/Common/Outliner/OutlinerCommands.cpp
--- a/Engine/Plugins/Experimental/UAF/UAF/Source/UAFEditor/Private/Common/Outliner/OutlinerCommands.cpp
+++ b/Engine/Plugins/Experimental/UAF/UAF/Source/UAFEditor/Private/Common/Outliner/OutlinerCommands.cpp
@@ -14,8 +14,10 @@ void FOutlinerCommands::RegisterCommands()
 	UI_COMMAND(FindReferencesInWorkspace, "In Workspace", "Finds all references to the selected items across assets in the current workspace.", EUserInterfaceActionType::Button, FInputChord());
 	UI_COMMAND(FindReferencesInAsset, "In Asset", "Finds all references to the selected items in the currently opened asset.", EUserInterfaceActionType::Button, FInputChord());	
 
-	UI_COMMAND(MakeEntryPublic, "Sets Entry to to Public", "Sets Entry to be Public.", EUserInterfaceActionType::Button, FInputChord());
+	UI_COMMAND(MakeEntryPublic, "Sets Entry to Public", "Sets Entry to be Public.", EUserInterfaceActionType::Button, FInputChord());
 	UI_COMMAND(MakeEntryPrivate, "Sets Entry to Private", "Sets Entry to be Private.", EUserInterfaceActionType::Button, FInputChord());
+
+	UI_COMMAND(OpenInNewTab, "Open in New Tab", "Opens the selected item in a new tab.", EUserInterfaceActionType::Button, FInputChord());
 }
 }
 
diff --git a/Engine/Plugins/Experimental/UAF/UAF/Source/UAFEditor/Private/Common/Outliner/OutlinerCommands.h b/Engine/Plugins/Experimental/UAF/UAF/Source/UAFEditor/Private/Common/Outliner/OutlinerCommands.h
--- a/Engine/Plugins/Experimental/UAF/UAF/Source/UAFEditor/Private/Common/Outliner/OutlinerCommands.h
+++ b/Engine/Plugins/Experimental/UAF/UAF/Source/UAFEditor/Private/Common/Outliner/OutlinerCommands.h
@@ -25,6 +25,7 @@ public:
 	TSharedPtr<FUICommandInfo> FindReferencesInAsset;
 	TSharedPtr<FUICommandInfo> MakeEntryPublic;
 	TSharedPtr<FUICommandInfo> MakeEntryPrivate;
+	TSharedPtr<FUICommandInfo> OpenInNewTab;
 };
 }
 
diff --git a/Engine/Plugins/Experimental/UAF/UAF/Source/UAFEditor/Private/Functions/Outliner/FunctionsOutlinerMode.cpp b/Engine/Plugins/Experimental/UAF/UAF/Source/UAFEditor/Private/Functions/Outliner/FunctionsOutlinerMode.cpp
--- a/Engine/Plugins/Experimental/UAF/UAF/Source/UAFEditor/Private/Functions/Outliner/FunctionsOutlinerMode.cpp
+++ b/Engine/Plugins/Experimental/UAF/UAF/Source/UAFEditor/Private/Functions/Outliner/FunctionsOutlinerMode.cpp
@@ -423,6 +423,12 @@ void FFunctionsOutlinerMode::RegisterToolMenus()
 		FToolMenuSection& FunctionsSection = Menu->AddDynamicSection("Functions", FNewToolMenuDelegate::CreateLambda([](UToolMenu* InMenu)
 		{
 			PopulateNewFunctionToolMenuEntries(InMenu, false);
+
+			if (const UCommonOutlinerItemMenuContext* OutlinerContext = InMenu->FindContext<UCommonOutlinerItemMenuContext>())
+			{
+				FToolMenuSection& Section = InMenu->AddSection("Navigation", LOCTEXT("FunctionNavigationLabel", "Navigation"), FToolMenuInsert(FCommonOutlinerMode::ContextToolEntrySectionName, EToolMenuInsertType::Before));
+				Section.AddMenuEntryWithCommandList(FOutlinerCommands::Get().OpenInNewTab, OutlinerContext->WeakCommandList.Pin());
+			}
 		}));
 	}
 	
@@ -487,40 +493,84 @@ void FFunctionsOutlinerMode::PopulateNewFunctionToolMenuEntries(UToolMenu* InMen
 
 void FFunctionsOutlinerMode::BindCommands(const TSharedRef<FUICommandList>& OutCommandList)
 {
+	CommandList->MapAction(
+		FOutlinerCommands::Get().OpenInNewTab,
+		FExecuteAction::CreateRaw(this, &FFunctionsOutlinerMode::OpenSelectedFunctionsInNewTab),
+		FCanExecuteAction::CreateRaw(this, &FFunctionsOutlinerMode::CanOpenSelectedFunctionsInNewTab),
+		FIsActionChecked(),
+		FIsActionButtonVisible::CreateRaw(this, &FFunctionsOutlinerMode::CanOpenSelectedFunctionsInNewTab));
+
 	FCommonOutlinerMode::BindCommands(OutCommandList);
 }
 
 void FFunctionsOutlinerMode::OnItemDoubleClick(FSceneOutlinerTreeItemPtr Item)
+{
+	if (OpenFunctionItem(Item))
+	{
+		return;
+	}
+
+	FCommonOutlinerMode::OnItemDoubleClick(Item);
+}
+
+bool FFunctionsOutlinerMode::CanRenameItem(const ISceneOutlinerTreeItem& Item) const
+{
+	return Item.IsA<FFunctionsOutlinerEntryItem>() || Item.IsA<FFunctionsOutlinerCategoryItem>() || Item.IsA<FFunctionsOutlinerCollapsedGraphItem>();
+}
+
+bool FFunctionsOutlinerMode::UpdateOperationDecorator(const FDragDropEvent& Event, const FSceneOutlinerDragValidationInfo* ValidationInfo) const
+{
+	TSharedPtr<FFunctionDragDropOp> GraphDropOp = Event.GetOperationAs<FFunctionDragDropOp>();
+	if (GraphDropOp.IsValid())
+	{
+		if (ValidationInfo)
+		{
+			const FSlateBrush* ValidHoverIcon = FAppStyle::GetBrush(TEXT("Graph.ConnectorFeedback.OK"));
+			const FSlateBrush* InvalidSymbol = FAppStyle::GetBrush(TEXT("Graph.ConnectorFeedback.Error"));	
+			GraphDropOp->SetSimpleFeedbackMessage(ValidationInfo->IsValid() ? ValidHoverIcon : InvalidSymbol, FLinearColor::White, ValidationInfo->ValidationText);
+		}
+		else
+		{
+			GraphDropOp->HoverTargetChanged();
+		}
+		
+		return true;
+	}
+	
+	return false;
+}
+
+bool FFunctionsOutlinerMode::OpenFunctionItem(const FSceneOutlinerTreeItemPtr Item, FDocumentTracker::EOpenDocumentCause OpenDocumentCause) const
 {
 	if (TSharedPtr<Workspace::IWorkspaceEditor> WorkspaceEditor = WeakWorkspaceEditor.Pin())
-	{		
+	{
 		const Workspace::FWorkspaceDocument& Document = WorkspaceEditor->GetFocusedWorkspaceDocument();
 		FWorkspaceOutlinerItemExport DocumentExport = Document.Export;
 		if (FFunctionsOutlinerEntryItem* FunctionItem = Item->CastTo<FFunctionsOutlinerEntryItem>())
 		{
 			if (FunctionItem->WeakLibraryNode.IsValid())
-			{			
+			{
 				const URigVMLibraryNode* FunctionNode = FunctionItem->WeakLibraryNode.Get();
 				UAnimNextRigVMAsset* Asset = FunctionNode->GetTypedOuter<UAnimNextRigVMAsset>();
 				UAnimNextRigVMAssetEditorData* EditorData = UncookedOnly::FUtils::GetEditorData(Asset);
-				
-				const FWorkspaceOutlinerItemExport RelativeAssetExport = DocumentExport.GetRelativeAssetExport(Asset);				
-				const FWorkspaceOutlinerItemExport AssetExport = UncookedOnly::FUtils::MakeAssetReferenceExport(Asset, RelativeAssetExport);				
+
+				const FWorkspaceOutlinerItemExport RelativeAssetExport = DocumentExport.GetRelativeAssetExport(Asset);
+				const FWorkspaceOutlinerItemExport AssetExport = UncookedOnly::FUtils::MakeAssetReferenceExport(Asset, RelativeAssetExport);
 				if (AssetExport.GetFirstAssetPath().IsValid() && RelativeAssetExport != AssetExport)
 				{
-					const FWorkspaceOutlinerItemExport FunctionsGroupExport("FunctionsGroup", AssetExport);						
-					WorkspaceEditor->OpenExports({UncookedOnly::FUtils::MakeFunctionExport(FunctionNode, FunctionsGroupExport)});	
-					return;
+					const FWorkspaceOutlinerItemExport FunctionsGroupExport("FunctionsGroup", AssetExport);
+					WorkspaceEditor->OpenExports({UncookedOnly::FUtils::MakeFunctionExport(FunctionNode, FunctionsGroupExport)}, OpenDocumentCause);
+					return true;
 				}
-								
+
 				if (URigVMGraph* ContainedGraph = FunctionNode->GetContainedGraph())
-				{					
+				{
 					if (IRigVMClientHost* Host = FunctionNode->GetImplementingOuter<IRigVMClientHost>())
 					{
 						if (URigVMEdGraph* ContainedEdGraph = Cast<URigVMEdGraph>(Host->GetEditorObjectForRigVMGraph(ContainedGraph)))
-						{							
-							WorkspaceEditor->OpenObjects({ContainedEdGraph});
-							return;
+						{
+							WorkspaceEditor->OpenObjects({ContainedEdGraph}, OpenDocumentCause);
+							return true;
 						}
 					}
 				}
@@ -539,8 +589,8 @@ void FFunctionsOutlinerMode::OnItemDoubleClick(FSceneOutlinerTreeItemPtr Item)
 				{
 					const FWorkspaceOutlinerItemExport FunctionsExport = UncookedOnly::FUtils::MakeFunctionExport(FunctionNode, AssetExport);
 					const FWorkspaceOutlinerItemExport CollapsedGraphExport = UncookedOnly::FUtils::MakeCollapsedGraphExport(CollapsedGraphItem->WeakCollapseNode.Get(), FunctionsExport);
-					WorkspaceEditor->OpenExports({CollapsedGraphExport});
-					return;
+					WorkspaceEditor->OpenExports({CollapsedGraphExport}, OpenDocumentCause);
+					return true;
 				}
 
 				if (URigVMGraph* ContainedGraph = CollapsedGraphItem->WeakCollapseNode->GetContainedGraph())
@@ -549,44 +599,45 @@ void FFunctionsOutlinerMode::OnItemDoubleClick(FSceneOutlinerTreeItemPtr Item)
 					{
 						if (URigVMEdGraph* ContainedEdGraph = Cast<URigVMEdGraph>(Host->GetEditorObjectForRigVMGraph(ContainedGraph)))
 						{
-							WorkspaceEditor->OpenObjects({ContainedEdGraph});
-							return;
+							WorkspaceEditor->OpenObjects({ContainedEdGraph}, OpenDocumentCause);
+							return true;
 						}
 					}
 				}
 			}
 		}
 	}
-	
-	FCommonOutlinerMode::OnItemDoubleClick(Item);
+
+	return false;
 }
 
-bool FFunctionsOutlinerMode::CanRenameItem(const ISceneOutlinerTreeItem& Item) const
+void FFunctionsOutlinerMode::OpenSelectedFunctionsInNewTab() const
 {
-	return Item.IsA<FFunctionsOutlinerEntryItem>() || Item.IsA<FFunctionsOutlinerCategoryItem>() || Item.IsA<FFunctionsOutlinerCollapsedGraphItem>();
+	TArray<FSceneOutlinerTreeItemPtr> SelectedItems = SceneOutliner->GetSelectedItems();
+	for (const FSceneOutlinerTreeItemPtr& Item : SelectedItems)
+	{
+		OpenFunctionItem(Item, FDocumentTracker::ForceOpenNewDocument);
+	}
 }
 
-bool FFunctionsOutlinerMode::UpdateOperationDecorator(const FDragDropEvent& Event, const FSceneOutlinerDragValidationInfo* ValidationInfo) const
+bool FFunctionsOutlinerMode::CanOpenSelectedFunctionsInNewTab() const
 {
-	TSharedPtr<FFunctionDragDropOp> GraphDropOp = Event.GetOperationAs<FFunctionDragDropOp>();
-	if (GraphDropOp.IsValid())
+	TArray<FSceneOutlinerTreeItemPtr> SelectedItems = SceneOutliner->GetSelectedItems();
+	for (const FSceneOutlinerTreeItemPtr& Item : SelectedItems)
 	{
-		if (ValidationInfo)
+		if (const FFunctionsOutlinerEntryItem* FunctionEntry = Item->CastTo<FFunctionsOutlinerEntryItem>())
 		{
-			const FSlateBrush* ValidHoverIcon = FAppStyle::GetBrush(TEXT("Graph.ConnectorFeedback.OK"));
-			const FSlateBrush* InvalidSymbol = FAppStyle::GetBrush(TEXT("Graph.ConnectorFeedback.Error"));	
-			GraphDropOp->SetSimpleFeedbackMessage(ValidationInfo->IsValid() ? ValidHoverIcon : InvalidSymbol, FLinearColor::White, ValidationInfo->ValidationText);
+			return true;
 		}
-		else
+		else if (const FFunctionsOutlinerCollapsedGraphItem* CategoryItem = Item->CastTo<FFunctionsOutlinerCollapsedGraphItem>())
 		{
-			GraphDropOp->HoverTargetChanged();
+			return true;
 		}
-		
-		return true;
 	}
-	
+
 	return false;
 }
+
 }
 
 #undef LOCTEXT_NAMESPACE // "FFunctionsOutlinerMode"
diff --git a/Engine/Plugins/Experimental/UAF/UAF/Source/UAFEditor/Private/Functions/Outliner/FunctionsOutlinerMode.h b/Engine/Plugins/Experimental/UAF/UAF/Source/UAFEditor/Private/Functions/Outliner/FunctionsOutlinerMode.h
--- a/Engine/Plugins/Experimental/UAF/UAF/Source/UAFEditor/Private/Functions/Outliner/FunctionsOutlinerMode.h
+++ b/Engine/Plugins/Experimental/UAF/UAF/Source/UAFEditor/Private/Functions/Outliner/FunctionsOutlinerMode.h
@@ -43,5 +43,11 @@ protected:
 	// End ISceneOutlinerMode overrides
 	
 	static void PopulateNewFunctionToolMenuEntries(UToolMenu* InMenu, bool bAddSeparator);	
+
+private: 
+	bool OpenFunctionItem(const FSceneOutlinerTreeItemPtr Item, FDocumentTracker::EOpenDocumentCause OpenDocumentCause = FDocumentTracker::NavigatingCurrentDocument) const;
+	void OpenSelectedFunctionsInNewTab() const;
+	bool CanOpenSelectedFunctionsInNewTab() const;
+
 };
 }
