Fix variables inside function and collapsed node graphs not updating variable nodes that have been moved

#Synced-CL 48401438

diff --git a/Engine/Plugins/Experimental/UAF/UAF/Source/UAFUncookedOnly/Internal/UncookedOnlyUtils.h b/Engine/Plugins/Experimental/UAF/UAF/Source/UAFUncookedOnly/Internal/UncookedOnlyUtils.h
--- a/Engine/Plugins/Experimental/UAF/UAF/Source/UAFUncookedOnly/Internal/UncookedOnlyUtils.h
+++ b/Engine/Plugins/Experimental/UAF/UAF/Source/UAFUncookedOnly/Internal/UncookedOnlyUtils.h
@@ -217,9 +217,10 @@ struct FUtils
 	static UE_API void SetVariableType(UAnimNextVariableEntry* VariableEntry, const FAnimNextParamType NewType, bool bSetupUndoRedo = true, bool bPrintPythonCommands = true);
 
 	// Replaces the provided variable reference with ReferenceToReplaceWith within the context of InEditorData
-	static UE_API void ReplaceVariableReferences(UAnimNextRigVMAssetEditorData* InEditorData, const FAnimNextSoftVariableReference& ReferenceToFind, const FAnimNextSoftVariableReference& ReferenceToReplaceWith, bool bSetupUndoRedo = true, bool bPrintPythonCommands = true);
+	static UE_API void ReplaceVariableReferences(UAnimNextRigVMAssetEditorData* InEditorData, const FAnimNextSoftVariableReference& ReferenceToFind, const FAnimNextSoftVariableReference& ReferenceToReplaceWith, bool bForceReplace = false, bool bSetupUndoRedo = true, bool bPrintPythonCommands = true);
+	
 	// Replaces the provided variable reference with ReferenceToReplaceWith within the context of the project
-	static UE_API void ReplaceVariableReferencesAcrossProject(const FAnimNextSoftVariableReference& ReferenceToFind, const FAnimNextSoftVariableReference& ReferenceToReplaceWith, bool bForceLoadAssets = false, bool bSetupUndoRedo = true, bool bPrintPythonCommands = true);
+	static UE_API void ReplaceVariableReferencesAcrossProject(const FAnimNextSoftVariableReference& ReferenceToFind, const FAnimNextSoftVariableReference& ReferenceToReplaceWith, bool bForceLoadAssets = false, bool bForceReplace = false, bool bSetupUndoRedo = true, bool bPrintPythonCommands = true);
 
 	// Generates a deterministic GUID provided a UStruct FProperty entry
 	static UE_API FGuid GenerateScriptStructPropertyGUID(const FProperty* Property);
diff --git a/Engine/Plugins/Experimental/UAF/UAF/Source/UAFUncookedOnly/Private/UncookedOnlyUtils.cpp b/Engine/Plugins/Experimental/UAF/UAF/Source/UAFUncookedOnly/Private/UncookedOnlyUtils.cpp
--- a/Engine/Plugins/Experimental/UAF/UAF/Source/UAFUncookedOnly/Private/UncookedOnlyUtils.cpp
+++ b/Engine/Plugins/Experimental/UAF/UAF/Source/UAFUncookedOnly/Private/UncookedOnlyUtils.cpp
@@ -932,21 +932,11 @@ FAnimNextParamType FUtils::FindVariableType(const FAnimNextVariableReference& In
 	}
 	else if (const UAnimNextRigVMAsset* Asset = Cast<UAnimNextRigVMAsset>(InVariableReference.GetObject()))
 	{
-		// Query the asset registry for an asset's variable exports
-		FAssetData AssetData(Asset);
-		FAnimNextAssetRegistryExports Exports;
-		if (GetExportedVariablesForAsset(AssetData, Exports))
+		const UAnimNextRigVMAssetEditorData* AssetEditorData = CastChecked<UAnimNextRigVMAssetEditorData>(Asset->EditorData);
+
+		if (const IAnimNextRigVMVariableInterface* Variable = Cast<IAnimNextRigVMVariableInterface>(AssetEditorData->FindEntry(InVariableReference.GetName())))
 		{
-			for(const FAnimNextExport& Export : Exports.Exports)
-			{
-				if(Export.Identifier == InVariableReference.GetName())
-				{
-					if (const FAnimNextVariableDeclarationData* VariableDeclaration = Export.Data.GetPtr<FAnimNextVariableDeclarationData>())
-					{
-						return VariableDeclaration->Type;
-					}
-				}
-			}
+			return Variable->GetType();
 		}
 	}
 
@@ -1144,6 +1134,11 @@ void FUtils::GetAssetVariableExports(const UAnimNextRigVMAssetEditorData* InEdit
 			}
 		}
 	}
+
+	for (const TObjectPtr<URigVMEdGraph>& FunctionGraph : InEditorData->FunctionEdGraphs)
+	{
+		GetSubGraphVariableExportsRecursive(FunctionGraph, OutExports);
+	}
 }
 
 void FUtils::GetStructVariableExports(const UScriptStruct* Struct, TSet<FAnimNextExport>& OutExports)
@@ -1486,6 +1481,19 @@ void FUtils::GetSubGraphVariableExportsRecursive(const URigVMEdGraph* RigVMEdGra
 				)
 			);
 		}
+		else if (URigVMCollapseNode* CollapseNode = Cast<URigVMCollapseNode>(EdNode->GetModelNode()))
+		{
+			if (URigVMGraph* ContainedGraph = CollapseNode->GetContainedGraph())
+			{
+				if (IRigVMClientHost* ClientHost = ContainedGraph->GetImplementingOuter<IRigVMClientHost>())
+				{
+					if (URigVMEdGraph* ContainerEdGraph = Cast<URigVMEdGraph>(ClientHost->GetEditorObjectForRigVMGraph(ContainedGraph)))
+					{
+						GetSubGraphVariableExportsRecursive(ContainerEdGraph, OutExports);
+					}
+				}
+			}
+		}
 	}
 }
 
@@ -1701,12 +1709,14 @@ UAnimNextVariableEntry* FUtils::MoveVariableToAsset(UAnimNextVariableEntry* Vari
 			if (bIsPublicVariable)
 			{
 				constexpr bool bForceLoadReferencingAssets = false;
-				ReplaceVariableReferencesAcrossProject(FindReference, ReplaceReference, bForceLoadReferencingAssets, bSetupUndoRedo, bPrintPythonCommands);
+				constexpr bool bForceReplace = true;
+				ReplaceVariableReferencesAcrossProject(FindReference, ReplaceReference, bForceLoadReferencingAssets, bForceReplace, bSetupUndoRedo, bPrintPythonCommands);
 			}
 			else
 			{
 				// Otherwise replace, potential, local references only
-				ReplaceVariableReferences(CurrentEditorData, FindReference, ReplaceReference, bSetupUndoRedo, bPrintPythonCommands);
+				constexpr bool bForceReplace = true;
+				ReplaceVariableReferences(CurrentEditorData, FindReference, ReplaceReference, bForceReplace, bSetupUndoRedo, bPrintPythonCommands);
 			}
 
 			CurrentEditorData->RemoveEntry(VariableEntry, bSetupUndoRedo, bPrintPythonCommands);
@@ -1747,7 +1757,7 @@ void FUtils::SetVariableType(UAnimNextVariableEntry* VariableEntry, const FAnimN
 	}	
 }
 
-void FUtils::ReplaceVariableReferences(UAnimNextRigVMAssetEditorData* InEditorData, const FAnimNextSoftVariableReference& SoftReferenceToFind, const FAnimNextSoftVariableReference& SoftReferenceToReplaceWith, bool bSetupUndoRedo /*= true*/, bool bPrintPythonCommands /*= true*/)
+void FUtils::ReplaceVariableReferences(UAnimNextRigVMAssetEditorData* InEditorData, const FAnimNextSoftVariableReference& SoftReferenceToFind, const FAnimNextSoftVariableReference& SoftReferenceToReplaceWith, bool bForceReplace /*= false*/, bool bSetupUndoRedo /*= true*/, bool bPrintPythonCommands /*= true*/)
 {
 	if (InEditorData == nullptr)
 	{
@@ -1829,130 +1839,163 @@ void FUtils::ReplaceVariableReferences(UAnimNextRigVMAssetEditorData* InEditorDa
 	};	
 	
 	TArray<URigVMEdGraphNode*> EdNodes;
-	for(const UAnimNextRigVMAssetEntry* Entry : InEditorData->Entries)
+
+	// Gather all the RigVM graphs to replace variables inside of
+	TArray<URigVMEdGraph*> RigVMEdGraphsToTraverse;
 	{
-		if(const IAnimNextRigVMGraphInterface* GraphInterface = Cast<IAnimNextRigVMGraphInterface>(Entry))
-		{			
-			if (URigVMEdGraph* RigVMEdGraph = GraphInterface->GetEdGraph())
+		for (const UAnimNextRigVMAssetEntry* Entry : InEditorData->Entries)
+		{
+			if (const IAnimNextRigVMGraphInterface* GraphInterface = Cast<IAnimNextRigVMGraphInterface>(Entry))
 			{
-				UAnimNextControllerBase* Controller = CastChecked<UAnimNextControllerBase>(RigVMEdGraph->GetController());
-				// Rely on RigVM controller to handle local variable references
-				if (bInternalReplace)
+				if (URigVMEdGraph* RigVMEdGraph = GraphInterface->GetEdGraph())
 				{
-					if (bReplacingWithNone)
+					RigVMEdGraphsToTraverse.Add(RigVMEdGraph);
+				}
+			}
+		}
+
+		RigVMEdGraphsToTraverse.Append(InEditorData->FunctionEdGraphs);
+	}
+
+	// Non-range-based loop as RigVMEdGraphsToTraverse can expand during iteration
+	for (int32 GraphToTraverseIndex = 0; GraphToTraverseIndex < RigVMEdGraphsToTraverse.Num(); ++GraphToTraverseIndex)
+	{
+		URigVMEdGraph* RigVMEdGraph = RigVMEdGraphsToTraverse[GraphToTraverseIndex];
+		if (!RigVMEdGraph)
+		{
+			continue;
+		}
+
+		UAnimNextControllerBase* Controller = CastChecked<UAnimNextControllerBase>(RigVMEdGraph->GetController());
+		// Rely on RigVM controller to handle local variable references
+		if (bInternalReplace)
+		{
+			if (bReplacingWithNone)
+			{
+				Controller->OnExternalVariableRemoved(ReferenceToFind.GetName(), bSetupUndoRedo);
+			}
+			else
+			{
+				if (bVariableTypeChanged)
+				{
+					const FRigVMTemplateArgumentType NewRigVMType = ReplaceReferenceType.ToRigVMTemplateArgument();
+					Controller->OnExternalVariableTypeChanged(ReferenceToFind.GetName(), NewRigVMType.CPPType.ToString(), NewRigVMType.CPPTypeObject.Get(), bSetupUndoRedo);
+				}
+
+				if (bVariableNameChanged)
+				{
+					Controller->OnExternalVariableRenamed(ReferenceToFind.GetName(), ReferenceToReplaceWith.GetName(), bSetupUndoRedo);
+				}
+			}
+		}
+		else
+		{
+			RigVMEdGraph->GetNodesOfClass(EdNodes);
+
+			for (URigVMEdGraphNode* EdNode : EdNodes)
+			{
+				if (URigVMVariableNode* VariableNode = Cast<URigVMVariableNode>(EdNode->GetModelNode()))
+				{
+					const bool bNameMatch = VariableNode->GetVariableName() == ReferenceToFind.GetName();
+
+					if (!bNameMatch)
 					{
-						Controller->OnExternalVariableRemoved(ReferenceToFind.GetName(), bSetupUndoRedo);
+						continue;
 					}
-					else
+
+					if (UAnimNextSharedVariableNode* SharedVariableNode = Cast<UAnimNextSharedVariableNode>(VariableNode))
 					{
-						if (bVariableTypeChanged)
-						{
-							const FRigVMTemplateArgumentType NewRigVMType = ReplaceReferenceType.ToRigVMTemplateArgument();
-							Controller->OnExternalVariableTypeChanged(ReferenceToFind.GetName(), NewRigVMType.CPPType.ToString(), NewRigVMType.CPPTypeObject.Get(), bSetupUndoRedo);
-						}
+						const UObject* SearchObject = ReferenceToFind.GetObject();
+						const bool bPropertySourceMatch = (SharedVariableNode->Asset && SharedVariableNode->Asset == SearchObject) || (SharedVariableNode->Struct && SharedVariableNode->Struct == SearchObject);
+						const bool bValidProperty =ReferenceToReplaceWith.ResolveProperty() != nullptr;
 
-						if (bVariableNameChanged)
+						if (bPropertySourceMatch && (bValidProperty || bReplacingWithNone || bForceReplace))
 						{
-							Controller->OnExternalVariableRenamed(ReferenceToFind.GetName(), ReferenceToReplaceWith.GetName(), bSetupUndoRedo);
+							if (bReplacingWithNone)
+							{
+								// [TODO] decide to either remove node or replace with None variable references (keeps node in place but will fail compilation until fixed up)
+								Controller->RemoveNode(SharedVariableNode, bSetupUndoRedo, bPrintPythonCommands);
+								//Controller->RefreshSharedVariableNode(SharedVariableNode->GetFName(), TEXT(""), ReferenceToReplaceWith.GetName(), TEXT(""), nullptr, true, true);	
+							}
+							else
+							{
+								if (bDifferentVariableSource && ReplaceOuterAsset == ThisOuterAsset)
+								{
+									Controller->ReplaceSharedVariableNodeWithVariableNode(SharedVariableNode, ReferenceToReplaceWith.GetName(), bSetupUndoRedo, bPrintPythonCommands);
+								}
+								else
+								{
+									const FRigVMTemplateArgumentType NewRigVMType = ReplaceReferenceType.ToRigVMTemplateArgument();
+
+									EnsureSharedVariableDependencyExists();
+
+									constexpr bool bSetupOrphanPins = true;
+									Controller->RefreshSharedVariableNode(SharedVariableNode->GetFName(), ReferenceToReplaceWith.GetObject().GetPath(), ReferenceToReplaceWith.GetName(), NewRigVMType.CPPType.ToString(), NewRigVMType.CPPTypeObject.Get(), bSetupUndoRedo, bSetupOrphanPins, bPrintPythonCommands);
+								}
+							}
 						}
 					}
-				}
-				else
-				{
-					RigVMEdGraph->GetNodesOfClass(EdNodes);
-
-					for (URigVMEdGraphNode* EdNode : EdNodes)
+					// For regular variable nodes, the FindOuterAsset has to match the asset being processed
+					else if (FindOuterAsset == ThisOuterAsset)
 					{
-						if (URigVMVariableNode* VariableNode = Cast<URigVMVariableNode>(EdNode->GetModelNode()))
+						if (bReplacingWithNone)
 						{
-							const bool bNameMatch = VariableNode->GetVariableName() == ReferenceToFind.GetName();
-
-							if (!bNameMatch)
-							{
-								continue;
-							}								
-							
-							if (UAnimNextSharedVariableNode* SharedVariableNode = Cast<UAnimNextSharedVariableNode>(VariableNode))
+							// Remove the node for now (can also leave an invalid one causing compile error/warning)
+							Controller->RemoveNode(VariableNode, bSetupUndoRedo, bPrintPythonCommands);
+						}
+						else
+						{
+							if (bDifferentVariableSource)
 							{
-								const UObject* SearchObject = ReferenceToFind.GetObject(); 
-								const bool bPropertySourceMatch = (SharedVariableNode->Asset && SharedVariableNode->Asset == SearchObject)  || (SharedVariableNode->Struct && SharedVariableNode->Struct == SearchObject);
-								const bool bValidProperty = ReferenceToReplaceWith.ResolveProperty() != nullptr;
-								
-								if (bPropertySourceMatch && (bValidProperty || bReplacingWithNone))
+								if (FindOuterAsset == ThisOuterAsset && ReplaceOuterAsset != ThisOuterAsset)
 								{
-									if (bReplacingWithNone)
-									{
-										// [TODO] decide to either remove node or replace with None variable references (keeps node in place but will fail compilation until fixed up)
-										Controller->RemoveNode(SharedVariableNode, bSetupUndoRedo, bPrintPythonCommands);
-										//Controller->RefreshSharedVariableNode(SharedVariableNode->GetFName(), TEXT(""), ReferenceToReplaceWith.GetName(), TEXT(""), nullptr, true, true);	
-									}
-									else
-									{
-										if (bDifferentVariableSource && ReplaceOuterAsset == ThisOuterAsset)
-										{
-											Controller->ReplaceSharedVariableNodeWithVariableNode(SharedVariableNode, ReferenceToReplaceWith.GetName(), bSetupUndoRedo, bPrintPythonCommands);
-										}
-										else
-										{
-											const FRigVMTemplateArgumentType NewRigVMType = ReplaceReferenceType.ToRigVMTemplateArgument();
-
-											EnsureSharedVariableDependencyExists();
-
-											constexpr bool bSetupOrphanPins = true;
-											Controller->RefreshSharedVariableNode(SharedVariableNode->GetFName(), ReferenceToReplaceWith.GetObject().GetPath(), ReferenceToReplaceWith.GetName(), NewRigVMType.CPPType.ToString(), NewRigVMType.CPPTypeObject.Get(), bSetupUndoRedo, bSetupOrphanPins, bPrintPythonCommands);											
-										}
-									}
+									EnsureSharedVariableDependencyExists();
+									Controller->ReplaceVariableNodeWithSharedVariableNode(VariableNode, ReferenceToReplaceWith.GetName(), ReferenceToReplaceWith.GetObject(), bSetupUndoRedo, bPrintPythonCommands);
+								}
+								else
+								{
+									UE_LOGFMT(LogAnimation, Error, "Found matching instance of Variable Type/Name ({0}/{1}) combination on {2} which differs from {3} from which the replace operation is being executed, this instance will be skipped.",
+										FindReferenceType.ToString(), SoftReferenceToFind.GetName().ToString(), FindOuterAsset->GetPathName(), ThisOuterAsset->GetPathName());
 								}
 							}
-							// For regular variable nodes, the FindOuterAsset has to match the asset being processed
-							else if (FindOuterAsset == ThisOuterAsset)
+							else
 							{
-								if (bReplacingWithNone)
+								if (bVariableTypeChanged)
 								{
-									// Remove the node for now (can also leave an invalid one causing compile error/warning)
-									Controller->RemoveNode(VariableNode, bSetupUndoRedo, bPrintPythonCommands);
+									const FRigVMTemplateArgumentType NewRigVMType = ReplaceReferenceType.ToRigVMTemplateArgument();
+									Controller->OnExternalVariableTypeChanged(ReferenceToFind.GetName(), NewRigVMType.CPPType.ToString(), NewRigVMType.CPPTypeObject.Get(), bSetupUndoRedo);
 								}
-								else
+
+								if (bVariableNameChanged)
 								{
-									if (bDifferentVariableSource)
-									{
-										if (FindOuterAsset == ThisOuterAsset && ReplaceOuterAsset != ThisOuterAsset)
-										{
-											EnsureSharedVariableDependencyExists();
-											Controller->ReplaceVariableNodeWithSharedVariableNode(VariableNode, ReferenceToReplaceWith.GetName(), ReferenceToReplaceWith.GetObject(), bSetupUndoRedo, bPrintPythonCommands);
-										}
-										else
-										{
-											UE_LOGFMT(LogAnimation, Error, "Found matching instance of Variable Type/Name ({0}/{1}) combination on {2} which differs from {3} from which the replace operation is being executed, this instance will be skipped.",
-											FindReferenceType.ToString(), SoftReferenceToFind.GetName().ToString(), FindOuterAsset->GetPathName(), ThisOuterAsset->GetPathName());
-										}
-									}
-									else
-									{					
-										if (bVariableTypeChanged)
-										{
-											const FRigVMTemplateArgumentType NewRigVMType = ReplaceReferenceType.ToRigVMTemplateArgument();
-											Controller->OnExternalVariableTypeChanged(ReferenceToFind.GetName(), NewRigVMType.CPPType.ToString(), NewRigVMType.CPPTypeObject.Get(), bSetupUndoRedo);
-										}
-
-										if (bVariableNameChanged)
-										{
-											Controller->OnExternalVariableRenamed(ReferenceToFind.GetName(), ReferenceToReplaceWith.GetName(), bSetupUndoRedo);
-										}
-									}
+									Controller->OnExternalVariableRenamed(ReferenceToFind.GetName(), ReferenceToReplaceWith.GetName(), bSetupUndoRedo);
 								}
 							}
 						}
 					}
 				}
+				else if (URigVMCollapseNode* CollapseNode = Cast<URigVMCollapseNode>(EdNode->GetModelNode()))
+				{
+					URigVMGraph* ContainedGraph = CollapseNode->GetContainedGraph();
+					if (IRigVMClientHost* ClientHost = ContainedGraph->GetImplementingOuter<IRigVMClientHost>())
+					{
+						if (URigVMEdGraph* ContainedEdGraph = Cast<URigVMEdGraph>(ClientHost->GetEditorObjectForRigVMGraph(ContainedGraph)))
+						{
+							if (!RigVMEdGraphsToTraverse.Contains(ContainedEdGraph))
+							{
+								RigVMEdGraphsToTraverse.Add(ContainedEdGraph);
+							}
+						}
+					}
+				}
 			}
-		}
 
-		EdNodes.Reset();
-	}	
+			EdNodes.Reset();
+		}
+	}
 }
 
-void FUtils::ReplaceVariableReferencesAcrossProject(const FAnimNextSoftVariableReference& ReferenceToFind, const FAnimNextSoftVariableReference& ReferenceToReplaceWith,  bool bForceLoadAssets, bool bSetupUndoRedo, bool bPrintPythonCommands)
+void FUtils::ReplaceVariableReferencesAcrossProject(const FAnimNextSoftVariableReference& ReferenceToFind, const FAnimNextSoftVariableReference& ReferenceToReplaceWith,  bool bForceLoadAssets, bool bForceReplace /*= false*/, bool bSetupUndoRedo, bool bPrintPythonCommands)
 {
 	if (ReferenceToFind.IsNone())
 	{
@@ -1998,7 +2041,7 @@ void FUtils::ReplaceVariableReferencesAcrossProject(const FAnimNextSoftVariableR
 
 	for (UAnimNextRigVMAssetEditorData* EditorData : EditorDataRequiringReplacing)
 	{
-		ReplaceVariableReferences(EditorData, ReferenceToFind, ReferenceToReplaceWith, bSetupUndoRedo, bPrintPythonCommands);
+		ReplaceVariableReferences(EditorData, ReferenceToFind, ReferenceToReplaceWith, bForceReplace, bSetupUndoRedo, bPrintPythonCommands);
 	}
 }
 
