Add compiler scope when moving a category of variables to a new asset to avoid recompiling after each individual variable move

#Synced-CL 48401438

diff --git a/Engine/Plugins/Experimental/UAF/UAF/Source/UAFEditor/Private/Variables/Outliner/VariablesOutlinerMode.cpp b/Engine/Plugins/Experimental/UAF/UAF/Source/UAFEditor/Private/Variables/Outliner/VariablesOutlinerMode.cpp
--- a/Engine/Plugins/Experimental/UAF/UAF/Source/UAFEditor/Private/Variables/Outliner/VariablesOutlinerMode.cpp
+++ b/Engine/Plugins/Experimental/UAF/UAF/Source/UAFEditor/Private/Variables/Outliner/VariablesOutlinerMode.cpp
@@ -641,6 +641,8 @@ void FVariablesOutlinerMode::MoveCategoryToAsset(const FString& CategoryPath, UA
 {
 	UAnimNextRigVMAssetEditorData* ThisEditorData = UncookedOnly::FUtils::GetEditorData(FromAsset);
 	UAnimNextRigVMAssetEditorData* EditorData = UncookedOnly::FUtils::GetEditorData(ToAsset);
+
+	UE::UAF::UncookedOnly::FCompilationScope CompilerResults(LOCTEXT("MoveCategoryToAsset", "Move Category To Asset"), MakeConstArrayView({ FromAsset, ToAsset }));
 	
 	TArray<FString> CategoriesToMove;
 	Algo::TransformIf(ThisEditorData->VariableAndFunctionCategories, CategoriesToMove, [&CategoryPath](const FString& CategoryName)
diff --git a/Engine/Plugins/Experimental/UAF/UAF/Source/UAFUncookedOnly/Private/AnimNextRigVMAssetEditorData.cpp b/Engine/Plugins/Experimental/UAF/UAF/Source/UAFUncookedOnly/Private/AnimNextRigVMAssetEditorData.cpp
--- a/Engine/Plugins/Experimental/UAF/UAF/Source/UAFUncookedOnly/Private/AnimNextRigVMAssetEditorData.cpp
+++ b/Engine/Plugins/Experimental/UAF/UAF/Source/UAFUncookedOnly/Private/AnimNextRigVMAssetEditorData.cpp
@@ -1077,6 +1077,8 @@ void UAnimNextRigVMAssetEditorData::RecompileVM()
 
 	TGuardValue<bool> CompilingGuard(bIsCompiling, true);
 
+	TRACE_CPUPROFILER_EVENT_SCOPE(UAnimNextRigVMAssetEditorData::RecompileVM);
+
 	UAnimNextRigVMAsset* Asset = FUtils::GetAsset<UAnimNextRigVMAsset>(this);
 	FCompilationScope CompileScope(Asset);
 
