Add compile warning when native shared variable property names overlap with other variable names

#Synced-CL 48453044

diff --git a/Engine/Plugins/Experimental/UAF/UAF/Source/UAFUncookedOnly/Private/AnimNextRigVMAssetEditorData.cpp b/Engine/Plugins/Experimental/UAF/UAF/Source/UAFUncookedOnly/Private/AnimNextRigVMAssetEditorData.cpp
--- a/Engine/Plugins/Experimental/UAF/UAF/Source/UAFUncookedOnly/Private/AnimNextRigVMAssetEditorData.cpp
+++ b/Engine/Plugins/Experimental/UAF/UAF/Source/UAFUncookedOnly/Private/AnimNextRigVMAssetEditorData.cpp
@@ -3194,6 +3194,13 @@ void UAnimNextRigVMAssetEditorData::SetUseExternalPackages(TArrayView<UAnimNextR
 
 FInstancedPropertyBag UAnimNextRigVMAssetEditorData::GenerateCombinedPropertyBag(const FRigVMCompileSettings& InSettings, const FAnimNextGetVariableCompileContext& InCompileContext) const
 {
+	// TODO:
+	// It would be nice to remove the requirement that variable names must be unique across all referenced/owned variables inside an asset.
+	// This stems from:
+	//   1. RigVM external variables needing to have unique names.
+	//   2. State tree needing to combine properties into a single property bag and property bag enforcing unique property names.
+	// A solution where variable names are qualified using their owned asset as a prefix should be investigated to remove friction when authoring assets with variable dependencies.
+
 	UAnimNextRigVMAsset* Asset = UE::UAF::UncookedOnly::FUtils::GetAsset<UAnimNextRigVMAsset>(this);
 
 	FInstancedPropertyBag CombinedPropertyBag;
@@ -3257,7 +3264,28 @@ FInstancedPropertyBag UAnimNextRigVMAssetEditorData::GenerateCombinedPropertyBag
 		TSharedRef<UE::UAF::FStructDataCache> StructData = UE::UAF::FStructDataCache::GetStructInfo(ReferencedStruct);
 		for (const UE::UAF::FStructDataCache::FPropertyInfo& PropertyInfo : StructData->GetProperties())
 		{
-			CombinedPropertyBag.AddProperty(PropertyInfo.Property->GetFName(), PropertyInfo.Property);
+			constexpr bool bOverrideProperty = false;
+			const EPropertyBagAlterationResult Result = CombinedPropertyBag.AddProperty(PropertyInfo.Property->GetFName(), PropertyInfo.Property, bOverrideProperty);
+
+			if (Result != EPropertyBagAlterationResult::Success)
+			{
+				if (Result == EPropertyBagAlterationResult::TargetPropertyAlreadyExists)
+				{
+					InCompileContext.GetAssetCompileContext().Error(
+						Asset,
+						LOCTEXT("AddNativeSharedVariablePropertyFailure_NameOverlaps", "The property for variable '{0}' inside Native Shared Variables '{1}' could not be generated as the name overlaps with another variable"),
+						FText::FromName(PropertyInfo.Property->GetFName()),
+						FText::FromName(ReferencedStruct->GetFName()));
+				}
+				else
+				{
+					InCompileContext.GetAssetCompileContext().Error(
+						Asset,
+						LOCTEXT("AddNativeSharedVariablePropertyFailure_Other", "The property for variable '{0}' inside Native Shared Variables '{1}' could not be generated"),
+						FText::FromName(PropertyInfo.Property->GetFName()),
+						FText::FromName(ReferencedStruct->GetFName()));
+				}
+			}
 		}
 	}
 
