Fix creating new shared variables asset from selection duplicating each variable entry

#Synced-CL 48503794

diff --git a/Engine/Plugins/Experimental/UAF/UAF/Source/UAFEditor/Private/Variables/Outliner/VariablesOutlinerMode.cpp b/Engine/Plugins/Experimental/UAF/UAF/Source/UAFEditor/Private/Variables/Outliner/VariablesOutlinerMode.cpp
--- a/Engine/Plugins/Experimental/UAF/UAF/Source/UAFEditor/Private/Variables/Outliner/VariablesOutlinerMode.cpp
+++ b/Engine/Plugins/Experimental/UAF/UAF/Source/UAFEditor/Private/Variables/Outliner/VariablesOutlinerMode.cpp
@@ -1856,14 +1856,6 @@ void FVariablesOutlinerMode::CreateSharedVariablesAssets() const
 		UAnimNextRigVMAsset* Asset = AssetItems[0]->SoftAsset.LoadSynchronous();
 		UniqueOuterEditorDatas.AddUnique(UncookedOnly::FUtils::GetEditorData(Asset));
 	}
-	
-	Selection.ForEachItem<FVariablesOutlinerEntryItem>([&EntriesToMove](const FVariablesOutlinerEntryItem& VariableItem)
-	{
-		if (UAnimNextVariableEntry* VariableEntry = VariableItem.WeakEntry.Get())
-		{
-			EntriesToMove.Add(VariableEntry);
-		}
-	});
 
 	// Prompt user for asset path
 	FString SaveObjectPath;
