UAF: Fix system factory referencing trashed classes after a Verse recompile

Change to using class path as a key, as that is invariant across compiles

#Synced-CL 48500753

diff --git a/Engine/Plugins/Experimental/UAF/UAF/Source/UAF/Private/Factory/SystemFactory.cpp b/Engine/Plugins/Experimental/UAF/UAF/Source/UAF/Private/Factory/SystemFactory.cpp
--- a/Engine/Plugins/Experimental/UAF/UAF/Source/UAF/Private/Factory/SystemFactory.cpp
+++ b/Engine/Plugins/Experimental/UAF/UAF/Source/UAF/Private/Factory/SystemFactory.cpp
@@ -17,7 +17,7 @@ static FDelegateHandle GSystemFactoryPreExitHandle;
 static TStripedMap<32, uint64, TWeakObjectPtr<const UAnimNextModule>> GSystemMap;
 
 // All registered object -> param mappings
-static TMap<TObjectKey<UClass>, TPair<FUAFSystemFactoryParams, FSystemFactory::FParamsInitializer>> GDefaultParamsForObject;
+static TMap<FTopLevelAssetPath, TPair<FUAFSystemFactoryParams, FSystemFactory::FParamsInitializer>> GDefaultParamsForObject;
 static TArray<UClass*> GAllFactoryClasses;
 
 const TPair<FUAFSystemFactoryParams, FSystemFactory::FParamsInitializer>* FindSystemParamsForClass(const UClass* InClass)
@@ -26,7 +26,7 @@ const TPair<FUAFSystemFactoryParams, FSystemFactory::FParamsInitializer>* FindSy
 	const TPair<FUAFSystemFactoryParams, FSystemFactory::FParamsInitializer>* FoundParams = nullptr;
 	while (FoundParams == nullptr && CurrentClass != nullptr)
 	{
-		FoundParams = GDefaultParamsForObject.Find(CurrentClass);
+		FoundParams = GDefaultParamsForObject.Find(CurrentClass->GetClassPathName());
 		CurrentClass = CurrentClass->GetSuperClass();
 	}
 
@@ -104,7 +104,7 @@ FUAFSystemFactoryParams FSystemFactory::GetDefaultParamsForObject(const UObject*
 
 void FSystemFactory::RegisterDefaultParamsForClass(const UClass* InClass, FUAFSystemFactoryParams&& InParams, FParamsInitializer&& InInitializer)
 {
-	Private::GDefaultParamsForObject.Add(InClass, { MoveTemp(InParams), MoveTemp(InInitializer) });
+	Private::GDefaultParamsForObject.Add(InClass->GetClassPathName(), {MoveTemp(InParams), MoveTemp(InInitializer)});
 	Private::GAllFactoryClasses.Add(const_cast<UClass*>(InClass));
 }
 
