Prevent creating shared variable assets from properties in a native shared variable struct

#Synced-CL 48503794

diff --git a/Engine/Plugins/Experimental/UAF/UAF/Source/UAFEditor/Private/Variables/Outliner/VariablesOutlinerMode.cpp b/Engine/Plugins/Experimental/UAF/UAF/Source/UAFEditor/Private/Variables/Outliner/VariablesOutlinerMode.cpp
--- a/Engine/Plugins/Experimental/UAF/UAF/Source/UAFEditor/Private/Variables/Outliner/VariablesOutlinerMode.cpp
+++ b/Engine/Plugins/Experimental/UAF/UAF/Source/UAFEditor/Private/Variables/Outliner/VariablesOutlinerMode.cpp
@@ -1924,9 +1924,11 @@ bool FVariablesOutlinerMode::CanCreateSharedVariablesAssets() const
 	const bool bSelectionOnlyContainsVariables = Selection.Num() == Selection.Num<FVariablesOutlinerEntryItem>();
 	const bool bSelectionOnlyContainsSingleAsset = Selection.Num() == 1 && Selection.Num<FOutlinerAssetItem>() == 1;
 
+	bool bSelectionContainsNativeSharedVariablesStructProperty = false;
+
 	// Prevent moving from an assets which would leave it empty (removing all variables)
 	TMap<const UAnimNextRigVMAssetEditorData*, int32> EditorDataToNumRemainingVariables;
-	Selection.ForEachItem<FVariablesOutlinerEntryItem>([&EditorDataToNumRemainingVariables](const FVariablesOutlinerEntryItem& Item)
+	Selection.ForEachItem<FVariablesOutlinerEntryItem>([&EditorDataToNumRemainingVariables, &bSelectionContainsNativeSharedVariablesStructProperty](const FVariablesOutlinerEntryItem& Item)
 	{
 		if (const UAnimNextVariableEntry* VariableEntry = Item.WeakEntry.Get())
 		{
@@ -1949,10 +1951,20 @@ bool FVariablesOutlinerMode::CanCreateSharedVariablesAssets() const
 				}
 			}
 		}
+		else if (const FSceneOutlinerTreeItemPtr& Parent = Item.GetParent(); Parent->IsA<FVariablesOutlinerStructSharedVariablesItem>())
+		{
+			bSelectionContainsNativeSharedVariablesStructProperty = true;
+			return false;
+		}
 
 		return true;
 	});
 
+	if (bSelectionContainsNativeSharedVariablesStructProperty)
+	{
+		return false;
+	}
+
 	TArray<int32> NumRemainVariablesArray;
 	
 	EditorDataToNumRemainingVariables.GenerateValueArray(NumRemainVariablesArray);
