Add crash reporter contexts of UAF instances and Workspace Editor instances

#Synced-CL 48508358

diff --git a/Engine/Plugins/Experimental/UAF/UAF/Source/UAF/Internal/CrashReporter/CrashReporterHandler.cpp b/Engine/Plugins/Experimental/UAF/UAF/Source/UAF/Internal/CrashReporter/CrashReporterHandler.cpp
new file mode 100644
--- /dev/null
+++ b/Engine/Plugins/Experimental/UAF/UAF/Source/UAF/Internal/CrashReporter/CrashReporterHandler.cpp
@@ -0,0 +1,159 @@
+// Copyright Epic Games, Inc. All Rights Reserved.
+
+#include "CrashReporterHandler.h"
+#include "UObject/Object.h"
+#include "HAL/IConsoleManager.h"
+#include "UObject/WeakObjectPtr.h"
+#include "UObject/WeakObjectPtrTemplates.h"
+
+#if WITH_ADDITIONAL_CRASH_CONTEXTS
+
+namespace UE::UAF
+{
+static bool bCrashHandlerEnabled = false;
+static FAutoConsoleVariableRef CVarCrashHandlerEnabled(
+	TEXT("UAF.CrashHandlerEnabled"),
+	bCrashHandlerEnabled,
+	TEXT("UAF will add context into the crash reporter."),
+	ECVF_Default);
+
+namespace CrashReporter
+{
+	struct FCrashReporterHandlerImpl
+	{
+	public:
+		struct FContext
+		{
+			explicit FContext(TNotNull<const UObject*> InOwner, TNotNull<const UObject*> InObjectOfInterest, FName InContext, uint32 InThreadID)
+				: Owner(InOwner) 
+				, ObjectOfInterest(InObjectOfInterest)
+				, Context(InContext)
+				, ThreadID(InThreadID)
+			{
+				ID = GenerateID();
+			}
+
+			UE_AUTORTFM_ALWAYS_OPEN
+			static uint32 GenerateID()
+			{
+				// This ID only needs to be unique; it doesn't need to be rolled back if an AutoRTFM transaction fails.
+				static std::atomic<uint32> IDGenerator = 0;
+				return ++IDGenerator;
+			}
+
+			TWeakObjectPtr<const UObject> Owner;
+			TWeakObjectPtr<const UObject> ObjectOfInterest;
+			FName Context;
+			uint32 ThreadID;
+			uint32 ID;
+		};
+
+	public:
+		uint32 PushInfo(TNotNull<const UObject*> InOwner, TNotNull<const UObject*> InObjectOfInterest, FName InContext)
+		{
+			UE::TScopeLock LockGuard(CriticalSection);
+			return Contexts.Emplace_GetRef(InOwner, InObjectOfInterest, InContext, FPlatformTLS::GetCurrentThreadId()).ID;
+		}
+
+		void PopInfo(uint32 ID)
+		{
+			UE::TScopeLock LockGuard(CriticalSection);
+			const int32 FoundIndex = Contexts.IndexOfByPredicate([ID](const FContext& Other)
+				{
+					return Other.ID == ID;
+				});
+			if (FoundIndex != INDEX_NONE)
+			{
+				Contexts.RemoveAtSwap(FoundIndex);
+			}
+		}
+
+		void DumpInfo(FCrashContextExtendedWriter& Writer)
+		{
+			// Callback from the crash reporter.Avoids memory allocation in case the crash is an OOM.
+			UE::TScopeLock LockGuard(CriticalSection);
+
+			for (int32 Index = 0; Index < Contexts.Num(); ++Index)
+			{
+				const FContext& Context = Contexts[Index];
+				Builder.Reset();
+			
+				if (const UObject* Owner = Context.Owner.Get())
+				{
+					Owner->GetFullName(Builder);
+				}
+				else
+				{
+					Builder.Append(TEXT("<Invalid>"));
+				}
+
+				if (const UObject* ObjectOfInterest = Context.ObjectOfInterest.Get())
+				{
+					ObjectOfInterest->GetFullName(Builder);
+				}
+				else
+				{
+					Builder.Append(TEXT("<Invalid>"));
+				}
+
+				Builder << TEXT('\n');
+				Builder << Context.Context;
+
+				Builder << TEXT('\n');
+				Builder << Context.ThreadID;
+
+				// Avoid allocating a new name for each entry.
+				constexpr int32 BufferSize = 16;
+				TCHAR Identifier[BufferSize];
+				FCString::Snprintf(Identifier, UE_ARRAY_COUNT(Identifier), TEXT("UAF%d"), Index);
+				Writer.AddString(Identifier, Builder.ToString());
+			}
+		}
+
+	private:
+		TArray<FContext, TInlineAllocator<32>> Contexts;
+		FTransactionallySafeCriticalSection CriticalSection;
+		TStringBuilder<1024> Builder;
+	};
+
+
+	//~ Use a shared ptr instead of manually deleting it in case we crash while unregistering the UAF module.
+	static TSharedPtr<FCrashReporterHandlerImpl> Instance;
+	static FDelegateHandle AdditionalCrashContextDelegateHandle;
+}
+
+void FCrashReporterHandler::Register()
+{
+	CrashReporter::Instance = MakeShared<CrashReporter::FCrashReporterHandlerImpl>();
+	CrashReporter::AdditionalCrashContextDelegateHandle = FGenericCrashContext::OnAdditionalCrashContextDelegate().AddSP(CrashReporter::Instance.ToSharedRef(), &CrashReporter::FCrashReporterHandlerImpl::DumpInfo);
+}
+
+void FCrashReporterHandler::Unregister()
+{
+	FGenericCrashContext::OnAdditionalCrashContextDelegate().Remove(CrashReporter::AdditionalCrashContextDelegateHandle);
+	// Keep the AdditionalCrashContextDelegateHandle valid in case it currently crashing. The pointer will be automatically removed (shared ptr) when the module unloads.
+}
+
+FCrashReporterScope::FCrashReporterScope(TNotNull<const UObject*> InOwner, TNotNull<const UObject*> ObjectOfInterest, FName Context)
+{
+	bWasEnabled = bCrashHandlerEnabled;
+	if (bWasEnabled)
+	{
+		check(CrashReporter::Instance);
+		ID = CrashReporter::Instance->PushInfo(InOwner, ObjectOfInterest, Context);
+	}
+}
+
+FCrashReporterScope::~FCrashReporterScope()
+{
+	if (bWasEnabled)
+	{
+		CrashReporter::Instance->PopInfo(ID);
+	}
+}
+}
+
+
+
+
+#endif // WITH_ADDITIONAL_CRASH_CONTEXTS
diff --git a/Engine/Plugins/Experimental/UAF/UAF/Source/UAF/Internal/CrashReporter/CrashReporterHandler.h b/Engine/Plugins/Experimental/UAF/UAF/Source/UAF/Internal/CrashReporter/CrashReporterHandler.h
new file mode 100644
--- /dev/null
+++ b/Engine/Plugins/Experimental/UAF/UAF/Source/UAF/Internal/CrashReporter/CrashReporterHandler.h
@@ -0,0 +1,38 @@
+// Copyright Epic Games, Inc. All Rights Reserved.
+
+#pragma once
+
+#include "GenericPlatform/GenericPlatformCrashContext.h"
+#include "Misc/NotNull.h"
+
+#if WITH_ADDITIONAL_CRASH_CONTEXTS
+
+namespace UE::UAF
+{
+struct FCrashReporterHandler
+{
+	static void Register();
+	static void Unregister();
+};
+
+class FCrashReporterScope
+{
+public:
+	explicit FCrashReporterScope(TNotNull<const UObject*> InOwner, TNotNull<const UObject*> ObjectOfInterest, FName Context);
+	~FCrashReporterScope();
+
+	FCrashReporterScope() = delete;
+	FCrashReporterScope(const FCrashReporterScope&) = delete;
+	FCrashReporterScope& operator=(const FCrashReporterScope&) = delete;
+
+private:
+	bool bWasEnabled = false;
+	uint32 ID = 0;
+};
+
+}
+
+#define UAF_CRASH_REPORTER_SCOPE(InOwner, InObjectOfInterest, InContextName) ::UE::UAF::FCrashReporterScope ANONYMOUS_VARIABLE(AddCrashContext) {(InOwner), (InObjectOfInterest), (#InContextName)}
+#endif // WITH_ADDITIONAL_CRASH_CONTEXTS 
+
+#define UE_WORKSPACE_EDITOR_CRASH_REPORTER_SCOPE()
diff --git a/Engine/Plugins/Experimental/UAF/UAF/Source/UAF/Private/AnimNextModuleImpl.cpp b/Engine/Plugins/Experimental/UAF/UAF/Source/UAF/Private/AnimNextModuleImpl.cpp
--- a/Engine/Plugins/Experimental/UAF/UAF/Source/UAF/Private/AnimNextModuleImpl.cpp
+++ b/Engine/Plugins/Experimental/UAF/UAF/Source/UAF/Private/AnimNextModuleImpl.cpp
@@ -27,6 +27,7 @@
 #include "Variables/AnimNextSoftFunctionPtr.h"
 #include "RigVMCore/RigVMGraphFunctionDefinition.h"
 #include "HierarchyTable.h"
+#include "CrashReporter/CrashReporterHandler.h"
 #include "Factory/SystemFactory.h"
 #include "Features/IModularFeatures.h"
 #include "RewindDebuggerRuntimeInterface/IRewindDebuggerRuntimeExtension.h"
@@ -46,6 +47,10 @@ namespace UE::UAF
 	void FAnimNextModuleImpl::StartupModule()
 	{
 		GetMutableDefault<UAnimNextConfig>()->LoadConfig();
+		
+#if WITH_ADDITIONAL_CRASH_CONTEXTS
+		FCrashReporterHandler::Register();
+#endif // WITH_ADDITIONAL_CRASH_CONTEXTS
 
 		FRigVMRegistry& RigVMRegistry = FRigVMRegistry::Get();
 		static TPair<UClass*, FRigVMRegistry::ERegisterObjectOperation> const AllowedObjectTypes[] =
@@ -140,6 +145,10 @@ namespace UE::UAF
 		FDataRegistry::Destroy();
 		FSystemFactory::Destroy();
 		FSystemReference::Destroy();
+		
+#if WITH_ADDITIONAL_CRASH_CONTEXTS
+		FCrashReporterHandler::Unregister();
+#endif // WITH_ADDITIONAL_CRASH_CONTEXTS
 	}
 }
 
diff --git a/Engine/Plugins/Experimental/UAF/UAF/Source/UAF/Private/Module/AnimNextModuleInstance.cpp b/Engine/Plugins/Experimental/UAF/UAF/Source/UAF/Private/Module/AnimNextModuleInstance.cpp
--- a/Engine/Plugins/Experimental/UAF/UAF/Source/UAF/Private/Module/AnimNextModuleInstance.cpp
+++ b/Engine/Plugins/Experimental/UAF/UAF/Source/UAF/Private/Module/AnimNextModuleInstance.cpp
@@ -12,6 +12,7 @@
 #include "SceneInterface.h"
 #include "Script/UAFScriptComponent.h"
 #include "Algo/TopologicalSort.h"
+#include "CrashReporter/CrashReporterHandler.h"
 #include "Factory/SystemFactory.h"
 #include "Factory/UAFSystemFactoryParams.h"
 #include "Logging/StructuredLog.h"
@@ -434,6 +435,8 @@ const UAnimNextModule* FAnimNextModuleInstance::GetSystem() const
 void FAnimNextModuleInstance::OnCompileJobFinished()
 {
 	using namespace UE::UAF;
+	
+	UAF_CRASH_REPORTER_SCOPE(Object, GetSystem(), OnCompileJobFinished);
 
 	FGuardValue_Bitfield(bIsRecreatingOnCompile, true);
 
diff --git a/Engine/Plugins/Experimental/UAF/UAF/Source/UAF/Private/Module/ModuleTickFunction.cpp b/Engine/Plugins/Experimental/UAF/UAF/Source/UAF/Private/Module/ModuleTickFunction.cpp
--- a/Engine/Plugins/Experimental/UAF/UAF/Source/UAF/Private/Module/ModuleTickFunction.cpp
+++ b/Engine/Plugins/Experimental/UAF/UAF/Source/UAF/Private/Module/ModuleTickFunction.cpp
@@ -6,6 +6,7 @@
 #include "AnimNextStats.h"
 #include "RigVMRuntimeDataRegistry.h"
 #include "Algo/TopologicalSort.h"
+#include "CrashReporter/CrashReporterHandler.h"
 #include "Module/AnimNextModule.h"
 #include "Module/AnimNextModuleInstance.h"
 #include "Module/ModuleTaskContext.h"
@@ -35,17 +36,20 @@ namespace UE::UAF
 
 void FModuleEventTickFunction::Initialize(float InDeltaTime)
 {
+	UAF_CRASH_REPORTER_SCOPE(ModuleInstance->GetObject(), ModuleInstance->GetSystem(), Initialize);
 	ModuleInstance->RunScriptEvent(InitializeEvent, InDeltaTime);
 }
 
 void FModuleEventTickFunction::ExecuteBindings_WT(float InDeltaTime)
 {
+	UAF_CRASH_REPORTER_SCOPE(ModuleInstance->GetObject(), ModuleInstance->GetSystem(), ExecuteBindings_WT);
 	ModuleInstance->RunScriptEvent(ExecuteBindings_WTEvent, InDeltaTime);
 }
 
 void FModuleEventTickFunction::EndTick(float DeltaTime)
 {
 	SCOPED_NAMED_EVENT(UAF_Module_EndTick, FColor::Orange);
+	UAF_CRASH_REPORTER_SCOPE(ModuleInstance->GetObject(), ModuleInstance->GetSystem(), EndTick);
 
 	// Give module a chance to finish up processing
 	ModuleInstance->EndExecution(DeltaTime);
@@ -121,6 +125,7 @@ void FModuleEventTickFunction::EndTick(float DeltaTime)
 void FModuleEventTickFunction::ExecuteTick(float DeltaTime, ELevelTick TickType, ENamedThreads::Type CurrentThread, const FGraphEventRef& MyCompletionGraphEvent)
 {
 	LLM_SCOPE_BYNAME(TEXT("UAF"));
+	UAF_CRASH_REPORTER_SCOPE(ModuleInstance->GetObject(), ModuleInstance->GetSystem(), ExecuteTick);
 	
 	check(PostExecuteGameThreadTaskList.IsEmpty());	// Should always be empty since we move it to the game thread task
 
@@ -182,6 +187,7 @@ void FModuleEventTickFunction::ExecuteTick(float DeltaTime, ELevelTick TickType,
 void FModuleEventTickFunction::Run(float InDeltaTime)
 {
 	SCOPED_NAMED_EVENT(UAF_Module_EventTick, FColor::Orange);
+	UAF_CRASH_REPORTER_SCOPE(ModuleInstance->GetObject(), ModuleInstance->GetSystem(), Run);
 	
 	while (!PreExecuteTasks.IsEmpty())
 	{
diff --git a/Engine/Plugins/Experimental/Workspace/Source/WorkspaceEditor/Private/WorkspaceEditor.cpp b/Engine/Plugins/Experimental/Workspace/Source/WorkspaceEditor/Private/WorkspaceEditor.cpp
--- a/Engine/Plugins/Experimental/Workspace/Source/WorkspaceEditor/Private/WorkspaceEditor.cpp
+++ b/Engine/Plugins/Experimental/Workspace/Source/WorkspaceEditor/Private/WorkspaceEditor.cpp
@@ -40,6 +40,9 @@
 namespace UE::Workspace
 {
 
+static bool GOutputWorkspaceStateToCrashContext = false;
+static FAutoConsoleVariableRef GOutputWorkspaceStateToCrashContextVar(TEXT("Workspace.OutputWorkspaceStateToCrashContext"), GOutputWorkspaceStateToCrashContext, TEXT("When enabled will output the state of all workspace editor instances to the crash context"));
+
 FWorkspaceEditorSelectionScope::FWorkspaceEditorSelectionScope(const TSharedPtr< IWorkspaceEditor>& InWorkspaceEditor) : WeakWorkspaceEditor(InWorkspaceEditor)
 {
 	const TSharedPtr<FWorkspaceEditor> SharedWorkspaceEditor = StaticCastSharedPtr<FWorkspaceEditor>(InWorkspaceEditor);
@@ -235,6 +238,13 @@ void FWorkspaceEditor::CreateWidgets()
 
 void FWorkspaceEditor::PostInitAssetEditor()
 {
+#if WITH_ADDITIONAL_CRASH_CONTEXTS
+	if (GOutputWorkspaceStateToCrashContext)
+	{
+		FGenericCrashContext::OnAdditionalCrashContextDelegate().AddSP(this, &FWorkspaceEditor::OnProvideAdditionalCrashContext);
+	}
+#endif // WITH_ADDITIONAL_CRASH_CONTEXTS
+
 	if (UContextObjectStore* ContextObjectStore = GetEditorModeManager().GetInteractiveToolsContext()->ContextObjectStore)
 	{
 		UAssetEditorToolkitMenuContext* AssetEditorContext = ContextObjectStore->FindContext<UAssetEditorToolkitMenuContext>();
@@ -1070,6 +1080,49 @@ bool FWorkspaceEditor::IsViewportPinned() const
 	return ViewportPinnedAssetPath.IsValid();
 }
 
+#if WITH_ADDITIONAL_CRASH_CONTEXTS
+void FWorkspaceEditor::OnProvideAdditionalCrashContext(FCrashContextExtendedWriter& CrashContextExtendedWriter) const
+{
+	thread_local TStringBuilder<1024> Builder;  
+	{
+		// Workspace asset
+		Builder.Append(GetNameSafe(Workspace));
+		Builder << TEXT('\n');
+
+		// Workspace asset
+		if (Workspace)
+		{
+			Builder << TEXT("Root Asset: ");
+			const UObject* RootAsset = Workspace->GetRootAsset();
+			Builder.Append(GetNameSafe(RootAsset));
+			Builder << TEXT("\nLoaded Assets:\n");
+
+			TArray<TObjectPtr<UObject>> LoadedWorkspaceAssets;
+			Workspace->GetAssets(LoadedWorkspaceAssets);
+			for (const TObjectPtr<UObject>& Asset : LoadedWorkspaceAssets)
+			{
+				Builder << TEXT("\t+ ");
+				Builder.Append(GetNameSafe(RootAsset));
+				Builder << TEXT('\n');
+			}
+		}
+
+		const FWorkspaceDocument& CurrentDocument = GetFocusedWorkspaceDocument();
+		if (CurrentDocument.Export.GetIdentifier() != NAME_None)
+		{
+			Builder << TEXT("Focussed Document:\n");
+			Builder << TEXT("\tExport Path: ");
+			Builder.Append(CurrentDocument.Export.GetFullPath());
+			Builder << TEXT("\tObject: ");
+			Builder.Append(CurrentDocument.Object ? CurrentDocument.Object->GetPathName() : TEXT("None"));
+			Builder << TEXT('\n');
+		}
+	}
+	
+	CrashContextExtendedWriter.AddString(*GetEditorName().ToString(), Builder.ToString());
+}
+#endif // WITH_ADDITIONAL_CRASH_CONTEXTS
+
 void FWorkspaceEditor::FindInContentBrowser_Execute()
 {
 	if (IsValid(Workspace))
diff --git a/Engine/Plugins/Experimental/Workspace/Source/WorkspaceEditor/Private/WorkspaceEditor.h b/Engine/Plugins/Experimental/Workspace/Source/WorkspaceEditor/Private/WorkspaceEditor.h
--- a/Engine/Plugins/Experimental/Workspace/Source/WorkspaceEditor/Private/WorkspaceEditor.h
+++ b/Engine/Plugins/Experimental/Workspace/Source/WorkspaceEditor/Private/WorkspaceEditor.h
@@ -150,6 +150,10 @@ private:
 	bool AreAssetEntriesModified() const;
 
 	bool IsViewportPinned() const;
+		
+#if WITH_ADDITIONAL_CRASH_CONTEXTS
+	void OnProvideAdditionalCrashContext(FCrashContextExtendedWriter& CrashContextExtendedWriter) const;
+#endif // WITH_ADDITIONAL_CRASH_CONTEXTS
 
 	/** The asset being edited */
 	TObjectPtr<UWorkspace> Workspace = nullptr;
@@ -187,4 +191,4 @@ private:
 	FSoftObjectPath ViewportPinnedAssetPath;
 };
 
-}
\ No newline at end of file
+}
