UAF: Change defered task queues to MPSC

We can enqueue tasks here from multiple threads
Additionally simplified the dequeue loops

#Synced-CL 48607479

diff --git a/Engine/Plugins/Experimental/UAF/UAF/Source/UAF/Internal/Module/ModuleTickFunction.h b/Engine/Plugins/Experimental/UAF/UAF/Source/UAF/Internal/Module/ModuleTickFunction.h
--- a/Engine/Plugins/Experimental/UAF/UAF/Source/UAF/Internal/Module/ModuleTickFunction.h
+++ b/Engine/Plugins/Experimental/UAF/UAF/Source/UAF/Internal/Module/ModuleTickFunction.h
@@ -57,8 +57,8 @@ struct FModuleEventTickFunction : public FTickFunction
 	FAnimNextModuleInstance* ModuleInstance = nullptr;
 
 	// Tasks queued to execute once per event
-	TSpscQueue<TUniqueFunction<void(const FModuleTaskContext&)>> PreExecuteTasks;
-	TSpscQueue<TUniqueFunction<void(const FModuleTaskContext&)>> PostExecuteTasks;
+	TMpscQueue<TUniqueFunction<void(const FModuleTaskContext&)>> PreExecuteTasks;
+	TMpscQueue<TUniqueFunction<void(const FModuleTaskContext&)>> PostExecuteTasks;
 
 	// Tasks queued to execute on the game thread after the current tick completes
 	TArray<TUniqueFunction<void()>> PostExecuteGameThreadTaskList;
diff --git a/Engine/Plugins/Experimental/UAF/UAF/Source/UAF/Private/Module/AnimNextModuleInstance.cpp b/Engine/Plugins/Experimental/UAF/UAF/Source/UAF/Private/Module/AnimNextModuleInstance.cpp
--- a/Engine/Plugins/Experimental/UAF/UAF/Source/UAF/Private/Module/AnimNextModuleInstance.cpp
+++ b/Engine/Plugins/Experimental/UAF/UAF/Source/UAF/Private/Module/AnimNextModuleInstance.cpp
@@ -643,7 +643,7 @@ void FAnimNextModuleInstance::QueueTask(FName InEventName, TUniqueFunction<void(
 		}
 	}
 
-	TSpscQueue<TUniqueFunction<void(const FModuleTaskContext&)>>* Queue = nullptr;
+	TMpscQueue<TUniqueFunction<void(const FModuleTaskContext&)>>* Queue = nullptr;
 	if (FoundTickFunction)
 	{
 		switch (InLocation)
diff --git a/Engine/Plugins/Experimental/UAF/UAF/Source/UAF/Private/Module/ModuleTickFunction.cpp b/Engine/Plugins/Experimental/UAF/UAF/Source/UAF/Private/Module/ModuleTickFunction.cpp
--- a/Engine/Plugins/Experimental/UAF/UAF/Source/UAF/Private/Module/ModuleTickFunction.cpp
+++ b/Engine/Plugins/Experimental/UAF/UAF/Source/UAF/Private/Module/ModuleTickFunction.cpp
@@ -189,11 +189,12 @@ void FModuleEventTickFunction::Run(float InDeltaTime)
 	SCOPED_NAMED_EVENT(UAF_Module_EventTick, FColor::Orange);
 	UAF_CRASH_REPORTER_SCOPE(ModuleInstance->GetObject(), ModuleInstance->GetSystem(), Run);
 	
-	while (!PreExecuteTasks.IsEmpty())
 	{
-		TOptional<TUniqueFunction<void(const FModuleTaskContext&)>> Function = PreExecuteTasks.Dequeue();
-		check(Function.IsSet());
-		Function.GetValue()(FModuleTaskContext(*ModuleInstance));
+		TUniqueFunction<void(const FModuleTaskContext&)> Function;
+		while (PreExecuteTasks.Dequeue(Function))
+		{
+			Function(FModuleTaskContext(*ModuleInstance));
+		}
 	}
 
 	if (bFirstUserEvent)
@@ -221,11 +222,12 @@ void FModuleEventTickFunction::Run(float InDeltaTime)
 
 	ModuleInstance->RaiseTraitEvents(ModuleInstance->OutputEventList);
 
-	while (!PostExecuteTasks.IsEmpty())
 	{
-		TOptional<TUniqueFunction<void(const FModuleTaskContext&)>> Function = PostExecuteTasks.Dequeue();
-		check(Function.IsSet());
-		Function.GetValue()(FModuleTaskContext(*ModuleInstance));
+		TUniqueFunction<void(const FModuleTaskContext&)> Function;
+		while (PostExecuteTasks.Dequeue(Function))
+		{
+			Function(FModuleTaskContext(*ModuleInstance));
+		}
 	}
 
 	if (bLastUserEvent)
diff --git a/Engine/Plugins/Experimental/UAF/UAF/Source/UAF/Private/Module/SystemReference.cpp b/Engine/Plugins/Experimental/UAF/UAF/Source/UAF/Private/Module/SystemReference.cpp
--- a/Engine/Plugins/Experimental/UAF/UAF/Source/UAF/Private/Module/SystemReference.cpp
+++ b/Engine/Plugins/Experimental/UAF/UAF/Source/UAF/Private/Module/SystemReference.cpp
@@ -15,6 +15,7 @@
 #include "Module/ModuleTaskContext.h"
 #include "Module/SystemDependency.h"
 #include "Module/UAFWeakSystemReference.h"
+#include "Containers/MpscQueue.h"
 
 namespace UE::UAF::Private
 {
@@ -27,20 +28,17 @@ namespace UE::UAF::Private
 		FRWLock SystemInstancesLock;
 
 		// Actions to perform once world tick is completed each frame
-		TSpscQueue<TUniqueFunction<void()>> PendingActions;
+		TMpscQueue<TUniqueFunction<void()>> PendingActions;
 
 		// Deferred cleanup of handles that are released this frame
 		void FlushPendingActions()
 		{
 			using namespace UE::UAF;
 
-			while (!PendingActions.IsEmpty())
+			TUniqueFunction<void()> PendingAction;
+			while (PendingActions.Dequeue(PendingAction))
 			{
-				TUniqueFunction<void()> PendingAction;
-				if (PendingActions.Dequeue(PendingAction))
-				{
-					PendingAction();
-				}
+				PendingAction();
 			}
 		}
 
