Handle collapsed graphs recursively when generating functions outliner hierarchy

#Synced-CL 48606655

diff --git a/Engine/Plugins/Experimental/UAF/UAF/Source/UAFEditor/Private/Functions/Outliner/FunctionsOutlinerHierarchy.cpp b/Engine/Plugins/Experimental/UAF/UAF/Source/UAFEditor/Private/Functions/Outliner/FunctionsOutlinerHierarchy.cpp
--- a/Engine/Plugins/Experimental/UAF/UAF/Source/UAFEditor/Private/Functions/Outliner/FunctionsOutlinerHierarchy.cpp
+++ b/Engine/Plugins/Experimental/UAF/UAF/Source/UAFEditor/Private/Functions/Outliner/FunctionsOutlinerHierarchy.cpp
@@ -117,6 +117,23 @@ FFunctionsOutlinerHierarchy::FFunctionsOutlinerHierarchy(ISceneOutlinerMode* Mod
 {
 }
 
+void FFunctionsOutlinerHierarchy::RecurseCollapseNodeGraphs(const UAnimNextRigVMAssetEditorData* EditorData, TArray<FSceneOutlinerTreeItemPtr>& OutItems, URigVMLibraryNode* LibraryNode) const
+{
+	int32 SubIndex = 0;
+	for (URigVMNode* Node : LibraryNode->GetContainedGraph()->GetNodes())
+	{
+		if (URigVMCollapseNode* CollapseNode = Cast<URigVMCollapseNode>(Node))
+		{
+			if (FSceneOutlinerTreeItemPtr CollapsedItem = Mode->CreateItemFor<FFunctionsOutlinerCollapsedGraphItem>(FFunctionsOutlinerCollapsedGraphItem::FItemData{UncookedOnly::FUtils::GetAsset<UAnimNextRigVMAsset>(EditorData), ++SubIndex, CollapseNode, LibraryNode}))
+			{
+				OutItems.Add(CollapsedItem);
+				
+				RecurseCollapseNodeGraphs(EditorData, OutItems, CollapseNode);
+			}
+		}
+	}
+}
+
 void FFunctionsOutlinerHierarchy::CreateItemsForAsset(const UAnimNextRigVMAssetEditorData* EditorData, TArray<FSceneOutlinerTreeItemPtr>& OutItems) const
 {
 	TSet<FString> ParsedCategories;
@@ -172,17 +189,7 @@ void FFunctionsOutlinerHierarchy::CreateItemsForAsset(const UAnimNextRigVMAssetE
 
 				if (LibraryNode->GetContainedGraph())
 				{
-					int32 SubIndex = 0;
-					for (URigVMNode* Node : LibraryNode->GetContainedGraph()->GetNodes())
-					{
-						if (URigVMCollapseNode* CollapseNode = Cast<URigVMCollapseNode>(Node))
-						{
-							if (FSceneOutlinerTreeItemPtr CollapsedItem = Mode->CreateItemFor<FFunctionsOutlinerCollapsedGraphItem>(FFunctionsOutlinerCollapsedGraphItem::FItemData{UncookedOnly::FUtils::GetAsset<UAnimNextRigVMAsset>(EditorData), ++SubIndex, CollapseNode, LibraryNode}))
-							{
-								OutItems.Add(CollapsedItem);
-							}
-						}
-					}
+					RecurseCollapseNodeGraphs(EditorData, OutItems, LibraryNode);
 				}
 			}
 		}
diff --git a/Engine/Plugins/Experimental/UAF/UAF/Source/UAFEditor/Private/Functions/Outliner/FunctionsOutlinerHierarchy.h b/Engine/Plugins/Experimental/UAF/UAF/Source/UAFEditor/Private/Functions/Outliner/FunctionsOutlinerHierarchy.h
--- a/Engine/Plugins/Experimental/UAF/UAF/Source/UAFEditor/Private/Functions/Outliner/FunctionsOutlinerHierarchy.h
+++ b/Engine/Plugins/Experimental/UAF/UAF/Source/UAFEditor/Private/Functions/Outliner/FunctionsOutlinerHierarchy.h
@@ -4,6 +4,8 @@
 
 #include "Common/Outliner/OutlinerHierarchy.h"
 
+class URigVMLibraryNode;
+
 namespace UE::UAF::Editor
 {
 class FFunctionsOutlinerHierarchy : public FOutlinerHierarchy
@@ -13,5 +15,7 @@ public:
 
 	FFunctionsOutlinerHierarchy(ISceneOutlinerMode* Mode);
 	virtual void CreateItemsForAsset(const UAnimNextRigVMAssetEditorData* EditorData, TArray<FSceneOutlinerTreeItemPtr>& OutItems) const override;
+private:
+	void RecurseCollapseNodeGraphs(const UAnimNextRigVMAssetEditorData* EditorData, TArray<FSceneOutlinerTreeItemPtr>& OutItems, URigVMLibraryNode* LibraryNode) const;
 };
 }
