UAF: More iteration on Pose IO

This moves the responsibility of converting/processing the LODPose into render data to the client system (as all will have slightly differing requirements).
OutputPose component is now vastly simplified, but now acts as a consistent reference pose, mesh and LOD for the whole of a system's execution.
Adds OnBeginExecution to module components and ensures it is called in the first user event so that the OutputPose component can pull its mesh & LOD.

#Synced-CL 48609941

diff --git a/Engine/Plugins/Experimental/UAF/UAF/Source/UAF/Internal/Graph/UAFOutputPoseComponent.h b/Engine/Plugins/Experimental/UAF/UAF/Source/UAF/Internal/Graph/UAFOutputPoseComponent.h
--- a/Engine/Plugins/Experimental/UAF/UAF/Source/UAF/Internal/Graph/UAFOutputPoseComponent.h
+++ b/Engine/Plugins/Experimental/UAF/UAF/Source/UAF/Internal/Graph/UAFOutputPoseComponent.h
@@ -14,6 +14,11 @@ class UMeshEntityComponentImpl_Skinned;
 class FEntitySkeletalAnimationModule;
 struct FUAFOutputPoseComponent;
 
+namespace UE::UAF
+{
+class IPoseAdapter;
+}
+
 // Module component used to identify output poses
 USTRUCT()
 struct FUAFOutputPoseComponent : public FUAFModuleInstanceComponent
@@ -28,32 +33,38 @@ struct FUAFOutputPoseComponent : public FUAFModuleInstanceComponent
 		return RefPose;
 	}
 
+	// Get the LOD we are currently outputting as
+	int32 GetLOD() const
+	{
+		return LOD;
+	}
+
+	// Function used to callback into 
+	using FGenerateRenderDataFunc = TFunctionRef<void(uint32 InSerialNumber, const FAnimNextGraphLODPose& InPose, const FAnimNextGraphReferencePose& InRefPose)>;
+
 private:
+	// FUAFModuleInstanceComponent interface
+	virtual void OnBeginExecution(float InDeltaTime) override;
+
 	// Write the pose
 	// @param    InPose    The pose to write
 	void WritePose(const FAnimNextGraphLODPose& InPose);
 
 	// Generate the pose etc. for rendering by a mesh.
-	// If the serial numbers differ, a read occurs, otherwise the read is skipped.
-	// If the supplied serial number matches the one held in this component, then the pose is assumed to have not been written, so no copy needs to occur
-	// If the serial number supplied is zero, then the read will always succeed
+	// If the serial numbers differ, a callback to the InFunction occurs, otherwise the call is skipped.
+	// If the supplied serial number matches the one held in this component, then the pose is assumed to have not been written, so no work needs to occur
+	// If the serial number supplied is zero, then the GenerateRenderData call will always be made
+	// InFunction can also be skipped if other invariants are not satisfied (e.g. null or mismatched mesh)
 	// @param    InSerialNumber    The current serial number that the reader holds
-	// @param    InMeshLOD         The pose LOD to generate
-	// @param    InLocalToWorld    The local to world transform used to generate bounds
-	// @param    InFunction        The function to call on read
+	// @param    InFunction        Function that will be called if conditions are met
 	// @return true if the read was successful (i.e. the serial numbers were mismatched)
-	UE_API bool GenerateMeshRenderData(uint32 InSerialNumber, int32 InMeshLOD, const FTransform& InLocalToWorld, TFunctionRef<void(TConstArrayView<FMatrix44f> InReferenceToLocalMatrices, const FBoxSphereBounds& InWorldBounds, const FBoxSphereBounds& InLocalBounds, uint32 InSerialNumber)> InReadFunction) const;
+	UE_API bool GenerateRenderData(uint32 InSerialNumber, FGenerateRenderDataFunc InFunction) const;
 
-	// Helper function used to calculate animated mesh bounds.
-	// Note: Adapted from USkinnedMeshComponent::CalcMeshBound
-	void CalcMeshBound(TConstArrayView<FTransform> InComponentTransforms, const FTransform& InLocalToWorld, FBoxSphereBounds& OutWorldBounds, FBoxSphereBounds& OutLocalBounds) const;
+	// Called to bind this component to a specific pose adapter (e.g. mesh component)
+	UE_API void Bind(UE::UAF::IPoseAdapter& InPoseAdapter);
 
-	// Helper function used to calculate physics-asset defined bounds.
-	// Note: Adapted from UPhysicsAsset::CalcAABB
-	void CalcBoundsFromPhysicsAsset(const UPhysicsAsset* InPhysicAsset, TConstArrayView<FTransform> InComponentTransforms, const FTransform& InLocalToWorld, FBox& OutWorldBounds, FBox& OutLocalBounds) const;
-
-	// Called to bind this component to a specific mesh host
-	UE_API void BindToHost(USkeletalMesh* InMesh, TFunction<void(void)>&& InSignalFunction);
+	// Unbinds from any pose adapter we are currently bound to
+	UE_API void Unbind();
 
 private:
 	friend FRigUnit_UAFWriteOutputPose;
@@ -71,8 +82,11 @@ private:
 	UPROPERTY(Transient)
 	FAnimNextGraphReferencePose RefPose;
 
-	// Function to call when a pose is written (on the GT)
-	TFunction<void(void)> SignalFunction;
+	// Pose adapter to use for output
+	UE::UAF::IPoseAdapter* PoseAdapter = nullptr;
+
+	// The LOD that our system is running at this tick
+	int32 LOD = 0;
 
 	// Serial number updated each time we write the pose
 	uint32 SerialNumber = 0;
diff --git a/Engine/Plugins/Experimental/UAF/UAF/Source/UAF/Internal/Module/AnimNextModuleInstance.h b/Engine/Plugins/Experimental/UAF/UAF/Source/UAF/Internal/Module/AnimNextModuleInstance.h
--- a/Engine/Plugins/Experimental/UAF/UAF/Source/UAF/Internal/Module/AnimNextModuleInstance.h
+++ b/Engine/Plugins/Experimental/UAF/UAF/Source/UAF/Internal/Module/AnimNextModuleInstance.h
@@ -160,6 +160,9 @@ private:
 	void OnCompileJobFinished();
 #endif
 
+	// Call module components' OnBeginExecution extension point
+	void BeginExecution(float InDeltaTime);
+
 	// Give the module components a chance to handle events and call their OnEndExecution extension point
 	void EndExecution(float InDeltaTime);
 
diff --git a/Engine/Plugins/Experimental/UAF/UAF/Source/UAF/Private/Graph/UAFOutputPoseComponent.cpp b/Engine/Plugins/Experimental/UAF/UAF/Source/UAF/Private/Graph/UAFOutputPoseComponent.cpp
--- a/Engine/Plugins/Experimental/UAF/UAF/Source/UAF/Private/Graph/UAFOutputPoseComponent.cpp
+++ b/Engine/Plugins/Experimental/UAF/UAF/Source/UAF/Private/Graph/UAFOutputPoseComponent.cpp
@@ -3,112 +3,37 @@
 #include "Graph/UAFOutputPoseComponent.h"
 
 #include "AnimationRuntime.h"
+#include "IPoseAdapter.h"
 #include "ReferencePose.h"
 #include "Engine/SkeletalMesh.h"
 #include "Graph/AnimNext_LODPose.h"
 #include "Module/AnimNextModuleInstance.h"
 #include "Module/ModuleEvents.h"
+#include "Module/ModuleTaskContext.h"
 #include "PhysicsEngine/PhysicsAsset.h"
 #include "PhysicsEngine/SkeletalBodySetup.h"
 
-void FUAFOutputPoseComponent::CalcBoundsFromPhysicsAsset(const UPhysicsAsset* InPhysicAsset, TConstArrayView<FTransform> InComponentTransforms, const FTransform& InLocalToWorld, FBox& OutWorldBounds, FBox& OutLocalBounds) const
+void FUAFOutputPoseComponent::OnBeginExecution(float InDeltaTime)
 {
-	OutLocalBounds = FBox(ForceInit);
-	OutWorldBounds = FBox(ForceInit);
-
-	const FReferenceSkeleton& RefSkeleton = Mesh->GetRefSkeleton();
-	TConstArrayView<int32> BodyIndexRefs = InPhysicAsset->BoundsBodies;
-	TConstArrayView<TObjectPtr<USkeletalBodySetup>> SkeletalBodySetups = InPhysicAsset->SkeletalBodySetups;
-
-	// Then iterate over bodies we want to consider, calculating bounding box for each
-	const int32 BodySetupNum = BodyIndexRefs.Num();
-	for(int32 BodySetupIndex = 0; BodySetupIndex < BodySetupNum; ++BodySetupIndex)
+	// Grab all the info we need to run the system this time around & generate/get the reference pose
+	if (PoseAdapter != nullptr)
 	{
-		const int32 BodyIndex = BodyIndexRefs[BodySetupIndex];
-		const UBodySetup* BodySetup = SkeletalBodySetups[BodyIndex];
-
-		if (BodySetupIndex + 1 < BodySetupNum)
-		{
-			int32 NextIndex = BodyIndexRefs[BodySetupIndex + 1];
-			FPlatformMisc::Prefetch(SkeletalBodySetups[NextIndex]);
-			FPlatformMisc::Prefetch(SkeletalBodySetups[NextIndex], PLATFORM_CACHE_LINE_SIZE);
-		}
+		UE::UAF::IPoseAdapter::FInputData InputData = PoseAdapter->GetPoseInput();
+		USkeletalMesh* OldMesh = Mesh;
+		Mesh = InputData.SkeletalMesh;
+		LOD = InputData.LOD;
 
-		const int32 BoneIndex = RefSkeleton.FindBoneIndex(BodySetup->BoneName);
-		if(BoneIndex != INDEX_NONE)
+		if (Mesh)
 		{
-			FBox BodySetupLocalBounds = BodySetup->AggGeom.CalcAABB(InComponentTransforms[BoneIndex]);
-
-			const FTransform WorldBoneTransform = InComponentTransforms[BoneIndex] * InLocalToWorld;
-			FBox BodySetupWorldBounds = BodySetup->AggGeom.CalcAABB(WorldBoneTransform);
-
-			auto SanitizeScaledBounds = [](FBox& InBounds)
+			if (Mesh != OldMesh)
 			{
-				if (InBounds.Min.X > InBounds.Max.X)
-				{
-					Swap(InBounds.Min.X, InBounds.Max.X);
-				}
-
-				if (InBounds.Min.Y > InBounds.Max.Y)
-				{
-					Swap(InBounds.Min.Y, InBounds.Max.Y);
-				}
-
-				if (InBounds.Min.Z > InBounds.Max.Z)
-				{
-					Swap(InBounds.Min.Z, InBounds.Max.Z);
-				}
-			};
-
-			SanitizeScaledBounds(BodySetupLocalBounds);
-			SanitizeScaledBounds(BodySetupWorldBounds);
-
-			OutLocalBounds += BodySetupLocalBounds;
-			OutWorldBounds += BodySetupWorldBounds;
+				RefPose.ReferencePose = UE::UAF::FDataRegistry::Get()->GetOrGenerateReferencePose(Mesh);
+			}
 		}
-	}
-
-	auto SanitizeZeroBounds = [&InLocalToWorld](FBox& InBox)
-	{
-		if(!InBox.IsValid)
+		else
 		{
-			InBox = FBox(InLocalToWorld.GetLocation(), InLocalToWorld.GetLocation());
+			RefPose.ReferencePose = UE::UAF::FDataHandle();
 		}
-
-		const float MinBoundSize = 1.f;
-		const FVector BoxSize = InBox.GetSize();
-
-		if(BoxSize.GetMin() < MinBoundSize)
-		{
-			const FVector ExpandByDelta ( FMath::Max<FVector::FReal>(0, MinBoundSize - BoxSize.X), FMath::Max<FVector::FReal>(0, MinBoundSize - BoxSize.Y), FMath::Max<FVector::FReal>(0, MinBoundSize - BoxSize.Z) );
-			InBox = InBox.ExpandBy(ExpandByDelta * 0.5f);	//expand by applies to both directions with GetSize applies to total size so divide by 2
-		}
-	};
-
-	SanitizeZeroBounds(OutLocalBounds);
-	SanitizeZeroBounds(OutWorldBounds);
-}
-
-void FUAFOutputPoseComponent::CalcMeshBound(TConstArrayView<FTransform> InComponentTransforms, const FTransform& InLocalToWorld, FBoxSphereBounds& OutWorldBounds, FBoxSphereBounds& OutLocalBounds) const
-{
-	const FVector DrawScale = InLocalToWorld.GetScale3D();
-	const UPhysicsAsset* PhysicsAsset = Mesh->GetPhysicsAsset();
-
-	// Can only use the PhysicsAsset to calculate the bounding box if we are not non-uniformly scaling the mesh.
-	if(DrawScale.IsUniform() && PhysicsAsset != nullptr && PhysicsAsset->SkeletalBodySetups.Num())
-	{
-		FBox WorldBox;
-		FBox LocalBox;
-		CalcBoundsFromPhysicsAsset(PhysicsAsset, InComponentTransforms, InLocalToWorld, WorldBox, LocalBox);
-		OutLocalBounds = FBoxSphereBounds(LocalBox);
-		OutWorldBounds = FBoxSphereBounds(WorldBox);
-	}
-	// Fallback is to use the one from the skeletal mesh.
-	// Usually pretty bad in terms of accuracy of where the SkelMesh Bounds are located (i.e. usually bigger than it needs to be)
-	else
-	{
-		OutLocalBounds = Mesh->GetBounds();
-		OutWorldBounds = OutLocalBounds.TransformBy(InLocalToWorld);
 	}
 }
 
@@ -124,23 +49,13 @@ void FUAFOutputPoseComponent::WritePose(const FAnimNextGraphLODPose& InPose)
 
 	Pose = InPose;
 
-	if (SignalFunction.IsSet())
+	if (PoseAdapter != nullptr)
 	{
-		// Signal on the GT that the pose was written
-		GetModuleInstance().QueueOutputTraitEvent(
-			MakeTraitEvent<FAnimNextModule_ActionEvent>([Function = SignalFunction]()
-			{
-				Function();
-			})
-		);
+		PoseAdapter->SignalPoseWritten(UE::UAF::FModuleTaskContext(GetModuleInstance()));
 	}
 }
 
-bool FUAFOutputPoseComponent::GenerateMeshRenderData(
-	uint32 InSerialNumber,
-	int32 InMeshLOD,
-	const FTransform& InLocalToWorld,
-	TFunctionRef<void(TConstArrayView<FMatrix44f> InReferenceToLocalMatrices, const FBoxSphereBounds& InWorldBounds, const FBoxSphereBounds& InLocalBounds, uint32 InSerialNumber)> InFunction) const
+bool FUAFOutputPoseComponent::GenerateRenderData(uint32 InSerialNumber, FGenerateRenderDataFunc InFunction) const
 {
 	// Zero serial indicates no valid pose to output, so early out
 	if (SerialNumber == 0)
@@ -159,105 +74,34 @@ bool FUAFOutputPoseComponent::GenerateMeshRenderData(
 		return false;
 	}
 
-	FMemMark MemMark(FMemStack::Get());
-
-	// Build render matrices from pose
-	const UE::UAF::FReferencePose& ReferencePose = Pose.LODPose.GetRefPose();
-	const TArray<FMatrix44f>& RefBasesInvMatrix = Mesh->GetRefBasesInvMatrix();
-
-	const TArrayView<const FBoneIndexType> MeshBoneIndexToParentMeshBoneIndexMap = ReferencePose.GetMeshBoneIndexToParentMeshBoneIndexMap();
-	const TArrayView<const FBoneIndexType> LODBoneIndexToMeshBoneIndexMap = ReferencePose.GetLODBoneIndexToMeshBoneIndexMap(Pose.LODPose.LODLevel);
-
-	checkf(LODBoneIndexToMeshBoneIndexMap.Num() == Pose.LODPose.GetNumBones(), TEXT("Buffer mismatch: %d:%d"), LODBoneIndexToMeshBoneIndexMap.Num(), Pose.LODPose.GetNumBones());
-
-	TArray<FTransform, TMemStackAllocator<>> ComponentSpacePose;
-	TArray<FMatrix44f, TMemStackAllocator<>> ReferenceToLocalMatrices;
-	ComponentSpacePose.SetNumUninitialized(MeshBoneIndexToParentMeshBoneIndexMap.Num());
-	ReferenceToLocalMatrices.SetNumUninitialized(MeshBoneIndexToParentMeshBoneIndexMap.Num());
-
-	const int32 NumPoseLODBones = Pose.LODPose.GetNumBones();
-	const FBoneIndexType* RESTRICT MeshBoneIndexToParentMeshBoneIndexMapPtr = MeshBoneIndexToParentMeshBoneIndexMap.GetData();
-	const FBoneIndexType* RESTRICT LODBoneIndexToMeshBoneIndexMapPtr = LODBoneIndexToMeshBoneIndexMap.GetData();
-	FTransform* RESTRICT ComponentSpaceData = ComponentSpacePose.GetData();
-
-	// First bone (if we have one) is always root bone, and it doesn't have a parent.
+	if (PoseAdapter == nullptr)
 	{
-		check(LODBoneIndexToMeshBoneIndexMap[0] == 0);
-		ComponentSpacePose[0] = Pose.LODPose.LocalTransformsView[0];
-		checkSlow(ComponentSpacePose[0].IsRotationNormalized());
-		checkSlow(!ComponentSpacePose[0].ContainsNaN());
-		const FMatrix44f ReferenceToLocal = (FMatrix44f)ComponentSpacePose[0].ToMatrixWithScale();
-		VectorMatrixMultiply(&ReferenceToLocalMatrices[0], &RefBasesInvMatrix[0], &ReferenceToLocal);
+		return false;
 	}
 
-	for (int32 LODBoneIndex = 1; LODBoneIndex < NumPoseLODBones; ++LODBoneIndex)
+	// Validate that the input we are transmitting matches the expected mesh of the adapter
+	// It is the assumed responsibility of client systems to ensure consistency of the mesh across the frame.
+	// It may be necessary when switching mesh mid-frame to recreate a new system to perform the work again.
+	// Note LOD can vary across the frame due to streaming and unfenced other LOD calculations, so clients must deal with LODs that are
+	// inconsistent between OnBeginExecution's GetPoseInput and this call
+	UE::UAF::IPoseAdapter::FInputData InputData = PoseAdapter->GetPoseInput();
+	if (ensure(InputData.SkeletalMesh == Mesh))
 	{
-		const FTransform LocalSpaceTransform = Pose.LODPose.LocalTransformsView[LODBoneIndex];
-
-		const FBoneIndexType MeshBoneIndex = LODBoneIndexToMeshBoneIndexMapPtr[LODBoneIndex];
-		const FBoneIndexType ParentMeshBoneIndex = MeshBoneIndexToParentMeshBoneIndexMapPtr[MeshBoneIndex];
-
-		const FTransform* RESTRICT ParentComponentSpaceTransform = ComponentSpaceData + ParentMeshBoneIndex;
-
-		FTransform* RESTRICT ComponentSpaceTransform = ComponentSpaceData + MeshBoneIndex;
-
-		FTransform::Multiply(ComponentSpaceTransform, &LocalSpaceTransform, ParentComponentSpaceTransform);
-
-		ComponentSpaceTransform->NormalizeRotation();
-
-		checkSlow(ComponentSpaceTransform->IsRotationNormalized());
-		checkSlow(!ComponentSpaceTransform->ContainsNaN());
-
-		const FMatrix44f ReferenceToLocal = (FMatrix44f)ComponentSpaceTransform->ToMatrixWithScale();
-		VectorMatrixMultiply(&ReferenceToLocalMatrices[MeshBoneIndex], &RefBasesInvMatrix[MeshBoneIndex], &ReferenceToLocal);
+		InFunction(SerialNumber, Pose, RefPose);
+		return true;
 	}
 
-	if (InMeshLOD < Pose.LODPose.LODLevel)
-	{
-		// If we evaluate the animation with lower quality than the visual mesh, we might
-		// have some transforms that we didn't compute. We'll grab the reference pose for them.
-		const TArrayView<const FBoneIndexType> LODBoneIndexToMeshBoneIndexMapForMeshLOD = ReferencePose.GetLODBoneIndexToMeshBoneIndexMap(InMeshLOD);
-
-		const int32 NumMeshLODBones = LODBoneIndexToMeshBoneIndexMapForMeshLOD.Num();
-		check(NumPoseLODBones <= NumMeshLODBones);
-
-		const FBoneIndexType* RESTRICT LODBoneIndexToMeshBoneIndexMapForMeshLODPtr = LODBoneIndexToMeshBoneIndexMapForMeshLOD.GetData();
-
-		for (int32 LODBoneIndex = NumPoseLODBones; LODBoneIndex < NumMeshLODBones; ++LODBoneIndex)
-		{
-			const FTransform LocalSpaceTransform = ReferencePose.GetRefPoseTransform(LODBoneIndex);
-
-			const FBoneIndexType MeshBoneIndex = LODBoneIndexToMeshBoneIndexMapForMeshLODPtr[LODBoneIndex];
-			const FBoneIndexType ParentMeshBoneIndex = MeshBoneIndexToParentMeshBoneIndexMapPtr[MeshBoneIndex];
-
-			const FTransform* RESTRICT ParentComponentSpaceTransform = ComponentSpaceData + ParentMeshBoneIndex;
-
-			FTransform* RESTRICT ComponentSpaceTransform = ComponentSpaceData + MeshBoneIndex;
-
-			FTransform::Multiply(ComponentSpaceTransform, &LocalSpaceTransform, ParentComponentSpaceTransform);
-
-			ComponentSpaceTransform->NormalizeRotation();
-
-			checkSlow(ComponentSpaceTransform->IsRotationNormalized());
-			checkSlow(!ComponentSpaceTransform->ContainsNaN());
-
-			const FMatrix44f ReferenceToLocal = (FMatrix44f)ComponentSpaceTransform->ToMatrixWithScale();
-			VectorMatrixMultiply(&ReferenceToLocalMatrices[MeshBoneIndex], &RefBasesInvMatrix[MeshBoneIndex], &ReferenceToLocal);
-		}
-	}
+	return false;
+}
 
-	// Calc bounds
-	FBoxSphereBounds WorldBounds(ForceInit);
-	FBoxSphereBounds LocalBounds(ForceInit);
-	CalcMeshBound(ComponentSpacePose, InLocalToWorld, WorldBounds, LocalBounds);
-	InFunction(ReferenceToLocalMatrices, WorldBounds, LocalBounds, SerialNumber);
-	return true;
+void FUAFOutputPoseComponent::Bind(UE::UAF::IPoseAdapter& InPoseAdapter)
+{
+	ensure(PoseAdapter == nullptr);	// Should unbind first
+	PoseAdapter = &InPoseAdapter;
 }
 
-void FUAFOutputPoseComponent::BindToHost(USkeletalMesh* InMesh, TFunction<void(void)>&& InSignalFunction)
+void FUAFOutputPoseComponent::Unbind()
 {
-	check(InMesh);
-	Mesh = InMesh;
-	SignalFunction = MoveTemp(InSignalFunction);
-	RefPose.ReferencePose = UE::UAF::FDataRegistry::Get()->GetOrGenerateReferencePose(Mesh);
+	ensure(PoseAdapter != nullptr);	// Should bind first
+	PoseAdapter = nullptr;
 }
\ No newline at end of file
diff --git a/Engine/Plugins/Experimental/UAF/UAF/Source/UAF/Private/Module/AnimNextModuleInstance.cpp b/Engine/Plugins/Experimental/UAF/UAF/Source/UAF/Private/Module/AnimNextModuleInstance.cpp
--- a/Engine/Plugins/Experimental/UAF/UAF/Source/UAF/Private/Module/AnimNextModuleInstance.cpp
+++ b/Engine/Plugins/Experimental/UAF/UAF/Source/UAF/Private/Module/AnimNextModuleInstance.cpp
@@ -486,6 +486,20 @@ const UE::UAF::FModuleEventTickFunction* FAnimNextModuleInstance::FindTickFuncti
 	return nullptr;
 }
 
+void FAnimNextModuleInstance::BeginExecution(float InDeltaTime)
+{
+	// Give each component a chance to begin execution
+	for(TInstancedStruct<FUAFAssetInstanceComponent>& InstancedStruct : Components)
+	{
+		FUAFModuleInstanceComponent* Component = InstancedStruct.GetMutablePtr<FUAFModuleInstanceComponent>();
+		if (Component == nullptr)
+		{
+			continue;
+		}
+		Component->OnBeginExecution(InDeltaTime);
+	}
+}
+
 void FAnimNextModuleInstance::EndExecution(float InDeltaTime)
 {
 	// Give the module a chance to handle events
diff --git a/Engine/Plugins/Experimental/UAF/UAF/Source/UAF/Private/Module/ModuleTaskContext.cpp b/Engine/Plugins/Experimental/UAF/UAF/Source/UAF/Private/Module/ModuleTaskContext.cpp
--- a/Engine/Plugins/Experimental/UAF/UAF/Source/UAF/Private/Module/ModuleTaskContext.cpp
+++ b/Engine/Plugins/Experimental/UAF/UAF/Source/UAF/Private/Module/ModuleTaskContext.cpp
@@ -18,6 +18,11 @@ void FModuleTaskContext::QueueInputTraitEvent(FAnimNextTraitEventPtr Event) cons
 	ModuleInstance.QueueInputTraitEvent(MoveTemp(Event));
 }
 
+void FModuleTaskContext::QueueOutputTraitEvent(FAnimNextTraitEventPtr Event) const
+{
+	ModuleInstance.QueueOutputTraitEvent(MoveTemp(Event));
+}
+
 void FModuleTaskContext::TryAccessComponent(UScriptStruct* InComponentType, TFunctionRef<void(FUAFAssetInstanceComponent&)> InFunction) const
 {
 	FUAFAssetInstanceComponent* Component = static_cast<FUAFAssetInstanceComponent*>(ModuleInstance.TryGetComponent(InComponentType));
diff --git a/Engine/Plugins/Experimental/UAF/UAF/Source/UAF/Private/Module/ModuleTickFunction.cpp b/Engine/Plugins/Experimental/UAF/UAF/Source/UAF/Private/Module/ModuleTickFunction.cpp
--- a/Engine/Plugins/Experimental/UAF/UAF/Source/UAF/Private/Module/ModuleTickFunction.cpp
+++ b/Engine/Plugins/Experimental/UAF/UAF/Source/UAF/Private/Module/ModuleTickFunction.cpp
@@ -212,6 +212,9 @@ void FModuleEventTickFunction::Run(float InDeltaTime)
 			Initialize(InDeltaTime);
 			ModuleInstance->TransitionToInitState(FAnimNextModuleInstance::EInitState::FirstUpdate);
 		}
+
+		// Call module component's extension points
+		ModuleInstance->BeginExecution(InDeltaTime);
 	}
 
 	ModuleInstance->RaiseTraitEvents(ModuleInstance->InputEventList);
diff --git a/Engine/Plugins/Experimental/UAF/UAF/Source/UAF/Public/IPoseAdapter.h b/Engine/Plugins/Experimental/UAF/UAF/Source/UAF/Public/IPoseAdapter.h
new file mode 100644
--- /dev/null
+++ b/Engine/Plugins/Experimental/UAF/UAF/Source/UAF/Public/IPoseAdapter.h
@@ -0,0 +1,50 @@
+﻿// Copyright Epic Games, Inc. All Rights Reserved.
+
+#pragma once
+
+#include "CoreTypes.h"
+#include "UObject/ObjectPtr.h"
+
+class USkeletalMesh;
+struct FAnimNextGraphLODPose;
+struct FAnimNextGraphReferencePose;
+struct FUAFOutputPoseComponent;
+
+namespace UE::UAF
+{
+struct FModuleTaskContext;
+}
+
+namespace UE::UAF
+{
+
+// Interface used by any host that wants to read a pose from a UAF system
+// Register this with FUAFOutputPoseComponent
+class IPoseAdapter
+{
+public:
+	virtual ~IPoseAdapter() = default;
+	
+	// All input data required to generate a pose
+	struct FInputData
+	{
+		// Mesh to use
+		TObjectPtr<USkeletalMesh> SkeletalMesh;
+
+		// LOD to run at
+		int32 LOD = 0;
+	};
+
+private:
+	friend struct ::FUAFOutputPoseComponent;
+
+	// Called back on a worker thread by FUAFOutputPoseComponent to retrieve the LOD and mesh to use when system execution begins
+	// The mesh & LOD is then consistent across a system's execution
+	// Consistency is validated with an ensure() in FUAFOutputPoseComponent prior to calling back via GenerateRenderData
+	virtual FInputData GetPoseInput() const = 0;
+
+	// Called back on a worker thread by FUAFOutputPoseComponent when a pose is written
+	virtual void SignalPoseWritten(const UE::UAF::FModuleTaskContext& InContext) = 0;
+};
+
+}
diff --git a/Engine/Plugins/Experimental/UAF/UAF/Source/UAF/Public/Module/ModuleTaskContext.h b/Engine/Plugins/Experimental/UAF/UAF/Source/UAF/Public/Module/ModuleTaskContext.h
--- a/Engine/Plugins/Experimental/UAF/UAF/Source/UAF/Public/Module/ModuleTaskContext.h
+++ b/Engine/Plugins/Experimental/UAF/UAF/Source/UAF/Public/Module/ModuleTaskContext.h
@@ -11,6 +11,7 @@ struct FUAFAssetInstanceComponent;
 struct FAnimNextModuleInstance;
 struct FInstancedPropertyBag;
 struct FAnimNextModuleInjectionComponent;
+struct FUAFOutputPoseComponent;
 
 namespace UE::UAF
 {
@@ -30,6 +31,10 @@ public:
 	// Input events will be processed in the next graph update after they are queued
 	UAF_API void QueueInputTraitEvent(FAnimNextTraitEventPtr Event) const;
 
+	// Queues an output trait event
+	// Output events will be processed at the end of the schedule tick
+	UAF_API void QueueOutputTraitEvent(FAnimNextTraitEventPtr Event) const;
+
 	// Access a module instance component of the specified type. If the component exists, then InFunction will be called
 	UAF_API void TryAccessComponent(UScriptStruct* InComponentType, TFunctionRef<void(FUAFAssetInstanceComponent&)> InFunction) const;
 
@@ -52,6 +57,7 @@ private:
 
 	friend UE::UAF::FModuleEventTickFunction;
 	friend ::FAnimNextModuleInjectionComponent;
+	friend ::FUAFOutputPoseComponent;
 };
 
 }
\ No newline at end of file
diff --git a/Engine/Plugins/Experimental/UAF/UAF/Source/UAF/Public/Module/UAFModuleInstanceComponent.h b/Engine/Plugins/Experimental/UAF/UAF/Source/UAF/Public/Module/UAFModuleInstanceComponent.h
--- a/Engine/Plugins/Experimental/UAF/UAF/Source/UAF/Public/Module/UAFModuleInstanceComponent.h
+++ b/Engine/Plugins/Experimental/UAF/UAF/Source/UAF/Public/Module/UAFModuleInstanceComponent.h
@@ -38,6 +38,9 @@ struct FUAFModuleInstanceComponent : public FUAFAssetInstanceComponent
 	// Called during system execution for any events to be handled
 	virtual void OnTraitEvent(FAnimNextTraitEvent& Event) {}
 
+	// Called at start of system execution each frame (this runs during the first user event - prior events, e.g. bindings will run before this)
+	virtual void OnBeginExecution(float InDeltaTime) {}
+
 	// Called at end of system execution each frame
 	virtual void OnEndExecution(float InDeltaTime) {}
 
