Create an empty workspace document for shared variables asset

#Synced-CL 48503794

diff --git a/Engine/Plugins/Experimental/UAF/UAF/Source/UAFEditor/Internal/Common/AnimNextAssetItemDetails.cpp b/Engine/Plugins/Experimental/UAF/UAF/Source/UAFEditor/Internal/Common/AnimNextAssetItemDetails.cpp
--- a/Engine/Plugins/Experimental/UAF/UAF/Source/UAFEditor/Internal/Common/AnimNextAssetItemDetails.cpp
+++ b/Engine/Plugins/Experimental/UAF/UAF/Source/UAFEditor/Internal/Common/AnimNextAssetItemDetails.cpp
@@ -84,7 +84,7 @@ bool FAnimNextAssetItemDetails::HandleDoubleClick(const FToolMenuContext& ToolMe
 				const TSharedRef<SVariablesView> VariablesView = StaticCastSharedRef<SVariablesView>(VariableDockTab->GetContent());
 				VariablesView->SetExportDirectly(SharedVariablesExport);
 
-				return true;
+				return false;
 			}
 		}
 	}
diff --git a/Engine/Plugins/Experimental/UAF/UAF/Source/UAFEditor/Private/AnimNextEditorModule.cpp b/Engine/Plugins/Experimental/UAF/UAF/Source/UAFEditor/Private/AnimNextEditorModule.cpp
--- a/Engine/Plugins/Experimental/UAF/UAF/Source/UAFEditor/Private/AnimNextEditorModule.cpp
+++ b/Engine/Plugins/Experimental/UAF/UAF/Source/UAFEditor/Private/AnimNextEditorModule.cpp
@@ -89,6 +89,7 @@
 #include "AssetToolsModule.h"
 #include "UAF/Debug/VariableWatchTabSummoner.h"
 #include "WorkspaceSchema.h"
+#include "Styling/SlateIconFinder.h"
 
 #define LOCTEXT_NAMESPACE "AnimNextEditorModule"
 
@@ -1703,6 +1704,63 @@ void FAnimNextEditorModule::RegisterWorkspaceDocumentTypes(Workspace::IWorkspace
 	});
 
 	WorkspaceEditorModule.RegisterDocumentSubObjectType(FTopLevelAssetPath(TEXT("/Script/UAFUncookedOnly.AnimNextEdGraphNode")), GraphNodeSubObjectArgs);
+
+	// Shared variables
+	{
+		Workspace::FObjectDocumentArgs DocumentArgs;
+		DocumentArgs.DocumentEditorMode = UAnimNextWorkspaceEditorMode::EM_AnimNextWorkspace;
+
+		DocumentArgs.OnMakeDocumentWidget = Workspace::FOnMakeDocumentWidget::CreateLambda([](const Workspace::FWorkspaceEditorContext& InContext) -> TSharedRef<SWidget>
+			{
+				return SNew(SBox)
+				.HAlign(HAlign_Center)
+				.VAlign(VAlign_Center)
+				[
+					SNew(STextBlock)
+					.Text(LOCTEXT("SharedVariablesDocument_HintText", "Edit variables inside the 'Variables' tab."))
+					.TextStyle(FAppStyle::Get(), "HintText")
+				];
+			});
+
+		DocumentArgs.OnGetTabIcon = Workspace::FOnGetTabIcon::CreateLambda([](const Workspace::FWorkspaceEditorContext& InContext) -> const FSlateBrush*
+			{
+				return FSlateIconFinder::FindIconForClass(UAnimNextSharedVariables::StaticClass()).GetIcon();
+			});
+
+		DocumentArgs.OnGetTabName = Workspace::FOnGetTabName::CreateLambda([](const Workspace::FWorkspaceEditorContext& InContext) -> FText
+			{
+				if (UAnimNextSharedVariables* SharedVariables = InContext.Document.GetTypedObject<UAnimNextSharedVariables>())
+				{
+					return FText::FromName(SharedVariables->GetFName());
+				}
+				else
+				{
+					return LOCTEXT("SharedVariablesDocument_EmptyTabName", "Shared Variables");
+				}
+			});
+
+		DocumentArgs.OnGetDocumentBreadcrumbTrail = Workspace::FOnGetDocumentBreadcrumbTrail::CreateLambda([](const Workspace::FWorkspaceEditorContext& InContext, TArray<TSharedPtr<Workspace::FWorkspaceBreadcrumb>>& OutBreadcrumbs)
+			{
+				if (UAnimNextSharedVariables* SharedVariables = InContext.Document.GetTypedObject<UAnimNextSharedVariables>())
+				{
+					const TSharedPtr<Workspace::FWorkspaceBreadcrumb>& Crumb = OutBreadcrumbs.Add_GetRef(MakeShared<Workspace::FWorkspaceBreadcrumb>());
+					
+					Crumb->OnGetLabel = Workspace::FWorkspaceBreadcrumb::FOnGetBreadcrumbLabel::CreateLambda([SharedVariablesName = SharedVariables->GetFName()]
+						{
+							return FText::FromName(SharedVariablesName);
+						});
+
+					Crumb->CanSave = Workspace::FWorkspaceBreadcrumb::FCanSaveBreadcrumb::CreateLambda([SharedVariables]
+						{
+							return SharedVariables->GetPackage()->IsDirty();
+						}
+					);
+				}
+			});
+
+		SharedVariablesTopLevelAssetPath = UAnimNextSharedVariables::StaticClass()->GetClassPathName();
+		WorkspaceEditorModule.RegisterObjectDocumentType(SharedVariablesTopLevelAssetPath, DocumentArgs);
+	}
 }
 
 void FAnimNextEditorModule::UnregisterWorkspaceDocumentTypes()
@@ -1712,6 +1770,7 @@ void FAnimNextEditorModule::UnregisterWorkspaceDocumentTypes()
 		Workspace::IWorkspaceEditorModule& WorkspaceEditorModule = FModuleManager::LoadModuleChecked<Workspace::IWorkspaceEditorModule>("WorkspaceEditor");
 		WorkspaceEditorModule.UnregisterObjectDocumentType(FTopLevelAssetPath(TEXT("/Script/UAF.AnimNextModule")));
 		WorkspaceEditorModule.UnregisterObjectDocumentType(FTopLevelAssetPath(TEXT("/Script/UAFUncookedOnly.AnimNextEdGraph")));
+		WorkspaceEditorModule.UnregisterObjectDocumentType(SharedVariablesTopLevelAssetPath);
 		WorkspaceEditorModule.UnregisterDocumentSubObjectType(FTopLevelAssetPath(TEXT("/Script/UAFUncookedOnly.AnimNextEdGraphNode")));
 	}
 }
diff --git a/Engine/Plugins/Experimental/UAF/UAF/Source/UAFEditor/Private/AnimNextEditorModule.h b/Engine/Plugins/Experimental/UAF/UAF/Source/UAFEditor/Private/AnimNextEditorModule.h
--- a/Engine/Plugins/Experimental/UAF/UAF/Source/UAFEditor/Private/AnimNextEditorModule.h
+++ b/Engine/Plugins/Experimental/UAF/UAF/Source/UAFEditor/Private/AnimNextEditorModule.h
@@ -71,6 +71,8 @@ private:
 	TSharedPtr<FAnimNextGraphPanelPinFactoryEditor> GraphPanelPinFactory;
 	TSharedPtr<FModuleEventPropertyTypeIdentifier> ModuleEventPropertyTypeIdentifier;
 
+	FTopLevelAssetPath SharedVariablesTopLevelAssetPath;
+
 	friend class FLocatorContext;
 };
 
