Crash when saving all data with a shared variable and reopening a UAF asset after save
#fix ensure that referenced anim-next asset (and its editor data) are pre/post-loaded
#Synced-CL 48648404

diff --git a/Engine/Plugins/Experimental/UAF/UAF/Source/UAFUncookedOnly/Private/Entries/AnimNextSharedVariablesEntry.cpp b/Engine/Plugins/Experimental/UAF/UAF/Source/UAFUncookedOnly/Private/Entries/AnimNextSharedVariablesEntry.cpp
--- a/Engine/Plugins/Experimental/UAF/UAF/Source/UAFUncookedOnly/Private/Entries/AnimNextSharedVariablesEntry.cpp
+++ b/Engine/Plugins/Experimental/UAF/UAF/Source/UAFUncookedOnly/Private/Entries/AnimNextSharedVariablesEntry.cpp
@@ -24,6 +24,9 @@ void UAnimNextSharedVariablesEntry::Initialize(UAnimNextRigVMAssetEditorData* In
 
 	if(Asset)
 	{
+		const_cast<UAnimNextSharedVariables*>(Asset.Get())->ConditionalPreload();
+		const_cast<UAnimNextSharedVariables*>(Asset.Get())->ConditionalPostLoad();
+		
 		UAnimNextSharedVariables_EditorData* EditorData = UE::UAF::UncookedOnly::FUtils::GetEditorData<UAnimNextSharedVariables_EditorData>(Asset.Get());
 		EditorData->ModifiedDelegate.AddUObject(this, &UAnimNextSharedVariablesEntry::HandleAssetModified);
 	}
