Allow binding data querying from named/typed sets

#Synced-CL 48607909

diff --git a/Engine/Plugins/Experimental/UAF/UAF/Source/UAF/Private/UAF/Attributes/AttributeNamedSet.cpp b/Engine/Plugins/Experimental/UAF/UAF/Source/UAF/Private/UAF/Attributes/AttributeNamedSet.cpp
--- a/Engine/Plugins/Experimental/UAF/UAF/Source/UAF/Private/UAF/Attributes/AttributeNamedSet.cpp
+++ b/Engine/Plugins/Experimental/UAF/UAF/Source/UAF/Private/UAF/Attributes/AttributeNamedSet.cpp
@@ -7,6 +7,11 @@ namespace UE::UAF
 {
 	FAttributeNamedSet::FAttributeNamedSet() = default;
 
+	FAttributeBindingDataPtr FAttributeNamedSet::GetBindingData() const
+	{
+		return FAttributeBindingDataPtr(Owner, true);
+	}
+
 	FName FAttributeNamedSet::GetName() const
 	{
 		return Name;
diff --git a/Engine/Plugins/Experimental/UAF/UAF/Source/UAF/Private/UAF/Attributes/AttributeTypedSet.cpp b/Engine/Plugins/Experimental/UAF/UAF/Source/UAF/Private/UAF/Attributes/AttributeTypedSet.cpp
--- a/Engine/Plugins/Experimental/UAF/UAF/Source/UAF/Private/UAF/Attributes/AttributeTypedSet.cpp
+++ b/Engine/Plugins/Experimental/UAF/UAF/Source/UAF/Private/UAF/Attributes/AttributeTypedSet.cpp
@@ -14,6 +14,11 @@ namespace UE::UAF
 		return FAttributeNamedSetPtr(NamedSet, true);
 	}
 
+	FAttributeBindingDataPtr FAttributeTypedSet::GetBindingData() const
+	{
+		return FAttributeBindingDataPtr(Owner, true);
+	}
+
 	UScriptStruct* FAttributeTypedSet::GetType() const
 	{
 		return Type;
diff --git a/Engine/Plugins/Experimental/UAF/UAF/Source/UAF/Public/UAF/Attributes/AttributeNamedSet.h b/Engine/Plugins/Experimental/UAF/UAF/Source/UAF/Public/UAF/Attributes/AttributeNamedSet.h
--- a/Engine/Plugins/Experimental/UAF/UAF/Source/UAF/Public/UAF/Attributes/AttributeNamedSet.h
+++ b/Engine/Plugins/Experimental/UAF/UAF/Source/UAF/Public/UAF/Attributes/AttributeNamedSet.h
@@ -12,6 +12,7 @@ namespace UE::UAF
 {
 	class FAttributeBindingData;
 
+	using FAttributeBindingDataPtr = TRefCountPtr<const class FAttributeBindingData>;
 	using FAttributeCollectionPtr = TSharedPtr<class FAttributeCollection, ESPMode::ThreadSafe>;
 	using FAttributeNamedSetPtr = TRefCountPtr<const class FAttributeNamedSet>;
 
@@ -63,6 +64,9 @@ namespace UE::UAF
 			TypedSetMapType::TConstIterator Iterator;
 		};
 
+		// Returns the owning binding data this named set is part of
+		[[nodiscard]] UE_API FAttributeBindingDataPtr GetBindingData() const;
+
 		// Returns the name of this set, O(1)
 		[[nodiscard]] UE_API FName GetName() const;
 
diff --git a/Engine/Plugins/Experimental/UAF/UAF/Source/UAF/Public/UAF/Attributes/AttributeTypedSet.h b/Engine/Plugins/Experimental/UAF/UAF/Source/UAF/Public/UAF/Attributes/AttributeTypedSet.h
--- a/Engine/Plugins/Experimental/UAF/UAF/Source/UAF/Public/UAF/Attributes/AttributeTypedSet.h
+++ b/Engine/Plugins/Experimental/UAF/UAF/Source/UAF/Public/UAF/Attributes/AttributeTypedSet.h
@@ -16,6 +16,7 @@ namespace UE::UAF
 	class FAttributeBindingData;
 	class FAttributeNamedSet;
 
+	using FAttributeBindingDataPtr = TRefCountPtr<const class FAttributeBindingData>;
 	using FAttributeNamedSetPtr = TRefCountPtr<const class FAttributeNamedSet>;
 	using FAttributeTypedSetPtr = TRefCountPtr<const class FAttributeTypedSet>;
 
@@ -27,6 +28,9 @@ namespace UE::UAF
 		// Returns the named set this typed set is part of
 		[[nodiscard]] UE_API FAttributeNamedSetPtr GetNamedSet() const;
 
+		// Returns the owning binding data this typed set is part of
+		[[nodiscard]] UE_API FAttributeBindingDataPtr GetBindingData() const;
+
 		// Returns the type of attributes held in this set, O(1)
 		[[nodiscard]] UE_API UScriptStruct* GetType() const;
 
