Rename collection and map iterators for consistency
Remove unnecessary value type from join operator
Move transformer specializations into Private namespace
Improve transformer assert error messages

#Synced-CL 48607909

diff --git a/Engine/Plugins/Experimental/UAF/UAF/Source/UAF/Private/UAF/Attributes/EngineAttributes.cpp b/Engine/Plugins/Experimental/UAF/UAF/Source/UAF/Private/UAF/Attributes/EngineAttributes.cpp
--- a/Engine/Plugins/Experimental/UAF/UAF/Source/UAF/Private/UAF/Attributes/EngineAttributes.cpp
+++ b/Engine/Plugins/Experimental/UAF/UAF/Source/UAF/Private/UAF/Attributes/EngineAttributes.cpp
@@ -24,16 +24,16 @@ namespace UE::UAF
 		RuntimeRegistry.RegisterMapInitializer<FStringAnimationAttribute>();
 		RuntimeRegistry.RegisterMapInitializer<FBoneTransformAnimationAttribute>();
 
-		RuntimeRegistry.RegisterAttributeValueTransformer<Transformers::FAccumulate_BoneTransformAttribute>();
-		RuntimeRegistry.RegisterAttributeValueTransformer<Transformers::FAccumulate_FloatAttribute>();
-		RuntimeRegistry.RegisterAttributeValueTransformer<Transformers::FInterpolate_BoneTransformAttribute>();
-		RuntimeRegistry.RegisterAttributeValueTransformer<Transformers::FInterpolate_FloatAttribute>();
-		RuntimeRegistry.RegisterAttributeValueTransformer<Transformers::FOverwrite_BoneTransformAttribute>();
-		RuntimeRegistry.RegisterAttributeValueTransformer<Transformers::FOverwrite_FloatAttribute>();
-		RuntimeRegistry.RegisterAttributeValueTransformer<Transformers::FSanitize_BoneTransformAttribute>();
+		RuntimeRegistry.RegisterAttributeValueTransformer<Transformers::Private::FAccumulate_BoneTransformAttribute>();
+		RuntimeRegistry.RegisterAttributeValueTransformer<Transformers::Private::FAccumulate_FloatAttribute>();
+		RuntimeRegistry.RegisterAttributeValueTransformer<Transformers::Private::FInterpolate_BoneTransformAttribute>();
+		RuntimeRegistry.RegisterAttributeValueTransformer<Transformers::Private::FInterpolate_FloatAttribute>();
+		RuntimeRegistry.RegisterAttributeValueTransformer<Transformers::Private::FOverwrite_BoneTransformAttribute>();
+		RuntimeRegistry.RegisterAttributeValueTransformer<Transformers::Private::FOverwrite_FloatAttribute>();
+		RuntimeRegistry.RegisterAttributeValueTransformer<Transformers::Private::FSanitize_BoneTransformAttribute>();
 
 #if UE_UAF_VALIDATE_MAPPED_VALUES
-		RuntimeRegistry.RegisterAttributeValueTransformer<Transformers::FSanitize_FloatAttribute>();
+		RuntimeRegistry.RegisterAttributeValueTransformer<Transformers::Private::FSanitize_FloatAttribute>();
 #endif
 	}
 }
diff --git a/Engine/Plugins/Experimental/UAF/UAF/Source/UAF/Private/UAF/Attributes/Transformers/Accumulate.cpp b/Engine/Plugins/Experimental/UAF/UAF/Source/UAF/Private/UAF/Attributes/Transformers/Accumulate.cpp
--- a/Engine/Plugins/Experimental/UAF/UAF/Source/UAF/Private/UAF/Attributes/Transformers/Accumulate.cpp
+++ b/Engine/Plugins/Experimental/UAF/UAF/Source/UAF/Private/UAF/Attributes/Transformers/Accumulate.cpp
@@ -11,7 +11,7 @@ namespace UE::UAF::Transformers
 	void FAccumulate::Apply(const FAttributeValueTransformerMapPtr& TransformerMap, const FAttributeCollection& CollectionA, const FAttributeCollection& CollectionB, float Alpha, FAttributeCollection& OutputCollection)
 	{
 		// Named sets must match as it ensures that our inputs have the same sizes/shapes
-		check(CollectionA.GetNamedSet() == CollectionB.GetNamedSet());
+		checkf(CollectionA.GetNamedSet() == CollectionB.GetNamedSet(), TEXT("Both input collections must use the same named set"));
 
 		const bool bIsOutputA = &CollectionA == &OutputCollection;
 		if (bIsOutputA && !FAnimWeight::IsRelevant(Alpha))
@@ -246,10 +246,8 @@ namespace UE::UAF::Transformers
 		}
 		else
 		{
-			check(CollectionA.GetNamedSet() == OutputCollection.GetNamedSet());
-
 			// Our output collection isn't either of our inputs, reset it
-			OutputCollection.Reset(OutputCollection.GetNamedSet());
+			OutputCollection.Reset(CollectionA.GetNamedSet());
 
 			// Set maps
 			OuterJoinBy(
@@ -376,128 +374,131 @@ namespace UE::UAF::Transformers
 		}
 	}
 
-	void FAccumulate_BoneTransformAttribute::TransformSetMap(const FAttributeSetMap* SetMapA, const FAttributeSetMap* SetMapB, float Alpha, FAttributeSetMap*& OutputSetMap, FReallocFun OutputAllocator)
+	namespace Private
 	{
-		check(SetMapA->GetTypedSet() == SetMapB->GetTypedSet());
-		check(SetMapA->Num() == SetMapB->Num());
-
-		if (OutputSetMap == nullptr)
+		void FAccumulate_BoneTransformAttribute::TransformSetMap(const FAttributeSetMap* SetMapA, const FAttributeSetMap* SetMapB, float Alpha, FAttributeSetMap*& OutputSetMap, FReallocFun OutputAllocator)
 		{
-			const FAttributeSetMap::FConstructArgs Args(SetMapA->GetTypedSet(), SetMapA->GetValueType(), OutputAllocator);
-			OutputSetMap = MakeAttributeSetMap<FBoneTransformAnimationAttribute>(Args);
-		}
+			check(SetMapA->GetTypedSet() == SetMapB->GetTypedSet());
+			check(SetMapA->Num() == SetMapB->Num());
 
-		check(SetMapA->GetTypedSet() == OutputSetMap->GetTypedSet());
-		check(SetMapA->Num() == OutputSetMap->Num());
+			if (OutputSetMap == nullptr)
+			{
+				const FAttributeSetMap::FConstructArgs Args(SetMapA->GetTypedSet(), SetMapA->GetValueType(), OutputAllocator);
+				OutputSetMap = MakeAttributeSetMap<FBoneTransformAnimationAttribute>(Args);
+			}
 
-		const TAttributeSetMap<FBoneTransformAnimationAttribute>* BoneTransformsA = CastChecked<FBoneTransformAnimationAttribute>(SetMapA);
-		const TAttributeSetMap<FBoneTransformAnimationAttribute>* BoneTransformsB = CastChecked<FBoneTransformAnimationAttribute>(SetMapB);
-		TAttributeSetMap<FBoneTransformAnimationAttribute>* OutputBoneTransforms = CastChecked<FBoneTransformAnimationAttribute>(OutputSetMap);
+			check(SetMapA->GetTypedSet() == OutputSetMap->GetTypedSet());
+			check(SetMapA->Num() == OutputSetMap->Num());
 
-		const FBoneTransformAnimationAttribute* TransformPtrA = BoneTransformsA->GetData();
-		const FBoneTransformAnimationAttribute* TransformPtrB = BoneTransformsB->GetData();
-		FBoneTransformAnimationAttribute* OutputTransformPtr = OutputBoneTransforms->GetData();
-		const FBoneTransformAnimationAttribute* TransformPtrAEnd = TransformPtrA + BoneTransformsA->Num();
+			const TAttributeSetMap<FBoneTransformAnimationAttribute>* BoneTransformsA = CastChecked<FBoneTransformAnimationAttribute>(SetMapA);
+			const TAttributeSetMap<FBoneTransformAnimationAttribute>* BoneTransformsB = CastChecked<FBoneTransformAnimationAttribute>(SetMapB);
+			TAttributeSetMap<FBoneTransformAnimationAttribute>* OutputBoneTransforms = CastChecked<FBoneTransformAnimationAttribute>(OutputSetMap);
 
-		const ScalarRegister VAlpha(Alpha);
+			const FBoneTransformAnimationAttribute* TransformPtrA = BoneTransformsA->GetData();
+			const FBoneTransformAnimationAttribute* TransformPtrB = BoneTransformsB->GetData();
+			FBoneTransformAnimationAttribute* OutputTransformPtr = OutputBoneTransforms->GetData();
+			const FBoneTransformAnimationAttribute* TransformPtrAEnd = TransformPtrA + BoneTransformsA->Num();
 
-		while (TransformPtrA < TransformPtrAEnd)
-		{
-			OutputTransformPtr->Value = TransformPtrA->Value;
-			OutputTransformPtr->Value.AccumulateWithShortestRotation(TransformPtrB->Value, VAlpha);
+			const ScalarRegister VAlpha(Alpha);
 
-			TransformPtrA++;
-			TransformPtrB++;
-			OutputTransformPtr++;
-		}
-	}
+			while (TransformPtrA < TransformPtrAEnd)
+			{
+				OutputTransformPtr->Value = TransformPtrA->Value;
+				OutputTransformPtr->Value.AccumulateWithShortestRotation(TransformPtrB->Value, VAlpha);
 
-	void FAccumulate_BoneTransformAttribute::TransformMap(const FAttributeMap* MapA, const FAttributeMap* MapB, float Alpha, FAttributeMap*& OutputMap, FReallocFun OutputAllocator)
-	{
-		check(OutputMap == nullptr);
+				TransformPtrA++;
+				TransformPtrB++;
+				OutputTransformPtr++;
+			}
+		}
 
-		const TAttributeMap<FBoneTransformAnimationAttribute>* BoneTransformsA = CastChecked<FBoneTransformAnimationAttribute>(MapA);
-		const TAttributeMap<FBoneTransformAnimationAttribute>* BoneTransformsB = CastChecked<FBoneTransformAnimationAttribute>(MapB);
-		TAttributeMap<FBoneTransformAnimationAttribute>* OutputBoneTransforms = MakeAttributeMap<FBoneTransformAnimationAttribute>(OutputAllocator);
+		void FAccumulate_BoneTransformAttribute::TransformMap(const FAttributeMap* MapA, const FAttributeMap* MapB, float Alpha, FAttributeMap*& OutputMap, FReallocFun OutputAllocator)
+		{
+			check(OutputMap == nullptr);
 
-		const ScalarRegister VAlpha(Alpha);
+			const TAttributeMap<FBoneTransformAnimationAttribute>* BoneTransformsA = CastChecked<FBoneTransformAnimationAttribute>(MapA);
+			const TAttributeMap<FBoneTransformAnimationAttribute>* BoneTransformsB = CastChecked<FBoneTransformAnimationAttribute>(MapB);
+			TAttributeMap<FBoneTransformAnimationAttribute>* OutputBoneTransforms = MakeAttributeMap<FBoneTransformAnimationAttribute>(OutputAllocator);
 
-		OuterJoinBy(
-			FAttributeMapJoinOp<FBoneTransformAnimationAttribute>(),
-			[&VAlpha, OutputBoneTransforms](FName Name, const FBoneTransformAnimationAttribute* A, const FBoneTransformAnimationAttribute* B)
-			{
-				// Missing values are set to the identity
-				const FTransform ValueA = A != nullptr ? A->Value : FTransform::Identity;
-				const FTransform ValueB = B != nullptr ? B->Value : FTransform::Identity;
+			const ScalarRegister VAlpha(Alpha);
 
-				FBoneTransformAnimationAttribute Output{ ValueA };
-				Output.Value.AccumulateWithShortestRotation(ValueB, VAlpha);
+			OuterJoinBy(
+				FAttributeMapJoinOp(),
+				[&VAlpha, OutputBoneTransforms](FName Name, const FBoneTransformAnimationAttribute* A, const FBoneTransformAnimationAttribute* B)
+				{
+					// Missing values are set to the identity
+					const FTransform& ValueA = A != nullptr ? A->Value : FTransform::Identity;
+					const FTransform& ValueB = B != nullptr ? B->Value : FTransform::Identity;
 
-				OutputBoneTransforms->Append(Name, Output);
-			},
-			BoneTransformsA->CreateConstIterator(),
-			BoneTransformsB->CreateConstIterator());
+					FBoneTransformAnimationAttribute Output{ ValueA };
+					Output.Value.AccumulateWithShortestRotation(ValueB, VAlpha);
 
-		OutputMap = OutputBoneTransforms;
-	}
+					OutputBoneTransforms->Append(Name, Output);
+				},
+				BoneTransformsA->CreateConstIterator(),
+				BoneTransformsB->CreateConstIterator());
 
-	void FAccumulate_FloatAttribute::TransformSetMap(const FAttributeSetMap* SetMapA, const FAttributeSetMap* SetMapB, float Alpha, FAttributeSetMap*& OutputSetMap, FReallocFun OutputAllocator)
-	{
-		check(SetMapA->GetTypedSet() == SetMapB->GetTypedSet());
-		check(SetMapA->Num() == SetMapB->Num());
+			OutputMap = OutputBoneTransforms;
+		}
 
-		if (OutputSetMap == nullptr)
+		void FAccumulate_FloatAttribute::TransformSetMap(const FAttributeSetMap* SetMapA, const FAttributeSetMap* SetMapB, float Alpha, FAttributeSetMap*& OutputSetMap, FReallocFun OutputAllocator)
 		{
-			const FAttributeSetMap::FConstructArgs Args(SetMapA->GetTypedSet(), SetMapA->GetValueType(), OutputAllocator);
-			OutputSetMap = MakeAttributeSetMap<FFloatAnimationAttribute>(Args);
-		}
+			check(SetMapA->GetTypedSet() == SetMapB->GetTypedSet());
+			check(SetMapA->Num() == SetMapB->Num());
 
-		check(SetMapA->GetTypedSet() == OutputSetMap->GetTypedSet());
-		check(SetMapA->Num() == OutputSetMap->Num());
+			if (OutputSetMap == nullptr)
+			{
+				const FAttributeSetMap::FConstructArgs Args(SetMapA->GetTypedSet(), SetMapA->GetValueType(), OutputAllocator);
+				OutputSetMap = MakeAttributeSetMap<FFloatAnimationAttribute>(Args);
+			}
 
-		const TAttributeSetMap<FFloatAnimationAttribute>* FloatsA = CastChecked<FFloatAnimationAttribute>(SetMapA);
-		const TAttributeSetMap<FFloatAnimationAttribute>* FloatsB = CastChecked<FFloatAnimationAttribute>(SetMapB);
-		TAttributeSetMap<FFloatAnimationAttribute>* OutputFloats = CastChecked<FFloatAnimationAttribute>(OutputSetMap);
+			check(SetMapA->GetTypedSet() == OutputSetMap->GetTypedSet());
+			check(SetMapA->Num() == OutputSetMap->Num());
 
-		const FFloatAnimationAttribute* FloatsPtrA = FloatsA->GetData();
-		const FFloatAnimationAttribute* FloatsPtrB = FloatsB->GetData();
-		FFloatAnimationAttribute* OutputFloatsPtr = OutputFloats->GetData();
-		const FFloatAnimationAttribute* FloatsPtrAEnd = FloatsPtrA + FloatsA->Num();
+			const TAttributeSetMap<FFloatAnimationAttribute>* FloatsA = CastChecked<FFloatAnimationAttribute>(SetMapA);
+			const TAttributeSetMap<FFloatAnimationAttribute>* FloatsB = CastChecked<FFloatAnimationAttribute>(SetMapB);
+			TAttributeSetMap<FFloatAnimationAttribute>* OutputFloats = CastChecked<FFloatAnimationAttribute>(OutputSetMap);
 
-		while (FloatsPtrA < FloatsPtrAEnd)
-		{
-			OutputFloatsPtr->Value = FloatsPtrA->Value + (FloatsPtrB->Value * Alpha);
+			const FFloatAnimationAttribute* FloatsPtrA = FloatsA->GetData();
+			const FFloatAnimationAttribute* FloatsPtrB = FloatsB->GetData();
+			FFloatAnimationAttribute* OutputFloatsPtr = OutputFloats->GetData();
+			const FFloatAnimationAttribute* FloatsPtrAEnd = FloatsPtrA + FloatsA->Num();
+
+			while (FloatsPtrA < FloatsPtrAEnd)
+			{
+				OutputFloatsPtr->Value = FloatsPtrA->Value + (FloatsPtrB->Value * Alpha);
 
-			FloatsPtrA++;
-			FloatsPtrB++;
-			OutputFloatsPtr++;
+				FloatsPtrA++;
+				FloatsPtrB++;
+				OutputFloatsPtr++;
+			}
 		}
-	}
 
-	void FAccumulate_FloatAttribute::TransformMap(const FAttributeMap* MapA, const FAttributeMap* MapB, float Alpha, FAttributeMap*& OutputMap, FReallocFun OutputAllocator)
-	{
-		check(OutputMap == nullptr);
+		void FAccumulate_FloatAttribute::TransformMap(const FAttributeMap* MapA, const FAttributeMap* MapB, float Alpha, FAttributeMap*& OutputMap, FReallocFun OutputAllocator)
+		{
+			check(OutputMap == nullptr);
 
-		const TAttributeMap<FFloatAnimationAttribute>* FloatsA = CastChecked<FFloatAnimationAttribute>(MapA);
-		const TAttributeMap<FFloatAnimationAttribute>* FloatsB = CastChecked<FFloatAnimationAttribute>(MapB);
-		TAttributeMap<FFloatAnimationAttribute>* OutputFloats = MakeAttributeMap<FFloatAnimationAttribute>(OutputAllocator);
+			const TAttributeMap<FFloatAnimationAttribute>* FloatsA = CastChecked<FFloatAnimationAttribute>(MapA);
+			const TAttributeMap<FFloatAnimationAttribute>* FloatsB = CastChecked<FFloatAnimationAttribute>(MapB);
+			TAttributeMap<FFloatAnimationAttribute>* OutputFloats = MakeAttributeMap<FFloatAnimationAttribute>(OutputAllocator);
 
-		OuterJoinBy(
-			FAttributeMapJoinOp<FFloatAnimationAttribute>(),
-			[&Alpha, OutputFloats](FName Name, const FFloatAnimationAttribute* A, const FFloatAnimationAttribute* B)
-			{
-				// Missing values are set to 0.0
-				const float ValueA = A != nullptr ? A->Value : 0.0f;
-				const float ValueB = B != nullptr ? B->Value : 0.0f;
+			OuterJoinBy(
+				FAttributeMapJoinOp(),
+				[&Alpha, OutputFloats](FName Name, const FFloatAnimationAttribute* A, const FFloatAnimationAttribute* B)
+				{
+					// Missing values are set to 0.0
+					const float ValueA = A != nullptr ? A->Value : 0.0f;
+					const float ValueB = B != nullptr ? B->Value : 0.0f;
 
-				FFloatAnimationAttribute Output;
-				Output.Value = ValueA + (ValueB * Alpha);
+					FFloatAnimationAttribute Output;
+					Output.Value = ValueA + (ValueB * Alpha);
 
-				OutputFloats->Append(Name, Output);
-			},
-			FloatsA->CreateConstIterator(),
-			FloatsB->CreateConstIterator());
+					OutputFloats->Append(Name, Output);
+				},
+				FloatsA->CreateConstIterator(),
+				FloatsB->CreateConstIterator());
 
-		OutputMap = OutputFloats;
+			OutputMap = OutputFloats;
+		}
 	}
 }
diff --git a/Engine/Plugins/Experimental/UAF/UAF/Source/UAF/Private/UAF/Attributes/Transformers/Interpolate.cpp b/Engine/Plugins/Experimental/UAF/UAF/Source/UAF/Private/UAF/Attributes/Transformers/Interpolate.cpp
--- a/Engine/Plugins/Experimental/UAF/UAF/Source/UAF/Private/UAF/Attributes/Transformers/Interpolate.cpp
+++ b/Engine/Plugins/Experimental/UAF/UAF/Source/UAF/Private/UAF/Attributes/Transformers/Interpolate.cpp
@@ -11,7 +11,7 @@ namespace UE::UAF::Transformers
 	void FInterpolate::Apply(const FAttributeValueTransformerMapPtr& TransformerMap, const FAttributeCollection& CollectionA, const FAttributeCollection& CollectionB, float Alpha, FAttributeCollection& OutputCollection)
 	{
 		// Named sets must match as it ensures that our inputs have the same sizes/shapes
-		check(CollectionA.GetNamedSet() == CollectionB.GetNamedSet());
+		checkf(CollectionA.GetNamedSet() == CollectionB.GetNamedSet(), TEXT("Both input collections must use the same named set"));
 
 		const bool bIsOutputA = &CollectionA == &OutputCollection;
 		if (bIsOutputA && !FAnimWeight::IsRelevant(Alpha))
@@ -252,10 +252,8 @@ namespace UE::UAF::Transformers
 		}
 		else
 		{
-			check(CollectionA.GetNamedSet() == OutputCollection.GetNamedSet());
-
 			// Our output collection isn't either of our inputs, reset it
-			OutputCollection.Reset(OutputCollection.GetNamedSet());
+			OutputCollection.Reset(CollectionA.GetNamedSet());
 
 			// Set maps
 			OuterJoinBy(
@@ -382,123 +380,126 @@ namespace UE::UAF::Transformers
 		}
 	}
 
-	void FInterpolate_BoneTransformAttribute::TransformSetMap(const FAttributeSetMap* SetMapA, const FAttributeSetMap* SetMapB, float Alpha, FAttributeSetMap*& OutputSetMap, FReallocFun OutputAllocator)
+	namespace Private
 	{
-		check(SetMapA->GetTypedSet() == SetMapB->GetTypedSet());
-		check(SetMapA->Num() == SetMapB->Num());
-
-		if (OutputSetMap == nullptr)
+		void FInterpolate_BoneTransformAttribute::TransformSetMap(const FAttributeSetMap* SetMapA, const FAttributeSetMap* SetMapB, float Alpha, FAttributeSetMap*& OutputSetMap, FReallocFun OutputAllocator)
 		{
-			const FAttributeSetMap::FConstructArgs Args(SetMapA->GetTypedSet(), SetMapA->GetValueType(), OutputAllocator);
-			OutputSetMap = MakeAttributeSetMap<FBoneTransformAnimationAttribute>(Args);
-		}
+			check(SetMapA->GetTypedSet() == SetMapB->GetTypedSet());
+			check(SetMapA->Num() == SetMapB->Num());
 
-		check(SetMapA->GetTypedSet() == OutputSetMap->GetTypedSet());
-		check(SetMapA->Num() == OutputSetMap->Num());
+			if (OutputSetMap == nullptr)
+			{
+				const FAttributeSetMap::FConstructArgs Args(SetMapA->GetTypedSet(), SetMapA->GetValueType(), OutputAllocator);
+				OutputSetMap = MakeAttributeSetMap<FBoneTransformAnimationAttribute>(Args);
+			}
 
-		const TAttributeSetMap<FBoneTransformAnimationAttribute>* BoneTransformsA = CastChecked<FBoneTransformAnimationAttribute>(SetMapA);
-		const TAttributeSetMap<FBoneTransformAnimationAttribute>* BoneTransformsB = CastChecked<FBoneTransformAnimationAttribute>(SetMapB);
-		TAttributeSetMap<FBoneTransformAnimationAttribute>* OutputBoneTransforms = CastChecked<FBoneTransformAnimationAttribute>(OutputSetMap);
+			check(SetMapA->GetTypedSet() == OutputSetMap->GetTypedSet());
+			check(SetMapA->Num() == OutputSetMap->Num());
 
-		const FBoneTransformAnimationAttribute* TransformPtrA = BoneTransformsA->GetData();
-		const FBoneTransformAnimationAttribute* TransformPtrB = BoneTransformsB->GetData();
-		FBoneTransformAnimationAttribute* OutputTransformPtr = OutputBoneTransforms->GetData();
-		const FBoneTransformAnimationAttribute* TransformPtrAEnd = TransformPtrA + BoneTransformsA->Num();
+			const TAttributeSetMap<FBoneTransformAnimationAttribute>* BoneTransformsA = CastChecked<FBoneTransformAnimationAttribute>(SetMapA);
+			const TAttributeSetMap<FBoneTransformAnimationAttribute>* BoneTransformsB = CastChecked<FBoneTransformAnimationAttribute>(SetMapB);
+			TAttributeSetMap<FBoneTransformAnimationAttribute>* OutputBoneTransforms = CastChecked<FBoneTransformAnimationAttribute>(OutputSetMap);
 
-		while (TransformPtrA < TransformPtrAEnd)
-		{
-			OutputTransformPtr->Value.Blend(TransformPtrA->Value, TransformPtrB->Value, Alpha);
+			const FBoneTransformAnimationAttribute* TransformPtrA = BoneTransformsA->GetData();
+			const FBoneTransformAnimationAttribute* TransformPtrB = BoneTransformsB->GetData();
+			FBoneTransformAnimationAttribute* OutputTransformPtr = OutputBoneTransforms->GetData();
+			const FBoneTransformAnimationAttribute* TransformPtrAEnd = TransformPtrA + BoneTransformsA->Num();
 
-			TransformPtrA++;
-			TransformPtrB++;
-			OutputTransformPtr++;
-		}
-	}
+			while (TransformPtrA < TransformPtrAEnd)
+			{
+				OutputTransformPtr->Value.Blend(TransformPtrA->Value, TransformPtrB->Value, Alpha);
 
-	void FInterpolate_BoneTransformAttribute::TransformMap(const FAttributeMap* MapA, const FAttributeMap* MapB, float Alpha, FAttributeMap*& OutputMap, FReallocFun OutputAllocator)
-	{
-		check(OutputMap == nullptr);
+				TransformPtrA++;
+				TransformPtrB++;
+				OutputTransformPtr++;
+			}
+		}
 
-		const TAttributeMap<FBoneTransformAnimationAttribute>* BoneTransformsA = CastChecked<FBoneTransformAnimationAttribute>(MapA);
-		const TAttributeMap<FBoneTransformAnimationAttribute>* BoneTransformsB = CastChecked<FBoneTransformAnimationAttribute>(MapB);
-		TAttributeMap<FBoneTransformAnimationAttribute>* OutputBoneTransforms = MakeAttributeMap<FBoneTransformAnimationAttribute>(OutputAllocator);
+		void FInterpolate_BoneTransformAttribute::TransformMap(const FAttributeMap* MapA, const FAttributeMap* MapB, float Alpha, FAttributeMap*& OutputMap, FReallocFun OutputAllocator)
+		{
+			check(OutputMap == nullptr);
 
-		OuterJoinBy(
-			FAttributeMapJoinOp<FBoneTransformAnimationAttribute>(),
-			[&Alpha, OutputBoneTransforms](FName Name, const FBoneTransformAnimationAttribute* A, const FBoneTransformAnimationAttribute* B)
-			{
-				// Missing values are set to the identity
-				const FTransform ValueA = A != nullptr ? A->Value : FTransform::Identity;
-				const FTransform ValueB = B != nullptr ? B->Value : FTransform::Identity;
+			const TAttributeMap<FBoneTransformAnimationAttribute>* BoneTransformsA = CastChecked<FBoneTransformAnimationAttribute>(MapA);
+			const TAttributeMap<FBoneTransformAnimationAttribute>* BoneTransformsB = CastChecked<FBoneTransformAnimationAttribute>(MapB);
+			TAttributeMap<FBoneTransformAnimationAttribute>* OutputBoneTransforms = MakeAttributeMap<FBoneTransformAnimationAttribute>(OutputAllocator);
 
-				FBoneTransformAnimationAttribute Output;
-				Output.Value.Blend(ValueA, ValueB, Alpha);
+			OuterJoinBy(
+				FAttributeMapJoinOp(),
+				[&Alpha, OutputBoneTransforms](FName Name, const FBoneTransformAnimationAttribute* A, const FBoneTransformAnimationAttribute* B)
+				{
+					// Missing values are set to the identity
+					const FTransform& ValueA = A != nullptr ? A->Value : FTransform::Identity;
+					const FTransform& ValueB = B != nullptr ? B->Value : FTransform::Identity;
 
-				OutputBoneTransforms->Append(Name, Output);
-			},
-			BoneTransformsA->CreateConstIterator(),
-			BoneTransformsB->CreateConstIterator());
+					FBoneTransformAnimationAttribute Output;
+					Output.Value.Blend(ValueA, ValueB, Alpha);
 
-		OutputMap = OutputBoneTransforms;
-	}
+					OutputBoneTransforms->Append(Name, Output);
+				},
+				BoneTransformsA->CreateConstIterator(),
+				BoneTransformsB->CreateConstIterator());
 
-	void FInterpolate_FloatAttribute::TransformSetMap(const FAttributeSetMap* SetMapA, const FAttributeSetMap* SetMapB, float Alpha, FAttributeSetMap*& OutputSetMap, FReallocFun OutputAllocator)
-	{
-		check(SetMapA->GetTypedSet() == SetMapB->GetTypedSet());
-		check(SetMapA->Num() == SetMapB->Num());
+			OutputMap = OutputBoneTransforms;
+		}
 
-		if (OutputSetMap == nullptr)
+		void FInterpolate_FloatAttribute::TransformSetMap(const FAttributeSetMap* SetMapA, const FAttributeSetMap* SetMapB, float Alpha, FAttributeSetMap*& OutputSetMap, FReallocFun OutputAllocator)
 		{
-			const FAttributeSetMap::FConstructArgs Args(SetMapA->GetTypedSet(), SetMapA->GetValueType(), OutputAllocator);
-			OutputSetMap = MakeAttributeSetMap<FFloatAnimationAttribute>(Args);
-		}
+			check(SetMapA->GetTypedSet() == SetMapB->GetTypedSet());
+			check(SetMapA->Num() == SetMapB->Num());
 
-		check(SetMapA->GetTypedSet() == OutputSetMap->GetTypedSet());
-		check(SetMapA->Num() == OutputSetMap->Num());
+			if (OutputSetMap == nullptr)
+			{
+				const FAttributeSetMap::FConstructArgs Args(SetMapA->GetTypedSet(), SetMapA->GetValueType(), OutputAllocator);
+				OutputSetMap = MakeAttributeSetMap<FFloatAnimationAttribute>(Args);
+			}
 
-		const TAttributeSetMap<FFloatAnimationAttribute>* FloatsA = CastChecked<FFloatAnimationAttribute>(SetMapA);
-		const TAttributeSetMap<FFloatAnimationAttribute>* FloatsB = CastChecked<FFloatAnimationAttribute>(SetMapB);
-		TAttributeSetMap<FFloatAnimationAttribute>* OutputFloats = CastChecked<FFloatAnimationAttribute>(OutputSetMap);
+			check(SetMapA->GetTypedSet() == OutputSetMap->GetTypedSet());
+			check(SetMapA->Num() == OutputSetMap->Num());
 
-		const FFloatAnimationAttribute* FloatsPtrA = FloatsA->GetData();
-		const FFloatAnimationAttribute* FloatsPtrB = FloatsB->GetData();
-		FFloatAnimationAttribute* OutputFloatsPtr = OutputFloats->GetData();
-		const FFloatAnimationAttribute* FloatsPtrAEnd = FloatsPtrA + FloatsA->Num();
+			const TAttributeSetMap<FFloatAnimationAttribute>* FloatsA = CastChecked<FFloatAnimationAttribute>(SetMapA);
+			const TAttributeSetMap<FFloatAnimationAttribute>* FloatsB = CastChecked<FFloatAnimationAttribute>(SetMapB);
+			TAttributeSetMap<FFloatAnimationAttribute>* OutputFloats = CastChecked<FFloatAnimationAttribute>(OutputSetMap);
 
-		while (FloatsPtrA < FloatsPtrAEnd)
-		{
-			OutputFloatsPtr->Value = FMath::Lerp(FloatsPtrA->Value, FloatsPtrB->Value, Alpha);
+			const FFloatAnimationAttribute* FloatsPtrA = FloatsA->GetData();
+			const FFloatAnimationAttribute* FloatsPtrB = FloatsB->GetData();
+			FFloatAnimationAttribute* OutputFloatsPtr = OutputFloats->GetData();
+			const FFloatAnimationAttribute* FloatsPtrAEnd = FloatsPtrA + FloatsA->Num();
+
+			while (FloatsPtrA < FloatsPtrAEnd)
+			{
+				OutputFloatsPtr->Value = FMath::Lerp(FloatsPtrA->Value, FloatsPtrB->Value, Alpha);
 
-			FloatsPtrA++;
-			FloatsPtrB++;
-			OutputFloatsPtr++;
+				FloatsPtrA++;
+				FloatsPtrB++;
+				OutputFloatsPtr++;
+			}
 		}
-	}
 
-	void FInterpolate_FloatAttribute::TransformMap(const FAttributeMap* MapA, const FAttributeMap* MapB, float Alpha, FAttributeMap*& OutputMap, FReallocFun OutputAllocator)
-	{
-		check(OutputMap == nullptr);
+		void FInterpolate_FloatAttribute::TransformMap(const FAttributeMap* MapA, const FAttributeMap* MapB, float Alpha, FAttributeMap*& OutputMap, FReallocFun OutputAllocator)
+		{
+			check(OutputMap == nullptr);
 
-		const TAttributeMap<FFloatAnimationAttribute>* FloatsA = CastChecked<FFloatAnimationAttribute>(MapA);
-		const TAttributeMap<FFloatAnimationAttribute>* FloatsB = CastChecked<FFloatAnimationAttribute>(MapB);
-		TAttributeMap<FFloatAnimationAttribute>* OutputFloats = MakeAttributeMap<FFloatAnimationAttribute>(OutputAllocator);
+			const TAttributeMap<FFloatAnimationAttribute>* FloatsA = CastChecked<FFloatAnimationAttribute>(MapA);
+			const TAttributeMap<FFloatAnimationAttribute>* FloatsB = CastChecked<FFloatAnimationAttribute>(MapB);
+			TAttributeMap<FFloatAnimationAttribute>* OutputFloats = MakeAttributeMap<FFloatAnimationAttribute>(OutputAllocator);
 
-		OuterJoinBy(
-			FAttributeMapJoinOp<FFloatAnimationAttribute>(),
-			[&Alpha, OutputFloats](FName Name, const FFloatAnimationAttribute* A, const FFloatAnimationAttribute* B)
-			{
-				// Missing values are set to 0.0
-				const float ValueA = A != nullptr ? A->Value : 0.0f;
-				const float ValueB = B != nullptr ? B->Value : 0.0f;
+			OuterJoinBy(
+				FAttributeMapJoinOp(),
+				[&Alpha, OutputFloats](FName Name, const FFloatAnimationAttribute* A, const FFloatAnimationAttribute* B)
+				{
+					// Missing values are set to 0.0
+					const float ValueA = A != nullptr ? A->Value : 0.0f;
+					const float ValueB = B != nullptr ? B->Value : 0.0f;
 
-				FFloatAnimationAttribute Output;
-				Output.Value = FMath::Lerp(ValueA, ValueB, Alpha);
+					FFloatAnimationAttribute Output;
+					Output.Value = FMath::Lerp(ValueA, ValueB, Alpha);
 
-				OutputFloats->Append(Name, Output);
-			},
-			FloatsA->CreateConstIterator(),
-			FloatsB->CreateConstIterator());
+					OutputFloats->Append(Name, Output);
+				},
+				FloatsA->CreateConstIterator(),
+				FloatsB->CreateConstIterator());
 
-		OutputMap = OutputFloats;
+			OutputMap = OutputFloats;
+		}
 	}
 }
diff --git a/Engine/Plugins/Experimental/UAF/UAF/Source/UAF/Private/UAF/Attributes/Transformers/Overwrite.cpp b/Engine/Plugins/Experimental/UAF/UAF/Source/UAF/Private/UAF/Attributes/Transformers/Overwrite.cpp
--- a/Engine/Plugins/Experimental/UAF/UAF/Source/UAF/Private/UAF/Attributes/Transformers/Overwrite.cpp
+++ b/Engine/Plugins/Experimental/UAF/UAF/Source/UAF/Private/UAF/Attributes/Transformers/Overwrite.cpp
@@ -10,9 +10,6 @@ namespace UE::UAF::Transformers
 
 	void FOverwrite::Apply(const FAttributeValueTransformerMapPtr& TransformerMap, const FAttributeCollection& SourceCollection, float Alpha, FAttributeCollection& OutputCollection)
 	{
-		// Named sets must match as it ensures that our inputs have the same sizes/shapes
-		check(SourceCollection.GetNamedSet() == OutputCollection.GetNamedSet());
-
 		const FAttributeValueTransformerList* TransformerList = TransformerMap->Find(TransformerName);
 		if (TransformerList == nullptr)
 		{
@@ -96,7 +93,7 @@ namespace UE::UAF::Transformers
 		else
 		{
 			// Our output collection isn't either of our inputs, reset it
-			OutputCollection.Reset(OutputCollection.GetNamedSet());
+			OutputCollection.Reset(SourceCollection.GetNamedSet());
 
 			// Set maps
 			OuterJoinBy(
@@ -168,122 +165,125 @@ namespace UE::UAF::Transformers
 		}
 	}
 
-	void FOverwrite_BoneTransformAttribute::TransformSetMap(const FAttributeSetMap* SetMap, float Alpha, FAttributeSetMap*& OutputSetMap, FReallocFun OutputAllocator)
+	namespace Private
 	{
-		if (OutputSetMap == nullptr)
+		void FOverwrite_BoneTransformAttribute::TransformSetMap(const FAttributeSetMap* SetMap, float Alpha, FAttributeSetMap*& OutputSetMap, FReallocFun OutputAllocator)
 		{
-			const FAttributeSetMap::FConstructArgs Args(SetMap->GetTypedSet(), SetMap->GetValueType(), OutputAllocator);
-			OutputSetMap = MakeAttributeSetMap<FBoneTransformAnimationAttribute>(Args);
-		}
+			if (OutputSetMap == nullptr)
+			{
+				const FAttributeSetMap::FConstructArgs Args(SetMap->GetTypedSet(), SetMap->GetValueType(), OutputAllocator);
+				OutputSetMap = MakeAttributeSetMap<FBoneTransformAnimationAttribute>(Args);
+			}
 
-		check(SetMap->GetTypedSet() == OutputSetMap->GetTypedSet());
-		check(SetMap->Num() == OutputSetMap->Num());
+			check(SetMap->GetTypedSet() == OutputSetMap->GetTypedSet());
+			check(SetMap->Num() == OutputSetMap->Num());
 
-		const TAttributeSetMap<FBoneTransformAnimationAttribute>* InputBoneTransforms = CastChecked<FBoneTransformAnimationAttribute>(SetMap);
-		TAttributeSetMap<FBoneTransformAnimationAttribute>* OutputBoneTransforms = CastChecked<FBoneTransformAnimationAttribute>(OutputSetMap);
+			const TAttributeSetMap<FBoneTransformAnimationAttribute>* InputBoneTransforms = CastChecked<FBoneTransformAnimationAttribute>(SetMap);
+			TAttributeSetMap<FBoneTransformAnimationAttribute>* OutputBoneTransforms = CastChecked<FBoneTransformAnimationAttribute>(OutputSetMap);
 
-		const FBoneTransformAnimationAttribute* InputTransformPtr = InputBoneTransforms->GetData();
-		FBoneTransformAnimationAttribute* OutputTransformPtr = OutputBoneTransforms->GetData();
-		const FBoneTransformAnimationAttribute* InputTransformPtrEnd = InputTransformPtr + InputBoneTransforms->Num();
+			const FBoneTransformAnimationAttribute* InputTransformPtr = InputBoneTransforms->GetData();
+			FBoneTransformAnimationAttribute* OutputTransformPtr = OutputBoneTransforms->GetData();
+			const FBoneTransformAnimationAttribute* InputTransformPtrEnd = InputTransformPtr + InputBoneTransforms->Num();
 
-		const ScalarRegister VAlpha(Alpha);
+			const ScalarRegister VAlpha(Alpha);
 
-		while (InputTransformPtr < InputTransformPtrEnd)
-		{
-			OutputTransformPtr->Value = InputTransformPtr->Value * VAlpha;
+			while (InputTransformPtr < InputTransformPtrEnd)
+			{
+				OutputTransformPtr->Value = InputTransformPtr->Value * VAlpha;
 
-			InputTransformPtr++;
-			OutputTransformPtr++;
+				InputTransformPtr++;
+				OutputTransformPtr++;
+			}
 		}
-	}
 
-	void FOverwrite_BoneTransformAttribute::TransformMap(const FAttributeMap* Map, float Alpha, FAttributeMap*& OutputMap, FReallocFun OutputAllocator)
-	{
-		const bool bAllocateOutput = OutputMap == nullptr;
-		if (bAllocateOutput)
+		void FOverwrite_BoneTransformAttribute::TransformMap(const FAttributeMap* Map, float Alpha, FAttributeMap*& OutputMap, FReallocFun OutputAllocator)
 		{
-			OutputMap = MakeAttributeMap<FBoneTransformAnimationAttribute>(OutputAllocator);
-		}
+			const bool bAllocateOutput = OutputMap == nullptr;
+			if (bAllocateOutput)
+			{
+				OutputMap = MakeAttributeMap<FBoneTransformAnimationAttribute>(OutputAllocator);
+			}
 
-		const TAttributeMap<FBoneTransformAnimationAttribute>* InputBoneTransforms = CastChecked<FBoneTransformAnimationAttribute>(Map);
-		TAttributeMap<FBoneTransformAnimationAttribute>* OutputBoneTransforms = CastChecked<FBoneTransformAnimationAttribute>(OutputMap);
+			const TAttributeMap<FBoneTransformAnimationAttribute>* InputBoneTransforms = CastChecked<FBoneTransformAnimationAttribute>(Map);
+			TAttributeMap<FBoneTransformAnimationAttribute>* OutputBoneTransforms = CastChecked<FBoneTransformAnimationAttribute>(OutputMap);
 
-		const ScalarRegister VAlpha(Alpha);
+			const ScalarRegister VAlpha(Alpha);
 
-		if (bAllocateOutput)
-		{
-			// We allocated a new output, append to it
-			for (auto It = InputBoneTransforms->CreateConstIterator(); It; ++It)
+			if (bAllocateOutput)
 			{
-				FBoneTransformAnimationAttribute Output{ It.GetValue()->Value * VAlpha };
-				OutputBoneTransforms->Append(It.GetName(), Output);
+				// We allocated a new output, append to it
+				for (auto It = InputBoneTransforms->CreateConstIterator(); It; ++It)
+				{
+					FBoneTransformAnimationAttribute Output{ It.GetValue()->Value * VAlpha };
+					OutputBoneTransforms->Append(It.GetName(), Output);
+				}
 			}
-		}
-		else
-		{
-			// We overwrite in place
-			check(InputBoneTransforms == OutputBoneTransforms);
-			for (auto It = OutputBoneTransforms->CreateIterator(); It; ++It)
+			else
 			{
-				It.GetValue()->Value *= VAlpha;
+				// We overwrite in place
+				check(InputBoneTransforms == OutputBoneTransforms);
+				for (auto It = OutputBoneTransforms->CreateIterator(); It; ++It)
+				{
+					It.GetValue()->Value *= VAlpha;
+				}
 			}
 		}
-	}
 
-	void FOverwrite_FloatAttribute::TransformSetMap(const FAttributeSetMap* SetMap, float Alpha, FAttributeSetMap*& OutputSetMap, FReallocFun OutputAllocator)
-	{
-		if (OutputSetMap == nullptr)
+		void FOverwrite_FloatAttribute::TransformSetMap(const FAttributeSetMap* SetMap, float Alpha, FAttributeSetMap*& OutputSetMap, FReallocFun OutputAllocator)
 		{
-			const FAttributeSetMap::FConstructArgs Args(SetMap->GetTypedSet(), SetMap->GetValueType(), OutputAllocator);
-			OutputSetMap = MakeAttributeSetMap<FFloatAnimationAttribute>(Args);
-		}
+			if (OutputSetMap == nullptr)
+			{
+				const FAttributeSetMap::FConstructArgs Args(SetMap->GetTypedSet(), SetMap->GetValueType(), OutputAllocator);
+				OutputSetMap = MakeAttributeSetMap<FFloatAnimationAttribute>(Args);
+			}
 
-		check(SetMap->GetTypedSet() == OutputSetMap->GetTypedSet());
-		check(SetMap->Num() == OutputSetMap->Num());
+			check(SetMap->GetTypedSet() == OutputSetMap->GetTypedSet());
+			check(SetMap->Num() == OutputSetMap->Num());
 
-		const TAttributeSetMap<FFloatAnimationAttribute>* InputFloats = CastChecked<FFloatAnimationAttribute>(SetMap);
-		TAttributeSetMap<FFloatAnimationAttribute>* OutputFloats = CastChecked<FFloatAnimationAttribute>(OutputSetMap);
+			const TAttributeSetMap<FFloatAnimationAttribute>* InputFloats = CastChecked<FFloatAnimationAttribute>(SetMap);
+			TAttributeSetMap<FFloatAnimationAttribute>* OutputFloats = CastChecked<FFloatAnimationAttribute>(OutputSetMap);
 
-		const FFloatAnimationAttribute* FloatsPtr = InputFloats->GetData();
-		FFloatAnimationAttribute* OutputFloatsPtr = OutputFloats->GetData();
-		const FFloatAnimationAttribute* FloatsPtrEnd = FloatsPtr + InputFloats->Num();
+			const FFloatAnimationAttribute* FloatsPtr = InputFloats->GetData();
+			FFloatAnimationAttribute* OutputFloatsPtr = OutputFloats->GetData();
+			const FFloatAnimationAttribute* FloatsPtrEnd = FloatsPtr + InputFloats->Num();
 
-		while (FloatsPtr < FloatsPtrEnd)
-		{
-			OutputFloatsPtr->Value = FloatsPtr->Value * Alpha;
+			while (FloatsPtr < FloatsPtrEnd)
+			{
+				OutputFloatsPtr->Value = FloatsPtr->Value * Alpha;
 
-			FloatsPtr++;
-			OutputFloatsPtr++;
+				FloatsPtr++;
+				OutputFloatsPtr++;
+			}
 		}
-	}
 
-	void FOverwrite_FloatAttribute::TransformMap(const FAttributeMap* Map, float Alpha, FAttributeMap*& OutputMap, FReallocFun OutputAllocator)
-	{
-		const bool bAllocateOutput = OutputMap == nullptr;
-		if (bAllocateOutput)
+		void FOverwrite_FloatAttribute::TransformMap(const FAttributeMap* Map, float Alpha, FAttributeMap*& OutputMap, FReallocFun OutputAllocator)
 		{
-			OutputMap = MakeAttributeMap<FFloatAnimationAttribute>(OutputAllocator);
-		}
+			const bool bAllocateOutput = OutputMap == nullptr;
+			if (bAllocateOutput)
+			{
+				OutputMap = MakeAttributeMap<FFloatAnimationAttribute>(OutputAllocator);
+			}
 
-		const TAttributeMap<FFloatAnimationAttribute>* InputFloats = CastChecked<FFloatAnimationAttribute>(Map);
-		TAttributeMap<FFloatAnimationAttribute>* OutputFloats = CastChecked<FFloatAnimationAttribute>(OutputMap);
+			const TAttributeMap<FFloatAnimationAttribute>* InputFloats = CastChecked<FFloatAnimationAttribute>(Map);
+			TAttributeMap<FFloatAnimationAttribute>* OutputFloats = CastChecked<FFloatAnimationAttribute>(OutputMap);
 
-		if (bAllocateOutput)
-		{
-			// We allocated a new output, append to it
-			for (auto It = InputFloats->CreateConstIterator(); It; ++It)
+			if (bAllocateOutput)
 			{
-				FFloatAnimationAttribute Output{ It.GetValue()->Value * Alpha };
-				OutputFloats->Append(It.GetName(), Output);
+				// We allocated a new output, append to it
+				for (auto It = InputFloats->CreateConstIterator(); It; ++It)
+				{
+					FFloatAnimationAttribute Output{ It.GetValue()->Value * Alpha };
+					OutputFloats->Append(It.GetName(), Output);
+				}
 			}
-		}
-		else
-		{
-			// We overwrite in place
-			check(InputFloats == OutputFloats);
-			for (auto It = OutputFloats->CreateIterator(); It; ++It)
+			else
 			{
-				It.GetValue()->Value *= Alpha;
+				// We overwrite in place
+				check(InputFloats == OutputFloats);
+				for (auto It = OutputFloats->CreateIterator(); It; ++It)
+				{
+					It.GetValue()->Value *= Alpha;
+				}
 			}
 		}
 	}
diff --git a/Engine/Plugins/Experimental/UAF/UAF/Source/UAF/Private/UAF/Attributes/Transformers/Sanitize.cpp b/Engine/Plugins/Experimental/UAF/UAF/Source/UAF/Private/UAF/Attributes/Transformers/Sanitize.cpp
--- a/Engine/Plugins/Experimental/UAF/UAF/Source/UAF/Private/UAF/Attributes/Transformers/Sanitize.cpp
+++ b/Engine/Plugins/Experimental/UAF/UAF/Source/UAF/Private/UAF/Attributes/Transformers/Sanitize.cpp
@@ -40,114 +40,117 @@ namespace UE::UAF::Transformers
 			Collection.CreateMapIterator());
 	}
 
-#if UE_UAF_VALIDATE_MAPPED_VALUES
-	// Don't inline this function to keep the stack usage down
-	static FORCENOINLINE void IsBoneTransformSetMapValid(const TAttributeSetMap<FBoneTransformAnimationAttribute>* BoneTransforms)
+	namespace Private
 	{
-		const FBoneTransformAnimationAttribute* TransformPtr = BoneTransforms->GetData();
-		const FBoneTransformAnimationAttribute* EndTransformPtr = TransformPtr + BoneTransforms->Num();
-
-		while (TransformPtr < EndTransformPtr)
+#if UE_UAF_VALIDATE_MAPPED_VALUES
+		// Don't inline this function to keep the stack usage down
+		static FORCENOINLINE void IsBoneTransformSetMapValid(const TAttributeSetMap<FBoneTransformAnimationAttribute>* BoneTransforms)
 		{
-			if (TransformPtr->Value.ContainsNaN())
-			{
-				ensureMsgf(false, TEXT("Bone (%s) contains NaN with Value:[%s]"),
-					*BoneTransforms->GetTypedSet()->GetName(FAttributeSetIndex(TransformPtr - BoneTransforms->GetData())).ToString(),
-					*TransformPtr->Value.ToString());
-			}
-			else if (!TransformPtr->Value.IsRotationNormalized())
+			const FBoneTransformAnimationAttribute* TransformPtr = BoneTransforms->GetData();
+			const FBoneTransformAnimationAttribute* EndTransformPtr = TransformPtr + BoneTransforms->Num();
+
+			while (TransformPtr < EndTransformPtr)
 			{
-				ensureMsgf(false, TEXT("Bone (%s) rotation not normalized with Rotation:[%s]"),
-					*BoneTransforms->GetTypedSet()->GetName(FAttributeSetIndex(TransformPtr - BoneTransforms->GetData())).ToString(),
-					*TransformPtr->Value.GetRotation().ToString());
+				if (TransformPtr->Value.ContainsNaN())
+				{
+					ensureMsgf(false, TEXT("Bone (%s) contains NaN with Value:[%s]"),
+						*BoneTransforms->GetTypedSet()->GetName(FAttributeSetIndex(TransformPtr - BoneTransforms->GetData())).ToString(),
+						*TransformPtr->Value.ToString());
+				}
+				else if (!TransformPtr->Value.IsRotationNormalized())
+				{
+					ensureMsgf(false, TEXT("Bone (%s) rotation not normalized with Rotation:[%s]"),
+						*BoneTransforms->GetTypedSet()->GetName(FAttributeSetIndex(TransformPtr - BoneTransforms->GetData())).ToString(),
+						*TransformPtr->Value.GetRotation().ToString());
+				}
+
+				TransformPtr++;
 			}
-
-			TransformPtr++;
 		}
-	}
 
-	// Don't inline this function to keep the stack usage down
-	static FORCENOINLINE void IsBoneTransformMapValid(const TAttributeMap<FBoneTransformAnimationAttribute>* BoneTransforms)
-	{
-		for (auto It = BoneTransforms->CreateConstIterator(); It; ++It)
+		// Don't inline this function to keep the stack usage down
+		static FORCENOINLINE void IsBoneTransformMapValid(const TAttributeMap<FBoneTransformAnimationAttribute>* BoneTransforms)
 		{
-			if (It.GetValue()->Value.ContainsNaN())
-			{
-				ensureMsgf(false, TEXT("Transform attribute (%s) contains NaN with Value:[%s]"),
-					*It.GetName().ToString(),
-					*It.GetValue()->Value.ToString());
-			}
-			else if (!It.GetValue()->Value.IsRotationNormalized())
+			for (auto It = BoneTransforms->CreateConstIterator(); It; ++It)
 			{
-				ensureMsgf(false, TEXT("Transform attribute (%s) rotation not normalized with Rotation:[%s]"),
-					*It.GetName().ToString(),
-					*It.GetValue()->Value.GetRotation().ToString());
+				if (It.GetValue()->Value.ContainsNaN())
+				{
+					ensureMsgf(false, TEXT("Transform attribute (%s) contains NaN with Value:[%s]"),
+						*It.GetName().ToString(),
+						*It.GetValue()->Value.ToString());
+				}
+				else if (!It.GetValue()->Value.IsRotationNormalized())
+				{
+					ensureMsgf(false, TEXT("Transform attribute (%s) rotation not normalized with Rotation:[%s]"),
+						*It.GetName().ToString(),
+						*It.GetValue()->Value.GetRotation().ToString());
+				}
 			}
 		}
-	}
 #endif
 
-	void FSanitize_BoneTransformAttribute::TransformSetMap(FAttributeSetMap* SetMap)
-	{
-		TAttributeSetMap<FBoneTransformAnimationAttribute>* BoneTransforms = CastChecked<FBoneTransformAnimationAttribute>(SetMap);
+		void FSanitize_BoneTransformAttribute::TransformSetMap(FAttributeSetMap* SetMap)
+		{
+			TAttributeSetMap<FBoneTransformAnimationAttribute>* BoneTransforms = CastChecked<FBoneTransformAnimationAttribute>(SetMap);
 
-		FBoneTransformAnimationAttribute* TransformPtr = BoneTransforms->GetData();
-		const FBoneTransformAnimationAttribute* TransformPtrEnd = TransformPtr + BoneTransforms->Num();
+			FBoneTransformAnimationAttribute* TransformPtr = BoneTransforms->GetData();
+			const FBoneTransformAnimationAttribute* TransformPtrEnd = TransformPtr + BoneTransforms->Num();
 
-		while (TransformPtr < TransformPtrEnd)
-		{
-			TransformPtr->Value.NormalizeRotation();
+			while (TransformPtr < TransformPtrEnd)
+			{
+				TransformPtr->Value.NormalizeRotation();
 
-			TransformPtr++;
-		}
+				TransformPtr++;
+			}
 
 #if UE_UAF_VALIDATE_MAPPED_VALUES
-		IsBoneTransformSetMapValid(BoneTransforms);
+			IsBoneTransformSetMapValid(BoneTransforms);
 #endif
-	}
-
-	void FSanitize_BoneTransformAttribute::TransformMap(FAttributeMap* Map)
-	{
-		TAttributeMap<FBoneTransformAnimationAttribute>* BoneTransforms = CastChecked<FBoneTransformAnimationAttribute>(Map);
+		}
 
-		for (auto It = BoneTransforms->CreateIterator(); It; ++It)
+		void FSanitize_BoneTransformAttribute::TransformMap(FAttributeMap* Map)
 		{
-			It.GetValue()->Value.NormalizeRotation();
-		}
+			TAttributeMap<FBoneTransformAnimationAttribute>* BoneTransforms = CastChecked<FBoneTransformAnimationAttribute>(Map);
+
+			for (auto It = BoneTransforms->CreateIterator(); It; ++It)
+			{
+				It.GetValue()->Value.NormalizeRotation();
+			}
 
 #if UE_UAF_VALIDATE_MAPPED_VALUES
-		IsBoneTransformMapValid(BoneTransforms);
+			IsBoneTransformMapValid(BoneTransforms);
 #endif
-	}
+		}
 
 #if UE_UAF_VALIDATE_MAPPED_VALUES
-	void FSanitize_FloatAttribute::TransformSetMap(FAttributeSetMap* SetMap)
-	{
-		const TAttributeSetMap<FFloatAnimationAttribute>* Floats = CastChecked<FFloatAnimationAttribute>(SetMap);
+		void FSanitize_FloatAttribute::TransformSetMap(FAttributeSetMap* SetMap)
+		{
+			const TAttributeSetMap<FFloatAnimationAttribute>* Floats = CastChecked<FFloatAnimationAttribute>(SetMap);
 
-		const FFloatAnimationAttribute* FloatsPtr = Floats->GetData();
-		const FFloatAnimationAttribute* FloatsPtrEnd = FloatsPtr + Floats->Num();
+			const FFloatAnimationAttribute* FloatsPtr = Floats->GetData();
+			const FFloatAnimationAttribute* FloatsPtrEnd = FloatsPtr + Floats->Num();
 
-		while (FloatsPtr < FloatsPtrEnd)
-		{
-			ensureMsgf(!FMath::IsNaN(FloatsPtr->Value), TEXT("Float attribute (%s) contains NaN with Value:[%3.3f]"),
-				*Floats->GetTypedSet()->GetName(FAttributeSetIndex(FloatsPtr - Floats->GetData())).ToString(),
-				FloatsPtr->Value);
+			while (FloatsPtr < FloatsPtrEnd)
+			{
+				ensureMsgf(!FMath::IsNaN(FloatsPtr->Value), TEXT("Float attribute (%s) contains NaN with Value:[%3.3f]"),
+					*Floats->GetTypedSet()->GetName(FAttributeSetIndex(FloatsPtr - Floats->GetData())).ToString(),
+					FloatsPtr->Value);
 
-			FloatsPtr++;
+				FloatsPtr++;
+			}
 		}
-	}
 
-	void FSanitize_FloatAttribute::TransformMap(FAttributeMap* Map)
-	{
-		const TAttributeMap<FFloatAnimationAttribute>* Floats = CastChecked<FFloatAnimationAttribute>(Map);
-
-		for (auto It = Floats->CreateConstIterator(); It; ++It)
+		void FSanitize_FloatAttribute::TransformMap(FAttributeMap* Map)
 		{
-			ensureMsgf(!FMath::IsNaN(It.GetValue()->Value), TEXT("Float attribute (%s) contains NaN with Value:[%3.3f]"),
-				*It.GetName().ToString(),
-				It.GetValue()->Value);
+			const TAttributeMap<FFloatAnimationAttribute>* Floats = CastChecked<FFloatAnimationAttribute>(Map);
+
+			for (auto It = Floats->CreateConstIterator(); It; ++It)
+			{
+				ensureMsgf(!FMath::IsNaN(It.GetValue()->Value), TEXT("Float attribute (%s) contains NaN with Value:[%3.3f]"),
+					*It.GetName().ToString(),
+					It.GetValue()->Value);
+			}
 		}
-	}
 #endif
+	}
 }
diff --git a/Engine/Plugins/Experimental/UAF/UAF/Source/UAF/Public/UAF/Attributes/AttributeCollection.h b/Engine/Plugins/Experimental/UAF/UAF/Source/UAF/Public/UAF/Attributes/AttributeCollection.h
--- a/Engine/Plugins/Experimental/UAF/UAF/Source/UAF/Public/UAF/Attributes/AttributeCollection.h
+++ b/Engine/Plugins/Experimental/UAF/UAF/Source/UAF/Public/UAF/Attributes/AttributeCollection.h
@@ -26,6 +26,9 @@ namespace UE::UAF
 		using EntryPtrType = typename std::conditional<bIsContainerConst, const typename ContainerType::FSetMapEntry*, typename ContainerType::FSetMapEntry*>::type;
 		using SetMapType = typename std::conditional<bIsContainerConst, const FAttributeSetMap*, FAttributeSetMap*>::type;
 
+		// Creates an empty iterator
+		[[nodiscard]] TAttributeCollectionSetMapIterator() = default;
+
 		// Creates an iterator over the set map entries within the specified collection
 		[[nodiscard]] TAttributeCollectionSetMapIterator(ContainerType& Collection);
 
@@ -65,6 +68,9 @@ namespace UE::UAF
 		using EntryPtrType = typename std::conditional<bIsContainerConst, const typename ContainerType::FMapEntry*, typename ContainerType::FMapEntry*>::type;
 		using MapType = typename std::conditional<bIsContainerConst, const FAttributeMap*, FAttributeMap*>::type;
 
+		// Creates an empty iterator
+		[[nodiscard]] TAttributeCollectionMapIterator() = default;
+
 		// Creates an iterator over the name/value map entries within the specified collection
 		[[nodiscard]] TAttributeCollectionMapIterator(ContainerType& Collection);
 
@@ -189,12 +195,12 @@ namespace UE::UAF
 		[[nodiscard]] UE_API FAttributeSetMap* FindOrAddSetMap(const FAttributeMappingKey& MappingKey);
 
 		// Iterators
-		using TSetMapIterator = TAttributeCollectionSetMapIterator<FAttributeCollection>;
-		using TConstSetMapIterator = TAttributeCollectionSetMapIterator<const FAttributeCollection>;
+		using FSetMapIterator = TAttributeCollectionSetMapIterator<FAttributeCollection>;
+		using FConstSetMapIterator = TAttributeCollectionSetMapIterator<const FAttributeCollection>;
 
 		// Returns an iterator over all set maps contained within, sorted by their key (value type first, then attribute type)
-		[[nodiscard]] TSetMapIterator CreateSetMapIterator();
-		[[nodiscard]] TConstSetMapIterator CreateConstSetMapIterator() const;
+		[[nodiscard]] FSetMapIterator CreateSetMapIterator();
+		[[nodiscard]] FConstSetMapIterator CreateConstSetMapIterator() const;
 
 		// Returns the number of name/value maps contained within
 		[[nodiscard]] UE_API int32 NumMaps() const;
@@ -248,12 +254,12 @@ namespace UE::UAF
 		[[nodiscard]] UE_API FAttributeMap* FindOrAddMap(UScriptStruct* ValueType);
 
 		// Iterators
-		using TMapIterator = TAttributeCollectionMapIterator<FAttributeCollection>;
-		using TConstMapIterator = TAttributeCollectionMapIterator<const FAttributeCollection>;
+		using FMapIterator = TAttributeCollectionMapIterator<FAttributeCollection>;
+		using FConstMapIterator = TAttributeCollectionMapIterator<const FAttributeCollection>;
 
 		// Returns an iterator over all name/value maps contained within, sorted by their key (value type)
-		[[nodiscard]] TMapIterator CreateMapIterator();
-		[[nodiscard]] TConstMapIterator CreateConstMapIterator() const;
+		[[nodiscard]] FMapIterator CreateMapIterator();
+		[[nodiscard]] FConstMapIterator CreateConstMapIterator() const;
 
 	private:
 		UE_API FAttributeSetMap** AddSetMapImpl(const FAttributeSetMap::FConstructArgs& Args, bool bConstructSetMap);
@@ -283,10 +289,10 @@ namespace UE::UAF
 		// The indirect allocator function pointer
 		FReallocFun ReallocFun = nullptr;
 
-		friend TSetMapIterator;
-		friend TConstSetMapIterator;
-		friend TMapIterator;
-		friend TConstMapIterator;
+		friend FSetMapIterator;
+		friend FConstSetMapIterator;
+		friend FMapIterator;
+		friend FConstMapIterator;
 	};
 
 	// A heap allocated attribute collection
@@ -298,6 +304,8 @@ namespace UE::UAF
 
 		// Constructs an empty collection based on the specified named set
 		explicit FAttributeCollectionHeap(const FAttributeNamedSetPtr& NamedSet);
+
+		using FAttributeCollection::operator=;
 	};
 
 	// A Mem-Stack allocated attribute collection
@@ -309,6 +317,8 @@ namespace UE::UAF
 
 		// Constructs an empty collection based on the specified named set
 		explicit FAttributeCollectionStack(const FAttributeNamedSetPtr& NamedSet);
+
+		using FAttributeCollection::operator=;
 	};
 
 	using FAttributeCollectionPtr = TSharedPtr<FAttributeCollection, ESPMode::ThreadSafe>;
@@ -318,11 +328,11 @@ namespace UE::UAF
 	{
 		const FAttributeMappingKey& GetLargestKey() const { return FAttributeMappingKey::LARGEST_VALUE; }
 
-		const FAttributeMappingKey& GetKey(const FAttributeCollection::TSetMapIterator& It) const { return It.GetKey(); }
-		const FAttributeMappingKey& GetKey(const FAttributeCollection::TConstSetMapIterator& It) const { return It.GetKey(); }
+		const FAttributeMappingKey& GetKey(const FAttributeCollection::FSetMapIterator& It) const { return It.GetKey(); }
+		const FAttributeMappingKey& GetKey(const FAttributeCollection::FConstSetMapIterator& It) const { return It.GetKey(); }
 
-		FAttributeSetMap* GetValue(const FAttributeCollection::TSetMapIterator& It) const { return It.GetMap(); }
-		const FAttributeSetMap* GetValue(const FAttributeCollection::TConstSetMapIterator& It) const { return It.GetMap(); }
+		FAttributeSetMap* GetValue(const FAttributeCollection::FSetMapIterator& It) const { return It.GetMap(); }
+		const FAttributeSetMap* GetValue(const FAttributeCollection::FConstSetMapIterator& It) const { return It.GetMap(); }
 	};
 
 	// An operator structure to join attribute collections using InnerJoin/OuterJoin/etc
@@ -333,11 +343,11 @@ namespace UE::UAF
 
 		UScriptStruct* GetLargestKey() const { return reinterpret_cast<UScriptStruct*>(LARGEST_KEY_VALUE); }
 
-		UScriptStruct* GetKey(const FAttributeCollection::TMapIterator& It) const { return It.GetValueType(); }
-		UScriptStruct* GetKey(const FAttributeCollection::TConstMapIterator& It) const { return It.GetValueType(); }
+		UScriptStruct* GetKey(const FAttributeCollection::FMapIterator& It) const { return It.GetValueType(); }
+		UScriptStruct* GetKey(const FAttributeCollection::FConstMapIterator& It) const { return It.GetValueType(); }
 
-		FAttributeMap* GetValue(const FAttributeCollection::TMapIterator& It) const { return It.GetMap(); }
-		const FAttributeMap* GetValue(const FAttributeCollection::TConstMapIterator& It) const { return It.GetMap(); }
+		FAttributeMap* GetValue(const FAttributeCollection::FMapIterator& It) const { return It.GetMap(); }
+		const FAttributeMap* GetValue(const FAttributeCollection::FConstMapIterator& It) const { return It.GetMap(); }
 	};
 }
 
@@ -447,14 +457,14 @@ namespace UE::UAF
 		return AddSetMap<ValueType>(MappingKey, Args...);
 	}
 
-	inline typename FAttributeCollection::TSetMapIterator FAttributeCollection::CreateSetMapIterator()
+	inline typename FAttributeCollection::FSetMapIterator FAttributeCollection::CreateSetMapIterator()
 	{
-		return TSetMapIterator(*this);
+		return FSetMapIterator(*this);
 	}
 
-	inline typename FAttributeCollection::TConstSetMapIterator FAttributeCollection::CreateConstSetMapIterator() const
+	inline typename FAttributeCollection::FConstSetMapIterator FAttributeCollection::CreateConstSetMapIterator() const
 	{
-		return TConstSetMapIterator(*this);
+		return FConstSetMapIterator(*this);
 	}
 
 	template <class ValueType>
@@ -508,14 +518,14 @@ namespace UE::UAF
 		return AddMap<ValueType>();
 	}
 
-	inline typename FAttributeCollection::TMapIterator FAttributeCollection::CreateMapIterator()
+	inline typename FAttributeCollection::FMapIterator FAttributeCollection::CreateMapIterator()
 	{
-		return TMapIterator(*this);
+		return FMapIterator(*this);
 	}
 
-	inline typename FAttributeCollection::TConstMapIterator FAttributeCollection::CreateConstMapIterator() const
+	inline typename FAttributeCollection::FConstMapIterator FAttributeCollection::CreateConstMapIterator() const
 	{
-		return TConstMapIterator(*this);
+		return FConstMapIterator(*this);
 	}
 
 	inline FAttributeCollectionHeap::FAttributeCollectionHeap()
@@ -555,7 +565,7 @@ namespace UE::UAF
 	template<class ContainerType>
 	inline TAttributeCollectionSetMapIterator<ContainerType>::operator bool() const
 	{
-		return EntryPtr < EntryEndPtr;
+		return EntryPtr != EntryEndPtr;
 	}
 
 	template<class ContainerType>
@@ -599,7 +609,7 @@ namespace UE::UAF
 	template<class ContainerType>
 	inline TAttributeCollectionMapIterator<ContainerType>::operator bool() const
 	{
-		return EntryPtr < EntryEndPtr;
+		return EntryPtr != EntryEndPtr;
 	}
 
 	template<class ContainerType>
diff --git a/Engine/Plugins/Experimental/UAF/UAF/Source/UAF/Public/UAF/Attributes/AttributeMap.h b/Engine/Plugins/Experimental/UAF/UAF/Source/UAF/Public/UAF/Attributes/AttributeMap.h
--- a/Engine/Plugins/Experimental/UAF/UAF/Source/UAF/Public/UAF/Attributes/AttributeMap.h
+++ b/Engine/Plugins/Experimental/UAF/UAF/Source/UAF/Public/UAF/Attributes/AttributeMap.h
@@ -115,6 +115,9 @@ namespace UE::UAF
 		using NamePtrType = typename std::conditional<bIsContainerConst, const FName*, FName*>::type;
 		using ValuePtrType = typename std::conditional<bIsContainerConst, const typename ContainerType::ValueType*, typename ContainerType::ValueType*>::type;
 
+		// Creates an empty iterator
+		[[nodiscard]] TAttributeMapIterator() = default;
+
 		// Creates an iterator over the name/value map entries within the specified attribute map
 		[[nodiscard]] TAttributeMapIterator(ContainerType& Map);
 
@@ -196,12 +199,12 @@ namespace UE::UAF
 		void SetUnsorted(int32 NumValues, GetNamePredicateType GetNamePredicate, GetValuePredicateType GetValuePredicate);
 
 		// Iterators
-		using TIterator = TAttributeMapIterator<TAttributeMap<InValueType>>;
-		using TConstIterator = TAttributeMapIterator<const TAttributeMap<InValueType>>;
+		using FIterator = TAttributeMapIterator<TAttributeMap<InValueType>>;
+		using FConstIterator = TAttributeMapIterator<const TAttributeMap<InValueType>>;
 
 		// Returns an iterator over all name/value pairs contained within, sorted by their name
-		[[nodiscard]] TIterator CreateIterator();
-		[[nodiscard]] TConstIterator CreateConstIterator() const;
+		[[nodiscard]] FIterator CreateIterator();
+		[[nodiscard]] FConstIterator CreateConstIterator() const;
 
 	protected:
 		static const int32 ValueTypeAlignment = sizeof(InValueType) >= 16 ? 16 : 8;
@@ -220,12 +223,11 @@ namespace UE::UAF
 
 		template<class T>
 		friend TAttributeMap<T>* MakeAttributeMap(FReallocFun ReallocFun);
-		friend TIterator;
-		friend TConstIterator;
+		friend FIterator;
+		friend FConstIterator;
 	};
 
 	// An operator structure to join attribute maps using InnerJoin/OuterJoin/etc
-	template<class ValueType>
 	struct FAttributeMapJoinOp
 	{
 		// We generally want the key in our predicate since the values by themselves have little meaning
@@ -237,13 +239,13 @@ namespace UE::UAF
 		}
 
 		// Special dummy FName that is larger than all other FNames in fast sorted order
-		FName GetLargestKey() const { return FName(FNameEntryId::FromUnstableInt(~0), FNameEntryId::FromUnstableInt(~0), 0x7fffffff); }
+		FName GetLargestKey() const { return FName(FNameEntryId::FromUnstableInt(MAX_int32), FNameEntryId::FromUnstableInt(MAX_int32), MAX_int32); }
 
-		FName GetKey(const typename TAttributeMap<ValueType>::TIterator& It) const { return It.GetName(); }
-		FName GetKey(const typename TAttributeMap<ValueType>::TConstIterator& It) const { return It.GetName(); }
+		template<typename IteratorType>
+		FName GetKey(const IteratorType& It) const { return It.GetName(); }
 
-		ValueType* GetValue(const typename TAttributeMap<ValueType>::TIterator& It) const { return It.GetValue(); }
-		const ValueType* GetValue(const typename TAttributeMap<ValueType>::TConstIterator& It) const { return It.GetValue(); }
+		template<typename IteratorType>
+		typename IteratorType::ValuePtrType GetValue(const IteratorType& It) const { return It.GetValue(); }
 	};
 
 	// Allocates and constructs a new attribute map for the specified value type
@@ -684,15 +686,15 @@ namespace UE::UAF
 	}
 
 	template<class InValueType>
-	inline typename TAttributeMap<InValueType>::TIterator TAttributeMap<InValueType>::CreateIterator()
+	inline typename TAttributeMap<InValueType>::FIterator TAttributeMap<InValueType>::CreateIterator()
 	{
-		return TIterator(*this);
+		return FIterator(*this);
 	}
 
 	template<class InValueType>
-	inline typename TAttributeMap<InValueType>::TConstIterator TAttributeMap<InValueType>::CreateConstIterator() const
+	inline typename TAttributeMap<InValueType>::FConstIterator TAttributeMap<InValueType>::CreateConstIterator() const
 	{
-		return TConstIterator(*this);
+		return FConstIterator(*this);
 	}
 
 	template<class InValueType>
@@ -722,8 +724,8 @@ namespace UE::UAF
 		InValueType* NewValues = reinterpret_cast<InValueType*>(NewDataPtr + NewValuesOffset);
 
 		// Relocate our name/value pairs
-		RelocateConstructItems<FName>(NewNames, OldNames, MaxValues);
-		RelocateConstructItems<InValueType>(NewValues, OldValues, MaxValues);
+		RelocateConstructItems<FName>(NewNames, OldNames, NumValues);
+		RelocateConstructItems<InValueType>(NewValues, OldValues, NumValues);
 
 		// Deallocate our old buffer
 		(*ReallocFun)(DataPtr, OldSize, 0);
@@ -767,7 +769,7 @@ namespace UE::UAF
 	template<class ContainerType>
 	inline TAttributeMapIterator<ContainerType>::operator bool() const
 	{
-		return NamePtr < NameEndPtr;
+		return NamePtr != NameEndPtr;
 	}
 
 	template<class ContainerType>
diff --git a/Engine/Plugins/Experimental/UAF/UAF/Source/UAF/Public/UAF/Attributes/Transformers/Accumulate.h b/Engine/Plugins/Experimental/UAF/UAF/Source/UAF/Public/UAF/Attributes/Transformers/Accumulate.h
--- a/Engine/Plugins/Experimental/UAF/UAF/Source/UAF/Public/UAF/Attributes/Transformers/Accumulate.h
+++ b/Engine/Plugins/Experimental/UAF/UAF/Source/UAF/Public/UAF/Attributes/Transformers/Accumulate.h
@@ -31,23 +31,26 @@ namespace UE::UAF::Transformers
 	//////////////////////////////////////////////////////////////////////////
 	// Engine built-in value type specializations
 
-	struct FAccumulate_BoneTransformAttribute
+	namespace Private
 	{
-		using TransformerType = FAccumulate;
-		using ValueType = FBoneTransformAnimationAttribute;
-
-		static void TransformSetMap(const FAttributeSetMap* SetMapA, const FAttributeSetMap* SetMapB, float Alpha, FAttributeSetMap*& OutputSetMap, FReallocFun OutputAllocator);
-		static void TransformMap(const FAttributeMap* MapA, const FAttributeMap* MapB, float Alpha, FAttributeMap*& OutputMap, FReallocFun OutputAllocator);
-	};
-
-	struct FAccumulate_FloatAttribute
-	{
-		using TransformerType = FAccumulate;
-		using ValueType = FFloatAnimationAttribute;
-
-		static void TransformSetMap(const FAttributeSetMap* SetMapA, const FAttributeSetMap* SetMapB, float Alpha, FAttributeSetMap*& OutputSetMap, FReallocFun OutputAllocator);
-		static void TransformMap(const FAttributeMap* MapA, const FAttributeMap* MapB, float Alpha, FAttributeMap*& OutputMap, FReallocFun OutputAllocator);
-	};
+		struct FAccumulate_BoneTransformAttribute
+		{
+			using TransformerType = FAccumulate;
+			using ValueType = FBoneTransformAnimationAttribute;
+
+			static void TransformSetMap(const FAttributeSetMap* SetMapA, const FAttributeSetMap* SetMapB, float Alpha, FAttributeSetMap*& OutputSetMap, FReallocFun OutputAllocator);
+			static void TransformMap(const FAttributeMap* MapA, const FAttributeMap* MapB, float Alpha, FAttributeMap*& OutputMap, FReallocFun OutputAllocator);
+		};
+
+		struct FAccumulate_FloatAttribute
+		{
+			using TransformerType = FAccumulate;
+			using ValueType = FFloatAnimationAttribute;
+
+			static void TransformSetMap(const FAttributeSetMap* SetMapA, const FAttributeSetMap* SetMapB, float Alpha, FAttributeSetMap*& OutputSetMap, FReallocFun OutputAllocator);
+			static void TransformMap(const FAttributeMap* MapA, const FAttributeMap* MapB, float Alpha, FAttributeMap*& OutputMap, FReallocFun OutputAllocator);
+		};
+	}
 }
 
 #undef UE_API
diff --git a/Engine/Plugins/Experimental/UAF/UAF/Source/UAF/Public/UAF/Attributes/Transformers/Interpolate.h b/Engine/Plugins/Experimental/UAF/UAF/Source/UAF/Public/UAF/Attributes/Transformers/Interpolate.h
--- a/Engine/Plugins/Experimental/UAF/UAF/Source/UAF/Public/UAF/Attributes/Transformers/Interpolate.h
+++ b/Engine/Plugins/Experimental/UAF/UAF/Source/UAF/Public/UAF/Attributes/Transformers/Interpolate.h
@@ -30,23 +30,26 @@ namespace UE::UAF::Transformers
 	//////////////////////////////////////////////////////////////////////////
 	// Engine built-in value type specializations
 
-	struct FInterpolate_BoneTransformAttribute
+	namespace Private
 	{
-		using TransformerType = FInterpolate;
-		using ValueType = FBoneTransformAnimationAttribute;
-
-		static void TransformSetMap(const FAttributeSetMap* SetMapA, const FAttributeSetMap* SetMapB, float Alpha, FAttributeSetMap*& OutputSetMap, FReallocFun OutputAllocator);
-		static void TransformMap(const FAttributeMap* MapA, const FAttributeMap* MapB, float Alpha, FAttributeMap*& OutputMap, FReallocFun OutputAllocator);
-	};
-
-	struct FInterpolate_FloatAttribute
-	{
-		using TransformerType = FInterpolate;
-		using ValueType = FFloatAnimationAttribute;
-
-		static void TransformSetMap(const FAttributeSetMap* SetMapA, const FAttributeSetMap* SetMapB, float Alpha, FAttributeSetMap*& OutputSetMap, FReallocFun OutputAllocator);
-		static void TransformMap(const FAttributeMap* MapA, const FAttributeMap* MapB, float Alpha, FAttributeMap*& OutputMap, FReallocFun OutputAllocator);
-	};
+		struct FInterpolate_BoneTransformAttribute
+		{
+			using TransformerType = FInterpolate;
+			using ValueType = FBoneTransformAnimationAttribute;
+
+			static void TransformSetMap(const FAttributeSetMap* SetMapA, const FAttributeSetMap* SetMapB, float Alpha, FAttributeSetMap*& OutputSetMap, FReallocFun OutputAllocator);
+			static void TransformMap(const FAttributeMap* MapA, const FAttributeMap* MapB, float Alpha, FAttributeMap*& OutputMap, FReallocFun OutputAllocator);
+		};
+
+		struct FInterpolate_FloatAttribute
+		{
+			using TransformerType = FInterpolate;
+			using ValueType = FFloatAnimationAttribute;
+
+			static void TransformSetMap(const FAttributeSetMap* SetMapA, const FAttributeSetMap* SetMapB, float Alpha, FAttributeSetMap*& OutputSetMap, FReallocFun OutputAllocator);
+			static void TransformMap(const FAttributeMap* MapA, const FAttributeMap* MapB, float Alpha, FAttributeMap*& OutputMap, FReallocFun OutputAllocator);
+		};
+	}
 }
 
 #undef UE_API
diff --git a/Engine/Plugins/Experimental/UAF/UAF/Source/UAF/Public/UAF/Attributes/Transformers/Overwrite.h b/Engine/Plugins/Experimental/UAF/UAF/Source/UAF/Public/UAF/Attributes/Transformers/Overwrite.h
--- a/Engine/Plugins/Experimental/UAF/UAF/Source/UAF/Public/UAF/Attributes/Transformers/Overwrite.h
+++ b/Engine/Plugins/Experimental/UAF/UAF/Source/UAF/Public/UAF/Attributes/Transformers/Overwrite.h
@@ -32,23 +32,26 @@ namespace UE::UAF::Transformers
 	//////////////////////////////////////////////////////////////////////////
 	// Engine built-in value type specializations
 
-	struct FOverwrite_BoneTransformAttribute
+	namespace Private
 	{
-		using TransformerType = FOverwrite;
-		using ValueType = FBoneTransformAnimationAttribute;
-
-		static void TransformSetMap(const FAttributeSetMap* SetMap, float Alpha, FAttributeSetMap*& OutputSetMap, FReallocFun OutputAllocator);
-		static void TransformMap(const FAttributeMap* Map, float Alpha, FAttributeMap*& OutputMap, FReallocFun OutputAllocator);
-	};
-
-	struct FOverwrite_FloatAttribute
-	{
-		using TransformerType = FOverwrite;
-		using ValueType = FFloatAnimationAttribute;
-
-		static void TransformSetMap(const FAttributeSetMap* SetMap, float Alpha, FAttributeSetMap*& OutputSetMap, FReallocFun OutputAllocator);
-		static void TransformMap(const FAttributeMap* Map, float Alpha, FAttributeMap*& OutputMap, FReallocFun OutputAllocator);
-	};
+		struct FOverwrite_BoneTransformAttribute
+		{
+			using TransformerType = FOverwrite;
+			using ValueType = FBoneTransformAnimationAttribute;
+
+			static void TransformSetMap(const FAttributeSetMap* SetMap, float Alpha, FAttributeSetMap*& OutputSetMap, FReallocFun OutputAllocator);
+			static void TransformMap(const FAttributeMap* Map, float Alpha, FAttributeMap*& OutputMap, FReallocFun OutputAllocator);
+		};
+
+		struct FOverwrite_FloatAttribute
+		{
+			using TransformerType = FOverwrite;
+			using ValueType = FFloatAnimationAttribute;
+
+			static void TransformSetMap(const FAttributeSetMap* SetMap, float Alpha, FAttributeSetMap*& OutputSetMap, FReallocFun OutputAllocator);
+			static void TransformMap(const FAttributeMap* Map, float Alpha, FAttributeMap*& OutputMap, FReallocFun OutputAllocator);
+		};
+	}
 }
 
 #undef UE_API
diff --git a/Engine/Plugins/Experimental/UAF/UAF/Source/UAF/Public/UAF/Attributes/Transformers/Sanitize.h b/Engine/Plugins/Experimental/UAF/UAF/Source/UAF/Public/UAF/Attributes/Transformers/Sanitize.h
--- a/Engine/Plugins/Experimental/UAF/UAF/Source/UAF/Public/UAF/Attributes/Transformers/Sanitize.h
+++ b/Engine/Plugins/Experimental/UAF/UAF/Source/UAF/Public/UAF/Attributes/Transformers/Sanitize.h
@@ -31,27 +31,30 @@ namespace UE::UAF::Transformers
 	//////////////////////////////////////////////////////////////////////////
 	// Engine built-in value type specializations
 
-	// Sanitize bone transforms (normalize rotation, ensure finite)
-	struct FSanitize_BoneTransformAttribute
+	namespace Private
 	{
-		using TransformerType = FSanitize;
-		using ValueType = FBoneTransformAnimationAttribute;
+		// Sanitize bone transforms (normalize rotation, ensure finite)
+		struct FSanitize_BoneTransformAttribute
+		{
+			using TransformerType = FSanitize;
+			using ValueType = FBoneTransformAnimationAttribute;
 
-		static void TransformSetMap(FAttributeSetMap* SetMap);
-		static void TransformMap(FAttributeMap* Map);
-	};
+			static void TransformSetMap(FAttributeSetMap* SetMap);
+			static void TransformMap(FAttributeMap* Map);
+		};
 
 #if UE_UAF_VALIDATE_MAPPED_VALUES
-	// Ensure values are not NaN (but Inf is valid)
-	struct FSanitize_FloatAttribute
-	{
-		using TransformerType = FSanitize;
-		using ValueType = FFloatAnimationAttribute;
-
-		static void TransformSetMap(FAttributeSetMap* SetMap);
-		static void TransformMap(FAttributeMap* Map);
-	};
+		// Ensure values are not NaN (but Inf is valid)
+		struct FSanitize_FloatAttribute
+		{
+			using TransformerType = FSanitize;
+			using ValueType = FFloatAnimationAttribute;
+
+			static void TransformSetMap(FAttributeSetMap* SetMap);
+			static void TransformMap(FAttributeMap* Map);
+		};
 #endif
+	}
 }
 
 #undef UE_API
diff --git a/Engine/Plugins/Experimental/UAF/UAF/Source/UAF/Public/UAF/Attributes/Transformers/Transformer.h b/Engine/Plugins/Experimental/UAF/UAF/Source/UAF/Public/UAF/Attributes/Transformers/Transformer.h
--- a/Engine/Plugins/Experimental/UAF/UAF/Source/UAF/Public/UAF/Attributes/Transformers/Transformer.h
+++ b/Engine/Plugins/Experimental/UAF/UAF/Source/UAF/Public/UAF/Attributes/Transformers/Transformer.h
@@ -145,11 +145,11 @@ namespace UE::UAF
 		FRawTransformerFunc GetValue(const FAttributeValueTransformerListIterator& It) const { return It.GetTransformerFunc(); }
 
 		// Normal FAttributeCollectionSetMapJoinOp
-		FAttributeMappingKey GetKey(const FAttributeCollection::TSetMapIterator& It) const { return It.GetKey(); }
-		FAttributeMappingKey GetKey(const FAttributeCollection::TConstSetMapIterator& It) const { return It.GetKey(); }
+		FAttributeMappingKey GetKey(const FAttributeCollection::FSetMapIterator& It) const { return It.GetKey(); }
+		FAttributeMappingKey GetKey(const FAttributeCollection::FConstSetMapIterator& It) const { return It.GetKey(); }
 
-		FAttributeSetMap* GetValue(const FAttributeCollection::TSetMapIterator& It) const { return It.GetMap(); }
-		const FAttributeSetMap* GetValue(const FAttributeCollection::TConstSetMapIterator& It) const { return It.GetMap(); }
+		FAttributeSetMap* GetValue(const FAttributeCollection::FSetMapIterator& It) const { return It.GetMap(); }
+		const FAttributeSetMap* GetValue(const FAttributeCollection::FConstSetMapIterator& It) const { return It.GetMap(); }
 	};
 
 	// An operator structure to join attribute value transformer lists and attribute collections using InnerJoin/OuterJoin/etc
@@ -166,11 +166,11 @@ namespace UE::UAF
 		FRawTransformerFunc GetValue(const FAttributeValueTransformerListIterator& It) const { return It.GetTransformerFunc(); }
 
 		// Normal FAttributeCollectionMapJoinOp
-		UScriptStruct* GetKey(const FAttributeCollection::TMapIterator& It) const { return It.GetValueType(); }
-		UScriptStruct* GetKey(const FAttributeCollection::TConstMapIterator& It) const { return It.GetValueType(); }
+		UScriptStruct* GetKey(const FAttributeCollection::FMapIterator& It) const { return It.GetValueType(); }
+		UScriptStruct* GetKey(const FAttributeCollection::FConstMapIterator& It) const { return It.GetValueType(); }
 
-		FAttributeMap* GetValue(const FAttributeCollection::TMapIterator& It) const { return It.GetMap(); }
-		const FAttributeMap* GetValue(const FAttributeCollection::TConstMapIterator& It) const { return It.GetMap(); }
+		FAttributeMap* GetValue(const FAttributeCollection::FMapIterator& It) const { return It.GetMap(); }
+		const FAttributeMap* GetValue(const FAttributeCollection::FConstMapIterator& It) const { return It.GetMap(); }
 	};
 
 	//////////////////////////////////////////////////////////////////////////
diff --git a/Engine/Plugins/Experimental/UAF/UAFTests/Source/UAFTestSuite/Private/UAF/Attributes/AttributeCollectionTest.cpp b/Engine/Plugins/Experimental/UAF/UAFTests/Source/UAFTestSuite/Private/UAF/Attributes/AttributeCollectionTest.cpp
--- a/Engine/Plugins/Experimental/UAF/UAFTests/Source/UAFTestSuite/Private/UAF/Attributes/AttributeCollectionTest.cpp
+++ b/Engine/Plugins/Experimental/UAF/UAFTests/Source/UAFTestSuite/Private/UAF/Attributes/AttributeCollectionTest.cpp
@@ -434,12 +434,12 @@ namespace UE::UAF::Tests
 
 			// Iterators (set map)
 			{
-				const auto GetSetMapKeyProjection = [](const typename FAttributeCollection::TSetMapIterator& It) { return It.GetKey(); };
-				const auto GetSetMapKeyProjectionConst = [](const typename FAttributeCollection::TConstSetMapIterator& It) { return It.GetKey(); };
-				const auto GetSetMapAttributeTypeProjection = [](const typename FAttributeCollection::TSetMapIterator& It) { return It.GetAttributeType(); };
-				const auto GetSetMapAttributeTypeProjectionConst = [](const typename FAttributeCollection::TConstSetMapIterator& It) { return It.GetAttributeType(); };
-				const auto GetSetMapProjection = [](const typename FAttributeCollection::TSetMapIterator& It) { return It.GetMap(); };
-				const auto GetSetMapProjectionConst = [](const typename FAttributeCollection::TConstSetMapIterator& It) { return It.GetMap(); };
+				const auto GetSetMapKeyProjection = [](const typename FAttributeCollection::FSetMapIterator& It) { return It.GetKey(); };
+				const auto GetSetMapKeyProjectionConst = [](const typename FAttributeCollection::FConstSetMapIterator& It) { return It.GetKey(); };
+				const auto GetSetMapAttributeTypeProjection = [](const typename FAttributeCollection::FSetMapIterator& It) { return It.GetAttributeType(); };
+				const auto GetSetMapAttributeTypeProjectionConst = [](const typename FAttributeCollection::FConstSetMapIterator& It) { return It.GetAttributeType(); };
+				const auto GetSetMapProjection = [](const typename FAttributeCollection::FSetMapIterator& It) { return It.GetMap(); };
+				const auto GetSetMapProjectionConst = [](const typename FAttributeCollection::FConstSetMapIterator& It) { return It.GetMap(); };
 
 				AddErrorIfFalse(IteratorSize(Collection.CreateSetMapIterator()) == 5, TEXT("Iterator did not have expected size"));
 				AddErrorIfFalse(IteratorSize(Collection.CreateConstSetMapIterator()) == 5, TEXT("Iterator did not have expected size"));
@@ -474,10 +474,10 @@ namespace UE::UAF::Tests
 
 			// Iterators (name/value map)
 			{
-				const auto GetMapValueTypeProjection = [](const typename FAttributeCollection::TMapIterator& It) { return It.GetValueType(); };
-				const auto GetMapValueTypeProjectionConst = [](const typename FAttributeCollection::TConstMapIterator& It) { return It.GetValueType(); };
-				const auto GetMapProjection = [](const typename FAttributeCollection::TMapIterator& It) { return It.GetMap(); };
-				const auto GetMapProjectionConst = [](const typename FAttributeCollection::TConstMapIterator& It) { return It.GetMap(); };
+				const auto GetMapValueTypeProjection = [](const typename FAttributeCollection::FMapIterator& It) { return It.GetValueType(); };
+				const auto GetMapValueTypeProjectionConst = [](const typename FAttributeCollection::FConstMapIterator& It) { return It.GetValueType(); };
+				const auto GetMapProjection = [](const typename FAttributeCollection::FMapIterator& It) { return It.GetMap(); };
+				const auto GetMapProjectionConst = [](const typename FAttributeCollection::FConstMapIterator& It) { return It.GetMap(); };
 
 				AddErrorIfFalse(IteratorSize(Collection.CreateMapIterator()) == 2, TEXT("Iterator did not have expected size"));
 				AddErrorIfFalse(IteratorSize(Collection.CreateConstMapIterator()) == 2, TEXT("Iterator did not have expected size"));
diff --git a/Engine/Plugins/Experimental/UAF/UAFTests/Source/UAFTestSuite/Private/UAF/Attributes/AttributeMapTest.cpp b/Engine/Plugins/Experimental/UAF/UAFTests/Source/UAFTestSuite/Private/UAF/Attributes/AttributeMapTest.cpp
--- a/Engine/Plugins/Experimental/UAF/UAFTests/Source/UAFTestSuite/Private/UAF/Attributes/AttributeMapTest.cpp
+++ b/Engine/Plugins/Experimental/UAF/UAFTests/Source/UAFTestSuite/Private/UAF/Attributes/AttributeMapTest.cpp
@@ -274,10 +274,10 @@ namespace UE::UAF::Tests
 
 			// Iterators
 			{
-				const auto GetNameProjection = [](const typename TAttributeMap<FFloatAnimationAttribute>::TIterator& It) { return It.GetName(); };
-				const auto GetNameProjectionConst = [](const typename TAttributeMap<FFloatAnimationAttribute>::TConstIterator& It) { return It.GetName(); };
-				const auto GetValueProjection = [](const typename TAttributeMap<FFloatAnimationAttribute>::TIterator& It) { return *It.GetValue(); };
-				const auto GetValueProjectionConst = [](const typename TAttributeMap<FFloatAnimationAttribute>::TConstIterator& It) { return *It.GetValue(); };
+				const auto GetNameProjection = [](const typename TAttributeMap<FFloatAnimationAttribute>::FIterator& It) { return It.GetName(); };
+				const auto GetNameProjectionConst = [](const typename TAttributeMap<FFloatAnimationAttribute>::FConstIterator& It) { return It.GetName(); };
+				const auto GetValueProjection = [](const typename TAttributeMap<FFloatAnimationAttribute>::FIterator& It) { return *It.GetValue(); };
+				const auto GetValueProjectionConst = [](const typename TAttributeMap<FFloatAnimationAttribute>::FConstIterator& It) { return *It.GetValue(); };
 				const auto TestValuePredicate = [](const FFloatAnimationAttribute& LHS, const FFloatAnimationAttribute& RHS) { return LHS.Value == RHS.Value; };
 
 				AddErrorIfFalse(IteratorSize(FloatMap.CreateIterator()) == 5, TEXT("Iterator did not have expected size"));
@@ -568,10 +568,10 @@ namespace UE::UAF::Tests
 
 			// Iterators
 			{
-				const auto GetNameProjection = [](const typename TAttributeMap<FBoneTransformAnimationAttribute>::TIterator& It) { return It.GetName(); };
-				const auto GetNameProjectionConst = [](const typename TAttributeMap<FBoneTransformAnimationAttribute>::TConstIterator& It) { return It.GetName(); };
-				const auto GetValueProjection = [](const typename TAttributeMap<FBoneTransformAnimationAttribute>::TIterator& It) { return *It.GetValue(); };
-				const auto GetValueProjectionConst = [](const typename TAttributeMap<FBoneTransformAnimationAttribute>::TConstIterator& It) { return *It.GetValue(); };
+				const auto GetNameProjection = [](const typename TAttributeMap<FBoneTransformAnimationAttribute>::FIterator& It) { return It.GetName(); };
+				const auto GetNameProjectionConst = [](const typename TAttributeMap<FBoneTransformAnimationAttribute>::FConstIterator& It) { return It.GetName(); };
+				const auto GetValueProjection = [](const typename TAttributeMap<FBoneTransformAnimationAttribute>::FIterator& It) { return *It.GetValue(); };
+				const auto GetValueProjectionConst = [](const typename TAttributeMap<FBoneTransformAnimationAttribute>::FConstIterator& It) { return *It.GetValue(); };
 				const auto TestValuePredicate = [](const FBoneTransformAnimationAttribute& LHS, const FBoneTransformAnimationAttribute& RHS) { return LHS.Value.Equals(RHS.Value); };
 
 				AddErrorIfFalse(IteratorSize(BoneTransformMap.CreateIterator()) == 5, TEXT("Iterator did not have expected size"));
@@ -863,10 +863,10 @@ namespace UE::UAF::Tests
 
 			// Iterators
 			{
-				const auto GetNameProjection = [](const typename TAttributeMap<FStringAnimationAttribute>::TIterator& It) { return It.GetName(); };
-				const auto GetNameProjectionConst = [](const typename TAttributeMap<FStringAnimationAttribute>::TConstIterator& It) { return It.GetName(); };
-				const auto GetValueProjection = [](const typename TAttributeMap<FStringAnimationAttribute>::TIterator& It) { return *It.GetValue(); };
-				const auto GetValueProjectionConst = [](const typename TAttributeMap<FStringAnimationAttribute>::TConstIterator& It) { return *It.GetValue(); };
+				const auto GetNameProjection = [](const typename TAttributeMap<FStringAnimationAttribute>::FIterator& It) { return It.GetName(); };
+				const auto GetNameProjectionConst = [](const typename TAttributeMap<FStringAnimationAttribute>::FConstIterator& It) { return It.GetName(); };
+				const auto GetValueProjection = [](const typename TAttributeMap<FStringAnimationAttribute>::FIterator& It) { return *It.GetValue(); };
+				const auto GetValueProjectionConst = [](const typename TAttributeMap<FStringAnimationAttribute>::FConstIterator& It) { return *It.GetValue(); };
 				const auto TestValuePredicate = [](const FStringAnimationAttribute& LHS, const FStringAnimationAttribute& RHS) { return LHS.Value == RHS.Value; };
 
 				AddErrorIfFalse(IteratorSize(StringMap.CreateIterator()) == 5, TEXT("Iterator did not have expected size"));
