UAF: Refactor to input variables - map variables across correctly to function reference nodes.

#rnx
#Synced-CL 48723628

diff --git a/Engine/Plugins/Experimental/UAF/UAF/Source/UAFEditor/Private/Common/GraphEditorSchemaActions.cpp b/Engine/Plugins/Experimental/UAF/UAF/Source/UAFEditor/Private/Common/GraphEditorSchemaActions.cpp
--- a/Engine/Plugins/Experimental/UAF/UAF/Source/UAFEditor/Private/Common/GraphEditorSchemaActions.cpp
+++ b/Engine/Plugins/Experimental/UAF/UAF/Source/UAFEditor/Private/Common/GraphEditorSchemaActions.cpp
@@ -462,6 +462,15 @@ UEdGraphNode* FAnimNextSchemaAction_Function::PerformAction(UEdGraph* ParentGrap
 				Controller->SelectNode(ModelNode, true, true, true);
 			}
 
+			for (const URigVMPin* Pin : ModelNode->GetPins())
+			{
+				if (!Pin->IsDefinedAsInputVariable())
+				{
+					continue;
+				}
+				Controller->BindPinToVariable(Pin->GetPinPath(), Pin->GetName());
+			}
+
 			Controller->CloseUndoBracket();
 		}
 		else
diff --git a/Engine/Plugins/Experimental/UAF/UAF/Source/UAFUncookedOnly/Private/AnimNextRigVMAssetEditorData.cpp b/Engine/Plugins/Experimental/UAF/UAF/Source/UAFUncookedOnly/Private/AnimNextRigVMAssetEditorData.cpp
--- a/Engine/Plugins/Experimental/UAF/UAF/Source/UAFUncookedOnly/Private/AnimNextRigVMAssetEditorData.cpp
+++ b/Engine/Plugins/Experimental/UAF/UAF/Source/UAFUncookedOnly/Private/AnimNextRigVMAssetEditorData.cpp
@@ -930,6 +930,18 @@ void UAnimNextRigVMAssetEditorData::BuildFunctionWrapperEvents(FAnimNextRigVMAss
 			continue;
 		}
 
+		if (!ExternalVariables.IsEmpty())
+		{
+			for (const URigVMPin* Pin : FunctionNode->GetPins())
+			{
+				if (!Pin->IsDefinedAsInputVariable())
+				{
+					continue;
+				}
+				Controller->BindPinToVariable(Pin->GetPinPath(), Pin->GetName());
+			}
+		}
+
 		// Link up Execute nodes if needed, function may be pure & lack an input pin
 		URigVMPin* CurrentExecuteOutputPin = EventNode->FindPin(FRigVMStruct::ExecuteContextName.ToString());
 		URigVMPin* ExecuteInputPin = FunctionNode->FindPin(FRigVMStruct::ExecuteContextName.ToString());
@@ -954,9 +966,10 @@ void UAnimNextRigVMAssetEditorData::BuildFunctionWrapperEvents(FAnimNextRigVMAss
 
 			bool bIsGetter = Argument.Direction == ERigVMPinDirection::Input;
 			FName InternalArgName = FName(FUtils::MakeFunctionWrapperVariableName(FunctionHeader.Name, Argument.Name));
-			int32 VariableIndex = ExternalVariables.IndexOfByPredicate([&InternalArgName](const FRigVMExternalVariable& InVariable)
+			FName SanitizedInternalArgName = FRigVMPropertyDescription::SanitizeName(InternalArgName);
+			int32 VariableIndex = ExternalVariables.IndexOfByPredicate([&SanitizedInternalArgName](const FRigVMExternalVariable& InVariable)
 			{
-				return InVariable.GetName() == InternalArgName;
+				return InVariable.GetName() == SanitizedInternalArgName;
 			});
 
 			if (VariableIndex == INDEX_NONE)
@@ -977,13 +990,13 @@ void UAnimNextRigVMAssetEditorData::BuildFunctionWrapperEvents(FAnimNextRigVMAss
 
 			if (bIsGetter)
 			{
-				URigVMVariableNode* FunctionParamVariableNode = Controller->AddVariableNode(InternalArgName
+				URigVMVariableNode* FunctionParamVariableNode = Controller->AddVariableNode(SanitizedInternalArgName
 					, Argument.CPPType.ToString()
 					, Argument.CPPTypeObject.Get()
 					, bIsGetter
 					, Argument.DefaultValue
 					, FVector2D::ZeroVector
-					, InternalArgName.ToString()
+					, SanitizedInternalArgName.ToString()
 					, bSetupUndoRedo);
 
 				if (!FunctionParamVariableNode)
@@ -1005,13 +1018,13 @@ void UAnimNextRigVMAssetEditorData::BuildFunctionWrapperEvents(FAnimNextRigVMAss
 			if (!bIsGetter)
 			{
 				
-				URigVMVariableNode* FunctionResultVariableNode = Controller->AddVariableNode(InternalArgName
+				URigVMVariableNode* FunctionResultVariableNode = Controller->AddVariableNode(SanitizedInternalArgName
 					, Argument.CPPType.ToString()
 					, Argument.CPPTypeObject.Get()
 					, bIsGetter
 					, Argument.DefaultValue
 					, FVector2D::ZeroVector
-					, InternalArgName.ToString()
+					, SanitizedInternalArgName.ToString()
 					, bSetupUndoRedo);
 
 				if (!FunctionResultVariableNode)
diff --git a/Engine/Plugins/Runtime/RigVM/Source/RigVM/Private/RigVMCore/RigVMGraphFunctionDefinition.cpp b/Engine/Plugins/Runtime/RigVM/Source/RigVM/Private/RigVMCore/RigVMGraphFunctionDefinition.cpp
--- a/Engine/Plugins/Runtime/RigVM/Source/RigVM/Private/RigVMCore/RigVMGraphFunctionDefinition.cpp
+++ b/Engine/Plugins/Runtime/RigVM/Source/RigVM/Private/RigVMCore/RigVMGraphFunctionDefinition.cpp
@@ -346,6 +346,11 @@ FArchive& operator<<(FArchive& Ar, FRigVMGraphFunctionHeader& Data)
 		{
 			Data.Layout.Reset();
 		}
+
+		Data.ExternalVariables.RemoveAll([](const FRigVMExternalVariable& ExternalVariable)
+		{
+			return !ExternalVariable.IsValid(true /* allow nullptr memory */);
+		});
 	}
 	else
 	{
@@ -889,8 +894,13 @@ bool FRigVMGraphFunctionData::PatchExternalVariablesToArguments()
 	
 	for (const FRigVMExternalVariable& ExternalVariable : Header.ExternalVariables)
 	{
-		// convert to external variable
+		if (!ExternalVariable.IsValid(true))
+		{
+			continue;
+		}
+		// convert external variable to argument
 		Header.Arguments.Emplace(ExternalVariable);
+		check(Header.Arguments.Last().IsValid());
 	}
 	Header.ExternalVariables.Reset();
 	
