Add value space support to UAF value runtime

#Synced-CL 48607909

diff --git a/Engine/Plugins/Experimental/UAF/UAF/Source/UAF/Private/UAF/Attributes/AttributeBindingData.cpp b/Engine/Plugins/Experimental/UAF/UAF/Source/UAF/Private/UAF/Attributes/AttributeBindingData.cpp
--- a/Engine/Plugins/Experimental/UAF/UAF/Source/UAF/Private/UAF/Attributes/AttributeBindingData.cpp
+++ b/Engine/Plugins/Experimental/UAF/UAF/Source/UAF/Private/UAF/Attributes/AttributeBindingData.cpp
@@ -344,6 +344,7 @@ namespace UE::UAF
 				const FName SetName = NamedSet->GetName();
 
 				FAttributeCollectionPtr DefaultAttributeValues = MakeShared<FAttributeCollectionHeap>(NamedSet);
+				DefaultAttributeValues->SetValueSpace(FValueSpace(EValueSpaceType::Local));
 
 				for (auto TypedSetIt = NamedSet->CreateTypedSetIterator(); TypedSetIt; ++TypedSetIt)
 				{
diff --git a/Engine/Plugins/Experimental/UAF/UAF/Source/UAF/Private/UAF/Attributes/AttributeCollection.cpp b/Engine/Plugins/Experimental/UAF/UAF/Source/UAF/Private/UAF/Attributes/AttributeCollection.cpp
--- a/Engine/Plugins/Experimental/UAF/UAF/Source/UAF/Private/UAF/Attributes/AttributeCollection.cpp
+++ b/Engine/Plugins/Experimental/UAF/UAF/Source/UAF/Private/UAF/Attributes/AttributeCollection.cpp
@@ -9,6 +9,7 @@ namespace UE::UAF
 {
 	FAttributeCollection::FAttributeCollection(const FAttributeCollection& Other)
 		: NamedSet(Other.NamedSet)
+		, ValueSpace(Other.ValueSpace)
 		, SetMapCount(Other.SetMapCount)
 		, SetMapCapacity(Other.SetMapCount)
 		, MapCount(Other.MapCount)
@@ -50,6 +51,7 @@ namespace UE::UAF
 
 	FAttributeCollection::FAttributeCollection(FAttributeCollection&& Other)
 		: NamedSet(MoveTemp(Other.NamedSet))
+		, ValueSpace(Other.ValueSpace)
 		, SetMaps(Other.SetMaps)
 		, SetMapCount(Other.SetMapCount)
 		, SetMapCapacity(Other.SetMapCapacity)
@@ -110,6 +112,7 @@ namespace UE::UAF
 			this->~FAttributeCollection();
 
 			NamedSet = Other.NamedSet;
+			ValueSpace = Other.ValueSpace;
 			SetMapCount = Other.SetMapCount;
 			SetMapCapacity = Other.SetMapCount;
 
@@ -162,6 +165,7 @@ namespace UE::UAF
 			this->~FAttributeCollection();
 
 			NamedSet = MoveTemp(Other.NamedSet);
+			ValueSpace = Other.ValueSpace;
 			SetMaps = Other.SetMaps;
 			SetMapCount = Other.SetMapCount;
 			SetMapCapacity = Other.SetMapCapacity;
@@ -190,6 +194,21 @@ namespace UE::UAF
 		return NamedSet;
 	}
 
+	FValueSpace FAttributeCollection::GetValueSpace() const
+	{
+		return ValueSpace;
+	}
+
+	void FAttributeCollection::SetValueSpace(FValueSpace InValueSpace)
+	{
+		ValueSpace = InValueSpace;
+	}
+
+	bool FAttributeCollection::IsAdditive() const
+	{
+		return ValueSpace.IsAdditive();
+	}
+
 	void FAttributeCollection::Reset(const FAttributeNamedSetPtr& InNamedSet)
 	{
 		// Free our maps but retain the arrays that contains them
@@ -206,6 +225,7 @@ namespace UE::UAF
 		MapCount = 0;
 
 		NamedSet = InNamedSet;
+		ValueSpace = FValueSpace();
 	}
 
 	FReallocFun FAttributeCollection::GetAllocator() const
diff --git a/Engine/Plugins/Experimental/UAF/UAF/Source/UAF/Private/UAF/Attributes/Transformers/Accumulate.cpp b/Engine/Plugins/Experimental/UAF/UAF/Source/UAF/Private/UAF/Attributes/Transformers/Accumulate.cpp
--- a/Engine/Plugins/Experimental/UAF/UAF/Source/UAF/Private/UAF/Attributes/Transformers/Accumulate.cpp
+++ b/Engine/Plugins/Experimental/UAF/UAF/Source/UAF/Private/UAF/Attributes/Transformers/Accumulate.cpp
@@ -12,6 +12,7 @@ namespace UE::UAF::Transformers
 	{
 		// Named sets must match as it ensures that our inputs have the same sizes/shapes
 		checkf(CollectionA.GetNamedSet() == CollectionB.GetNamedSet(), TEXT("Both input collections must use the same named set"));
+		checkf(CollectionA.GetValueSpace() == CollectionB.GetValueSpace(), TEXT("Both input collections must have matching value spaces: [%s] != [%s]"), *CollectionA.GetValueSpace().ToString(), *CollectionB.GetValueSpace().ToString());
 
 		const bool bIsOutputA = &CollectionA == &OutputCollection;
 		if (bIsOutputA && !FAnimWeight::IsRelevant(Alpha))
@@ -248,6 +249,7 @@ namespace UE::UAF::Transformers
 		{
 			// Our output collection isn't either of our inputs, reset it
 			OutputCollection.Reset(CollectionA.GetNamedSet());
+			OutputCollection.SetValueSpace(CollectionA.GetValueSpace());
 
 			// Set maps
 			OuterJoinBy(
diff --git a/Engine/Plugins/Experimental/UAF/UAF/Source/UAF/Private/UAF/Attributes/Transformers/Interpolate.cpp b/Engine/Plugins/Experimental/UAF/UAF/Source/UAF/Private/UAF/Attributes/Transformers/Interpolate.cpp
--- a/Engine/Plugins/Experimental/UAF/UAF/Source/UAF/Private/UAF/Attributes/Transformers/Interpolate.cpp
+++ b/Engine/Plugins/Experimental/UAF/UAF/Source/UAF/Private/UAF/Attributes/Transformers/Interpolate.cpp
@@ -12,6 +12,7 @@ namespace UE::UAF::Transformers
 	{
 		// Named sets must match as it ensures that our inputs have the same sizes/shapes
 		checkf(CollectionA.GetNamedSet() == CollectionB.GetNamedSet(), TEXT("Both input collections must use the same named set"));
+		checkf(CollectionA.GetValueSpace() == CollectionB.GetValueSpace(), TEXT("Both input collections must have matching value spaces: [%s] != [%s]"), *CollectionA.GetValueSpace().ToString(), *CollectionB.GetValueSpace().ToString());
 
 		const bool bIsOutputA = &CollectionA == &OutputCollection;
 		if (bIsOutputA && !FAnimWeight::IsRelevant(Alpha))
@@ -254,6 +255,7 @@ namespace UE::UAF::Transformers
 		{
 			// Our output collection isn't either of our inputs, reset it
 			OutputCollection.Reset(CollectionA.GetNamedSet());
+			OutputCollection.SetValueSpace(CollectionA.GetValueSpace());
 
 			// Set maps
 			OuterJoinBy(
diff --git a/Engine/Plugins/Experimental/UAF/UAF/Source/UAF/Private/UAF/Attributes/Transformers/Overwrite.cpp b/Engine/Plugins/Experimental/UAF/UAF/Source/UAF/Private/UAF/Attributes/Transformers/Overwrite.cpp
--- a/Engine/Plugins/Experimental/UAF/UAF/Source/UAF/Private/UAF/Attributes/Transformers/Overwrite.cpp
+++ b/Engine/Plugins/Experimental/UAF/UAF/Source/UAF/Private/UAF/Attributes/Transformers/Overwrite.cpp
@@ -94,6 +94,7 @@ namespace UE::UAF::Transformers
 		{
 			// Our output collection isn't either of our inputs, reset it
 			OutputCollection.Reset(SourceCollection.GetNamedSet());
+			OutputCollection.SetValueSpace(SourceCollection.GetValueSpace());
 
 			// Set maps
 			OuterJoinBy(
diff --git a/Engine/Plugins/Experimental/UAF/UAF/Source/UAF/Private/UAF/Attributes/ValueSpace.cpp b/Engine/Plugins/Experimental/UAF/UAF/Source/UAF/Private/UAF/Attributes/ValueSpace.cpp
new file mode 100644
--- /dev/null
+++ b/Engine/Plugins/Experimental/UAF/UAF/Source/UAF/Private/UAF/Attributes/ValueSpace.cpp
@@ -0,0 +1,74 @@
+// Copyright Epic Games, Inc. All Rights Reserved.
+
+#include "UAF/Attributes/ValueSpace.h"
+
+namespace UE::UAF
+{
+	namespace Private
+	{
+		const TCHAR* ValueSpaceTypeToString(EValueSpaceType Type)
+		{
+			switch (Type)
+			{
+			default:
+			case EValueSpaceType::None:			return TEXT("None");
+			case EValueSpaceType::Local:		return TEXT("Local");
+			case EValueSpaceType::Component:	return TEXT("Component");
+			case EValueSpaceType::Mixed:		return TEXT("Mixed");
+			case EValueSpaceType::AI:			return TEXT("AI");
+			}
+		}
+
+		FString MixedSpaceFlagsToString(EMixedSpaceFlags Flags)
+		{
+			FString Result;
+
+			if (EnumHasAnyFlags(Flags, EMixedSpaceFlags::MeshRotation))
+			{
+				Result = TEXT("MeshRotation");
+			}
+
+			if (EnumHasAnyFlags(Flags, EMixedSpaceFlags::MeshScale))
+			{
+				if (!Result.IsEmpty())
+				{
+					Result += TEXT(" | ");
+				}
+
+				Result += TEXT("MeshScale");
+			}
+
+			if (EnumHasAnyFlags(Flags, EMixedSpaceFlags::RootRotation))
+			{
+				if (!Result.IsEmpty())
+				{
+					Result += TEXT(" | ");
+				}
+
+				Result += TEXT("RootRotation");
+			}
+
+			if (Result.IsEmpty())
+			{
+				Result = TEXT("None");
+			}
+
+			return Result;
+		}
+	}
+
+	FString FValueSpace::ToString() const
+	{
+		const TCHAR* IsAdditiveStr = IsAdditive() ? TEXT("Additive ") : TEXT("");
+		const TCHAR* ValueTypeStr = Private::ValueSpaceTypeToString(GetType());
+
+		FString Result = FString::Printf(TEXT("%s%s"), IsAdditiveStr, ValueTypeStr);
+
+		if (GetType() == EValueSpaceType::Mixed)
+		{
+			Result += FString::Printf(TEXT(" (%s)"), *Private::MixedSpaceFlagsToString(GetMixedSpaceFlags()));
+		}
+
+		return Result;
+	}
+}
diff --git a/Engine/Plugins/Experimental/UAF/UAF/Source/UAF/Public/UAF/Attributes/AttributeCollection.h b/Engine/Plugins/Experimental/UAF/UAF/Source/UAF/Public/UAF/Attributes/AttributeCollection.h
--- a/Engine/Plugins/Experimental/UAF/UAF/Source/UAF/Public/UAF/Attributes/AttributeCollection.h
+++ b/Engine/Plugins/Experimental/UAF/UAF/Source/UAF/Public/UAF/Attributes/AttributeCollection.h
@@ -9,6 +9,7 @@
 #include "UAF/Attributes/AttributeMappingKey.h"
 #include "UAF/Attributes/AttributeSetMap.h"
 #include "UAF/Attributes/IndirectAllocator.h"
+#include "UAF/Attributes/ValueSpace.h"
 
 #include "AttributeCollection.generated.h"
 
@@ -137,6 +138,16 @@ namespace UE::UAF
 		// Returns the named set this collection maps to
 		[[nodiscard]] UE_API const FAttributeNamedSetPtr& GetNamedSet() const;
 
+		// Returns the space our values live in
+		[[nodiscard]] UE_API FValueSpace GetValueSpace() const;
+
+		// Sets the space our values live in
+		// Note that care must be taken when changing the value space as all values within the collection must be migrated accordingly
+		UE_API void SetValueSpace(FValueSpace ValueSpace);
+
+		// Returns whether or not the space our values live in is additive
+		[[nodiscard]] UE_API bool IsAdditive() const;
+
 		// Resets the collection with the specified named set
 		UE_API void Reset(const FAttributeNamedSetPtr& NamedSet);
 
@@ -268,6 +279,9 @@ namespace UE::UAF
 		// The named set we are based on, every set mapping uses a typed set within it
 		FAttributeNamedSetPtr NamedSet;
 
+		// The space in which our values live
+		FValueSpace ValueSpace;
+
 		// An array of set mappings (sorted by mapping key)
 		FSetMapEntry* SetMaps = nullptr;
 
@@ -533,8 +547,8 @@ namespace UE::UAF
 	{
 	}
 
-	inline FAttributeCollectionHeap::FAttributeCollectionHeap(const FAttributeNamedSetPtr& NamedSet)
-		: FAttributeCollection(NamedSet, &FAllocatorTypeTrait<FDefaultAllocator>::Realloc)
+	inline FAttributeCollectionHeap::FAttributeCollectionHeap(const FAttributeNamedSetPtr& InNamedSet)
+		: FAttributeCollection(InNamedSet, &FAllocatorTypeTrait<FDefaultAllocator>::Realloc)
 	{
 	}
 
@@ -543,8 +557,8 @@ namespace UE::UAF
 	{
 	}
 
-	inline FAttributeCollectionStack::FAttributeCollectionStack(const FAttributeNamedSetPtr& NamedSet)
-		: FAttributeCollection(NamedSet, &FAllocatorTypeTrait<TMemStackAllocator<>>::Realloc)
+	inline FAttributeCollectionStack::FAttributeCollectionStack(const FAttributeNamedSetPtr& InNamedSet)
+		: FAttributeCollection(InNamedSet, &FAllocatorTypeTrait<TMemStackAllocator<>>::Realloc)
 	{
 	}
 
diff --git a/Engine/Plugins/Experimental/UAF/UAF/Source/UAF/Public/UAF/Attributes/ValueSpace.h b/Engine/Plugins/Experimental/UAF/UAF/Source/UAF/Public/UAF/Attributes/ValueSpace.h
new file mode 100644
--- /dev/null
+++ b/Engine/Plugins/Experimental/UAF/UAF/Source/UAF/Public/UAF/Attributes/ValueSpace.h
@@ -0,0 +1,160 @@
+// Copyright Epic Games, Inc. All Rights Reserved.
+
+#pragma once
+
+#include "CoreMinimal.h"
+
+#define UE_API UAF_API
+
+namespace UE::UAF
+{
+	// Encodes the possible spaces values can live in
+	enum class EValueSpaceType : uint16
+	{
+		// Unknown value space
+		None		= 0x0000,
+
+		// Bone transforms: relative to their parent bone
+		// Everything else: normal space (no parent/child relationships)
+		Local		= 0x0001,
+
+		// Bone transforms: relative to the owning component
+		// Everything else: normal space (no parent/child relationships)
+		Component	= 0x0002,
+
+		// Bone transforms: rotations/translations/scales can be in individual spaces, @see EMixedSpaceFlags
+		//                  Only safe with uniform scale
+		// Everything else: normal space (no parent/child relationships)
+		Mixed		= 0x0003,
+
+		// Values are in machine learning latent space
+		// Value types will dictate which space they are in
+		AI			= 0x0004,
+	};
+
+	// Encodes the possible permutations of mixed spaces that bone transforms can live in
+	enum class EMixedSpaceFlags : uint16
+	{
+		// No mixed space, equivalent to Local Space
+		None			= 0x0000,
+
+		// Rotations are relative to the owning component
+		// Only safe with uniform scale
+		MeshRotation	= 0x0010,
+
+		// Scales are relative to the owning component
+		// Only safe with uniform scale
+		MeshScale		= 0x0020,
+
+		// Rotations are relative to the root bone
+		// Only safe with uniform scale
+		RootRotation	= 0x0040,
+	};
+
+	ENUM_CLASS_FLAGS(EMixedSpaceFlags);
+
+	// Encapsulates the possible spaces that values can live in
+	class FValueSpace
+	{
+	public:
+		// Constructs an unknown space
+		FValueSpace() = default;
+
+		// Constructs with the specified space type and additive flag
+		explicit FValueSpace(EValueSpaceType Type, bool bIsAdditive = false);
+
+		// Constructs a mixed space with the specified flags and additive flag
+		explicit FValueSpace(EMixedSpaceFlags Flags, bool bIsAdditive = false);
+
+		// Returns the value space type
+		[[nodiscard]] EValueSpaceType GetType() const;
+
+		// Sets the value space type, preserving whether we are additive or not
+		void SetType(EValueSpaceType Type);
+
+		// Returns the mixed space flags
+		[[nodiscard]] EMixedSpaceFlags GetMixedSpaceFlags() const;
+
+		// Sets the value space type to mixed along with the specified flags, preserving whether we are additive or not
+		void SetMixedSpaceFlags(EMixedSpaceFlags Flags);
+
+		// Returns whether we are additive or not
+		[[nodiscard]] bool IsAdditive() const;
+
+		// Sets whether we are additive or not, preserving the type and mixed flags
+		void SetIsAdditive(bool bIsAdditive);
+
+		// Equality operators
+		bool operator==(FValueSpace Other) const;
+		bool operator!=(FValueSpace Other) const;
+
+		// Converts this value space into a string representation
+		[[nodiscard]] UE_API FString ToString() const;
+
+	private:
+		static constexpr uint16 SPACE_CATEGORY_MASK		= 0x000F;
+		static constexpr uint16 MIXED_SPACE_FLAGS_MASK	= 0x00F0;
+		static constexpr uint16 IS_ADDITIVE_MASK		= 0x8000;
+
+		// Our internal value is opaque and packed
+		uint16 Packed = static_cast<uint16>(EValueSpaceType::None);
+	};
+
+	//////////////////////////////////////////////////////////////////////////
+	// Implementation
+
+	inline FValueSpace::FValueSpace(EValueSpaceType Type, bool bIsAdditive)
+		: Packed(static_cast<uint16>(Type) | (bIsAdditive ? IS_ADDITIVE_MASK : 0))
+	{
+	}
+
+	inline FValueSpace::FValueSpace(EMixedSpaceFlags Flags, bool bIsAdditive)
+		: Packed(static_cast<uint16>(EValueSpaceType::Mixed) | static_cast<uint16>(Flags) | (bIsAdditive ? IS_ADDITIVE_MASK : 0))
+	{
+	}
+
+	inline EValueSpaceType FValueSpace::GetType() const
+	{
+		return static_cast<EValueSpaceType>(Packed & SPACE_CATEGORY_MASK);
+	}
+
+	inline void FValueSpace::SetType(EValueSpaceType Type)
+	{
+		// Only preserve the additive flag
+		Packed = (Packed & IS_ADDITIVE_MASK) | static_cast<uint16>(Type);
+	}
+
+	inline EMixedSpaceFlags FValueSpace::GetMixedSpaceFlags() const
+	{
+		return static_cast<EMixedSpaceFlags>(Packed & MIXED_SPACE_FLAGS_MASK);
+	}
+
+	inline void FValueSpace::SetMixedSpaceFlags(EMixedSpaceFlags Flags)
+	{
+		// Only preserve the additive flag and make sure our category is mixed
+		Packed = (Packed & IS_ADDITIVE_MASK) | static_cast<uint16>(EValueSpaceType::Mixed) | static_cast<uint16>(Flags);
+	}
+
+	inline bool FValueSpace::IsAdditive() const
+	{
+		return (Packed & IS_ADDITIVE_MASK) != 0;
+	}
+
+	inline void FValueSpace::SetIsAdditive(bool bIsAdditive)
+	{
+		// Preserve everything else
+		Packed = (Packed & ~IS_ADDITIVE_MASK) | (bIsAdditive ? IS_ADDITIVE_MASK : 0);
+	}
+
+	inline bool FValueSpace::operator==(FValueSpace Other) const
+	{
+		return Packed == Other.Packed;
+	}
+
+	inline bool FValueSpace::operator!=(FValueSpace Other) const
+	{
+		return Packed != Other.Packed;
+	}
+}
+
+#undef UE_API
diff --git a/Engine/Plugins/Experimental/UAF/UAFAnimGraph/Source/UAFAnimGraph/Private/EvaluationVM/Tasks/PushAnimSequenceKeyframe.cpp b/Engine/Plugins/Experimental/UAF/UAFAnimGraph/Source/UAFAnimGraph/Private/EvaluationVM/Tasks/PushAnimSequenceKeyframe.cpp
--- a/Engine/Plugins/Experimental/UAF/UAFAnimGraph/Source/UAFAnimGraph/Private/EvaluationVM/Tasks/PushAnimSequenceKeyframe.cpp
+++ b/Engine/Plugins/Experimental/UAF/UAFAnimGraph/Source/UAFAnimGraph/Private/EvaluationVM/Tasks/PushAnimSequenceKeyframe.cpp
@@ -41,6 +41,7 @@ void FAnimNextAnimSequenceKeyframeTask::Execute(UE::UAF::FEvaluationVM& VM) cons
 		const FAttributeEvaluationContext& VMContext = VM.GetContext();
 
 		const bool bIsAdditive = AnimSequencePtr->IsValidAdditive();
+		const EAdditiveAnimationType AdditiveType = AnimSequencePtr->GetAdditiveAnimType();
 		const bool bUseRawData = FDecompressionTools::ShouldUseRawData(AnimSequencePtr, ExtractionContext);
 
 		if (VMContext.IsValid())
@@ -48,6 +49,17 @@ void FAnimNextAnimSequenceKeyframeTask::Execute(UE::UAF::FEvaluationVM& VM) cons
 			constexpr bool bInitializeWithDefaultValues = true;
 			FAttributeCollectionStack Collection = VM.MakeAttributeCollection(bIsAdditive, bInitializeWithDefaultValues);
 
+			FValueSpace ValueSpace;
+			if (bIsAdditive && AdditiveType == AAT_RotationOffsetMeshSpace)
+			{
+				ValueSpace = FValueSpace(EMixedSpaceFlags::MeshRotation, bIsAdditive);
+			}
+			else
+			{
+				ValueSpace = FValueSpace(EValueSpaceType::Local, bIsAdditive);
+			}
+			Collection.SetValueSpace(ValueSpace);
+
 			FDecompressionTools::GetAnimationPose(AnimSequencePtr, ExtractionContext, VMContext, Collection, bUseRawData);
 			FDecompressionTools::GetAnimationCurves(AnimSequencePtr, ExtractionContext, VMContext, Collection, bUseRawData);
 			FDecompressionTools::GetAnimationAttributes(AnimSequencePtr, ExtractionContext, VMContext, Collection, bUseRawData);
