UAF: Fix up printed python command when a rig vm graph got removed

#Synced-CL 48757676

diff --git a/Engine/Plugins/Experimental/UAF/UAF/Source/UAFUncookedOnly/Private/AnimNextRigVMAssetEditorData.cpp b/Engine/Plugins/Experimental/UAF/UAF/Source/UAFUncookedOnly/Private/AnimNextRigVMAssetEditorData.cpp
--- a/Engine/Plugins/Experimental/UAF/UAF/Source/UAFUncookedOnly/Private/AnimNextRigVMAssetEditorData.cpp
+++ b/Engine/Plugins/Experimental/UAF/UAF/Source/UAFUncookedOnly/Private/AnimNextRigVMAssetEditorData.cpp
@@ -641,7 +641,7 @@ void UAnimNextRigVMAssetEditorData::HandleRigVMGraphRemoved(const FRigVMClient*
 		if(!bSuspendPythonMessagesForRigVMClient)
 		{
 			const FString AssetName = RigVMGraph->GetSchema()->GetSanitizedName(GetName(), true, false);
-			RigVMPythonUtils::Print(AssetName, FString::Printf(TEXT("asset.add_graph('%s')"), *RigVMGraph->GetName()));
+			RigVMPythonUtils::Print(AssetName, FString::Printf(TEXT("asset.remove_graph('%s')"), *RigVMGraph->GetName()));
 		}
 #endif
 	}
