UAF: Fix for SetDefaultValueFromString always returning false, even if the setting of value succeeded

#Synced-CL 48757676

diff --git a/Engine/Plugins/Experimental/UAF/UAF/Source/UAFUncookedOnly/Private/Variables/AnimNextProgrammaticVariable.cpp b/Engine/Plugins/Experimental/UAF/UAF/Source/UAFUncookedOnly/Private/Variables/AnimNextProgrammaticVariable.cpp
--- a/Engine/Plugins/Experimental/UAF/UAF/Source/UAFUncookedOnly/Private/Variables/AnimNextProgrammaticVariable.cpp
+++ b/Engine/Plugins/Experimental/UAF/UAF/Source/UAFUncookedOnly/Private/Variables/AnimNextProgrammaticVariable.cpp
@@ -63,7 +63,7 @@ bool FAnimNextProgrammaticVariable::SetDefaultValueFromString(const FString& InD
 		return false;
 	}
 
-	return false;
+	return true;
 }
 
 const FInstancedPropertyBag& FAnimNextProgrammaticVariable::GetPropertyBag() const
