UAF: Fix binding compilation issues

Ensure we process the bindings of shared variables that systems use
Correctly report errors with property bindings
Minor optimization to pin-type resolution when generating binding graphs
Add additional error feedback on variable binding widget (red text, tooltip details)

#Synced-CL 48734215

diff --git a/Engine/Plugins/Experimental/UAF/UAF/Source/UAFEditor/Private/Variables/VariableBindingPropertyCustomization.cpp b/Engine/Plugins/Experimental/UAF/UAF/Source/UAFEditor/Private/Variables/VariableBindingPropertyCustomization.cpp
--- a/Engine/Plugins/Experimental/UAF/UAF/Source/UAFEditor/Private/Variables/VariableBindingPropertyCustomization.cpp
+++ b/Engine/Plugins/Experimental/UAF/UAF/Source/UAFEditor/Private/Variables/VariableBindingPropertyCustomization.cpp
@@ -73,6 +73,10 @@ void FVariableBindingPropertyCustomization::CustomizeHeader(TSharedRef<IProperty
 		.AutoWidth()
 		[
 			SAssignNew(ValueWidget, SComboButton)
+			.ForegroundColor_Lambda([this]()
+			{
+				return bIsBindingValid ? FSlateColor::UseForeground() : FLinearColor::Red;
+			})
 			.Visibility_Lambda([this]()
 			{
 				return bShowBindingSelector ? EVisibility::Visible : EVisibility::Collapsed;
@@ -179,6 +183,7 @@ void FVariableBindingPropertyCustomization::Refresh()
 	{
 		NameText = LOCTEXT("MultipleTypes", "Multiple Types");
 		TooltipText = LOCTEXT("MultipleTypes", "Multiple Types");
+		bIsBindingValid = true;
 		bShowBindingSelector = true;
 		return;
 	}
@@ -221,6 +226,7 @@ void FVariableBindingPropertyCustomization::Refresh()
 		// Common values
 		NameText = GetBindingDisplayNameText(BindingData);
 		TooltipText = GetBindingTooltipText(BindingData);
+		bIsBindingValid = ValidateBinding(BindingData);
 		bShowBindingSelector = true;
 	}
 	else if(BindingStruct != nullptr)
@@ -228,6 +234,7 @@ void FVariableBindingPropertyCustomization::Refresh()
 		// No common values, but common struct, allow selection
 		NameText = LOCTEXT("MultipleValues", "Multiple Values");
 		TooltipText = LOCTEXT("MultipleValues", "Multiple Values");
+		bIsBindingValid = true;
 		bShowBindingSelector = true;
 	}
 	else
@@ -235,6 +242,7 @@ void FVariableBindingPropertyCustomization::Refresh()
 		// No common struct, dont allow binding selection 
 		NameText = LOCTEXT("MultipleValues", "Multiple Values");
 		TooltipText = LOCTEXT("MultipleValues", "Multiple Values");
+		bIsBindingValid = true;
 		bShowBindingSelector = false;
 	}
 }
@@ -281,6 +289,21 @@ FText FVariableBindingPropertyCustomization::GetBindingTooltipText(TConstStructV
 	return TextBuilder.ToText();
 }
 
+bool FVariableBindingPropertyCustomization::ValidateBinding(TConstStructView<FAnimNextVariableBindingData> InBindingData)
+{
+	if(InBindingData.IsValid())
+	{
+		UncookedOnly::IAnimNextUncookedOnlyModule& UncookedOnlyModule = FModuleManager::GetModuleChecked<UncookedOnly::IAnimNextUncookedOnlyModule>("UAFUncookedOnly");
+		if(TSharedPtr<UncookedOnly::IVariableBindingType> BindingType = UncookedOnlyModule.FindVariableBindingType(InBindingData.GetScriptStruct()))
+		{
+			FText Error;
+			return BindingType->Validate(InBindingData, Error);
+		}
+	}
+
+	return true;
+}
+
 TSharedRef<SWidget> FVariableBindingPropertyCustomization::CreateBindingWidget() const
 {
 	if(!Type.IsValid())
diff --git a/Engine/Plugins/Experimental/UAF/UAF/Source/UAFEditor/Private/Variables/VariableBindingPropertyCustomization.h b/Engine/Plugins/Experimental/UAF/UAF/Source/UAFEditor/Private/Variables/VariableBindingPropertyCustomization.h
--- a/Engine/Plugins/Experimental/UAF/UAF/Source/UAFEditor/Private/Variables/VariableBindingPropertyCustomization.h
+++ b/Engine/Plugins/Experimental/UAF/UAF/Source/UAFEditor/Private/Variables/VariableBindingPropertyCustomization.h
@@ -38,6 +38,9 @@ private:
 	/** Gets a name to display for a variable binding's tooltip in the editor */
 	static FText GetBindingTooltipText(TConstStructView<FAnimNextVariableBindingData> InBindingData);
 
+	/** Checks a binding for validity in the editor */
+	static bool ValidateBinding(TConstStructView<FAnimNextVariableBindingData> InBindingData);
+
 	/** Create a binding widget */
 	TSharedRef<SWidget> CreateBindingWidget() const;
 
@@ -51,6 +54,7 @@ private:
 	FSlateColor IconColor = FLinearColor::Gray;
 	TSharedPtr<SWidget> ValueWidget;
 	TSharedPtr<SWidget> ContainerWidget;
+	bool bIsBindingValid = false;
 	bool bRefreshRequested = false;
 	bool bShowBindingSelector = false;
 };
diff --git a/Engine/Plugins/Experimental/UAF/UAF/Source/UAFUncookedOnly/Internal/Variables/AnimNextUniversalObjectLocatorBindingData.h b/Engine/Plugins/Experimental/UAF/UAF/Source/UAFUncookedOnly/Internal/Variables/AnimNextUniversalObjectLocatorBindingData.h
--- a/Engine/Plugins/Experimental/UAF/UAF/Source/UAFUncookedOnly/Internal/Variables/AnimNextUniversalObjectLocatorBindingData.h
+++ b/Engine/Plugins/Experimental/UAF/UAF/Source/UAFUncookedOnly/Internal/Variables/AnimNextUniversalObjectLocatorBindingData.h
@@ -46,11 +46,7 @@ struct FAnimNextUniversalObjectLocatorBindingData : public FAnimNextVariableBind
 	// FAnimNextVariableBindingData interface
 	virtual bool IsValid() const override
 	{
-		const bool bValidLocator = !Locator.IsEmpty();
-		const bool bValidUOL = Type == FAnimNextUniversalObjectLocatorBindingType::UOL;
-		const bool bValidProperty = (Type == FAnimNextUniversalObjectLocatorBindingType::Property && !Property.IsPathToFieldEmpty());
-		const bool bValidFunction = ((Type == FAnimNextUniversalObjectLocatorBindingType::Function || Type == FAnimNextUniversalObjectLocatorBindingType::HoistedFunction) && !Function.IsNull());
-		return bValidLocator && (bValidUOL || bValidProperty || bValidFunction);
+		return !Locator.IsEmpty();
 	}
 
 #if WITH_EDITORONLY_DATA
diff --git a/Engine/Plugins/Experimental/UAF/UAF/Source/UAFUncookedOnly/Private/UncookedOnlyUtils.cpp b/Engine/Plugins/Experimental/UAF/UAF/Source/UAFUncookedOnly/Private/UncookedOnlyUtils.cpp
--- a/Engine/Plugins/Experimental/UAF/UAF/Source/UAFUncookedOnly/Private/UncookedOnlyUtils.cpp
+++ b/Engine/Plugins/Experimental/UAF/UAF/Source/UAFUncookedOnly/Private/UncookedOnlyUtils.cpp
@@ -350,32 +350,48 @@ void FUtils::CompileVariableBindingsInternal(const FRigVMCompileSettings& InSett
 
 	for(const UAnimNextRigVMAssetEntry* Entry : EditorData->Entries)
 	{
-		const IAnimNextRigVMVariableInterface* Variable = Cast<IAnimNextRigVMVariableInterface>(Entry);
-		if(Variable == nullptr)
+		auto AddBindingForEntry = [bInThreadSafe, &Module, &BindingGroups](const IAnimNextRigVMVariableInterface* InVariable)
 		{
-			continue;
-		}
+			TConstStructView<FAnimNextVariableBindingData> Binding = InVariable->GetBinding();
+			if(!Binding.IsValid() || !Binding.Get<FAnimNextVariableBindingData>().IsValid())
+			{
+				return;
+			}
 
-		TConstStructView<FAnimNextVariableBindingData> Binding = Variable->GetBinding();
-		if(!Binding.IsValid() || !Binding.Get<FAnimNextVariableBindingData>().IsValid())
-		{
-			continue;
-		}
+			if(Binding.Get<FAnimNextVariableBindingData>().IsThreadSafe() != bInThreadSafe)
+			{
+				return;
+			}
+
+			TSharedPtr<IVariableBindingType> BindingType = Module.FindVariableBindingType(Binding.GetScriptStruct());
+			if(!BindingType.IsValid())
+			{
+				return;
+			}
 
-		if(Binding.Get<FAnimNextVariableBindingData>().IsThreadSafe() != bInThreadSafe)
+			TArray<IVariableBindingType::FBindingGraphInput>& Group = BindingGroups.FindOrAdd(BindingType.Get());
+			FRigVMTemplateArgumentType RigVMArg = InVariable->GetType().ToRigVMTemplateArgument();
+			Group.Add({ InVariable->GetVariableName(), RigVMArg.GetBaseCPPType(), RigVMArg.CPPTypeObject, Binding});
+		};
+		
+		if(const IAnimNextRigVMVariableInterface* Variable = Cast<IAnimNextRigVMVariableInterface>(Entry))
 		{
-			continue;
+			AddBindingForEntry(Variable);
 		}
-
-		TSharedPtr<IVariableBindingType> BindingType = Module.FindVariableBindingType(Binding.GetScriptStruct());
-		if(!BindingType.IsValid())
+		else if (const UAnimNextSharedVariablesEntry* SharedVariables = Cast<UAnimNextSharedVariablesEntry>(Entry))
 		{
-			continue;
+			if (const UAnimNextRigVMAsset* SharedVariablesAsset = SharedVariables->GetAsset())
+			{
+				const UAnimNextRigVMAssetEditorData* SharedVariablesEditorData = GetEditorData<const UAnimNextRigVMAssetEditorData>(SharedVariablesAsset);
+				for(const UAnimNextRigVMAssetEntry* SharedVariablesEntry : SharedVariablesEditorData->Entries)
+				{
+					if(const IAnimNextRigVMVariableInterface* SharedVariable = Cast<IAnimNextRigVMVariableInterface>(SharedVariablesEntry))
+					{
+						AddBindingForEntry(SharedVariable);
+					}
+				}
+			}
 		}
-
-		TArray<IVariableBindingType::FBindingGraphInput>& Group = BindingGroups.FindOrAdd(BindingType.Get());
-		FRigVMTemplateArgumentType RigVMArg = Variable->GetType().ToRigVMTemplateArgument();
-		Group.Add({ Variable->GetVariableName(), RigVMArg.GetBaseCPPType(), RigVMArg.CPPTypeObject, Binding});
 	}
 
 	const bool bHasBindings = BindingGroups.Num() > 0;
diff --git a/Engine/Plugins/Experimental/UAF/UAF/Source/UAFUncookedOnly/Private/Variables/UniversalObjectLocatorBindingType.cpp b/Engine/Plugins/Experimental/UAF/UAF/Source/UAFUncookedOnly/Private/Variables/UniversalObjectLocatorBindingType.cpp
--- a/Engine/Plugins/Experimental/UAF/UAF/Source/UAFUncookedOnly/Private/Variables/UniversalObjectLocatorBindingType.cpp
+++ b/Engine/Plugins/Experimental/UAF/UAF/Source/UAFUncookedOnly/Private/Variables/UniversalObjectLocatorBindingType.cpp
@@ -489,6 +489,10 @@ FText FUniversalObjectLocatorBindingType::GetDisplayText(TConstStructView<FAnimN
 		{
 			Property->GetFName().AppendString(StringBuilder);
 		}
+		else if (!LocatorBinding->Property.IsPathToFieldEmpty())
+		{
+			StringBuilder.Append(TEXT("???"));
+		}
 		else
 		{
 			StringBuilder.Append(TEXT("None"));
@@ -504,6 +508,10 @@ FText FUniversalObjectLocatorBindingType::GetDisplayText(TConstStructView<FAnimN
 		{
 			Function->GetFName().AppendString(StringBuilder);
 		}
+		else if (!LocatorBinding->Function.IsNull())
+		{
+			StringBuilder.Append(TEXT("???"));
+		}
 		else
 		{
 			StringBuilder.Append(TEXT("None"));
@@ -558,6 +566,10 @@ FText FUniversalObjectLocatorBindingType::GetTooltipText(TConstStructView<FAnimN
 		{
 			Property->GetFName().AppendString(StringBuilder);
 		}
+		else if (!LocatorBinding->Property.IsPathToFieldEmpty())
+		{
+			StringBuilder.Append(TEXT("???"));
+		}
 		else
 		{
 			StringBuilder.Append(TEXT("None"));
@@ -569,6 +581,10 @@ FText FUniversalObjectLocatorBindingType::GetTooltipText(TConstStructView<FAnimN
 		{
 			Function->GetFName().AppendString(StringBuilder);
 		}
+		else if (!LocatorBinding->Function.IsNull())
+		{
+			StringBuilder.Append(TEXT("???"));
+		}
 		else
 		{
 			StringBuilder.Append(TEXT("None"));
@@ -586,9 +602,60 @@ FText FUniversalObjectLocatorBindingType::GetTooltipText(TConstStructView<FAnimN
 	LocatorBinding->Locator.ToString(ScopeStringBuilder);
 	TextBuilder.AppendLine(FText::Format(LOCTEXT("ParameterUOLTooltipFormat", "UOL: {0}"), FText::FromStringView(ScopeStringBuilder)));
 
+	FText ValidationError;
+	if (!Validate(InBindingData, ValidationError))
+	{
+		TextBuilder.AppendLine(FText::Format(LOCTEXT("ValidationErrorFormat", "Error: {0}"), ValidationError));
+	}
+
 	return TextBuilder.ToText();
 }
 
+bool FUniversalObjectLocatorBindingType::Validate(TConstStructView<FAnimNextVariableBindingData> InBindingData, FText& OutError) const
+{
+	const FAnimNextUniversalObjectLocatorBindingData& LocatorBinding = InBindingData.Get<FAnimNextUniversalObjectLocatorBindingData>();
+	if (LocatorBinding.Locator.IsEmpty())
+	{
+		return true;
+	}
+
+	switch(LocatorBinding.Type)
+	{
+	case FAnimNextUniversalObjectLocatorBindingType::UOL:
+		return true;
+	case FAnimNextUniversalObjectLocatorBindingType::Property:
+	{
+		if (LocatorBinding.Property.Get() == nullptr)
+		{
+			OutError = FText::Format(LOCTEXT("PropertyBindingError", "Property binding '{0}' could not be found"), FText::FromString(LocatorBinding.Property.ToString()));
+			return false;
+		}
+		return true;
+	}
+	case FAnimNextUniversalObjectLocatorBindingType::Function:
+	{
+		if(LocatorBinding.Function.Get() == nullptr)
+		{
+			OutError = FText::Format(LOCTEXT("FunctionBindingError", "Function '{0}' could not be found"), FText::FromString(LocatorBinding.Function.ToString()));
+			return false;
+		}
+		return true;
+	}
+	case FAnimNextUniversalObjectLocatorBindingType::HoistedFunction:
+	{
+		if(LocatorBinding.Function.Get() == nullptr)
+		{
+			OutError = FText::Format(LOCTEXT("HoistedFunctionBindingError", "Function '{0}' could not be found"), FText::FromString(LocatorBinding.Function.ToString()));
+			return false;
+		}
+		return true;
+	}
+	default:
+		checkNoEntry();
+		return false;
+	}
+}
+
 void FUniversalObjectLocatorBindingType::BuildBindingGraphFragment(const FRigVMCompileSettings& InSettings, const FBindingGraphFragmentArgs& InArgs, URigVMPin*& OutExecTail, FVector2D& OutLocation) const
 {
 	if( InArgs.Event != FRigUnit_AnimNextExecuteBindings_GT::StaticStruct() &&
@@ -627,13 +694,13 @@ void FUniversalObjectLocatorBindingType::BuildBindingGraphFragment(const FRigVMC
 		if(ResolverNode == nullptr)
 		{
 			InSettings.ReportError(TEXT("Could not spawn UOL resolver function"));
-			return;
+			continue;
 		}
 		URigVMPin* LocatorPin = ResolverNode->FindPin(GET_MEMBER_NAME_STRING_CHECKED(FRigUnit_ResolveUniversalObjectLocator, Locator));
 		if(LocatorPin == nullptr)
 		{
 			InSettings.ReportError(TEXT("Could find Locator pin"));
-			return;
+			continue;
 		}
 		FString DefaultLocatorValue;
 		FUniversalObjectLocator::StaticStruct()->ExportText(DefaultLocatorValue, &LocatorDataPair.Key, nullptr, nullptr, PPF_None, nullptr);
@@ -641,19 +708,27 @@ void FUniversalObjectLocatorBindingType::BuildBindingGraphFragment(const FRigVMC
 		if(!bLocatorPinDefaultValueSet)
 		{
 			InSettings.ReportError(TEXT("Could not set Locator pin's default value"));
-			return;
+			continue;
 		}
 		URigVMPin* ResolvedObjectPin = ResolverNode->FindPin(GET_MEMBER_NAME_STRING_CHECKED(FRigUnit_ResolveUniversalObjectLocator, Object));
 		if(ResolvedObjectPin == nullptr)
 		{
 			InSettings.ReportError(TEXT("Could not find UOL resolver object pin"));
-			return;
+			continue;
 		}
 
 		// Add unit nodes to call function or fetch property and set appropriate variables
 		float AccumulatedYOffset = 0.0f;
 		for(const FBindingGraphInput* Input : LocatorDataPair.Value.Inputs)
 		{
+			FText ValidationError;
+			if (!Validate(Input->BindingData, ValidationError))
+			{
+				FText ValidationErrorReport = FText::Format(LOCTEXT("ValidationErrorFormat", "Variable '{0}': {1}"), FText::FromName(Input->VariableName), ValidationError);
+				InSettings.ReportError(ValidationErrorReport.ToString());
+				continue;
+			}
+
 			const FAnimNextUniversalObjectLocatorBindingData& LocatorBinding = Input->BindingData.Get<FAnimNextUniversalObjectLocatorBindingData>();
 			URigVMTemplateNode* GetterNode = nullptr;
 			switch(LocatorBinding.Type)
@@ -662,28 +737,49 @@ void FUniversalObjectLocatorBindingType::BuildBindingGraphFragment(const FRigVMC
 				break;
 			case FAnimNextUniversalObjectLocatorBindingType::Property:
 			{
+				if (LocatorBinding.Property.Get() == nullptr)
+				{
+					InSettings.ReportErrorf(TEXT("Property binding '%s' for variable '%s' could not be found"), *LocatorBinding.Property.ToString(), *Input->VariableName.ToString());
+					continue;
+				}
+
 				GetterNode = InArgs.Controller->AddTemplateNode(FRigVMDispatch_GetObjectProperty().GetTemplateNotation(), FVector2D(200.0f, YOffset + AccumulatedYOffset), TEXT(""), false);
 				if(GetterNode == nullptr)
 				{
 					InSettings.ReportError(TEXT("Could not spawn Get Object Property node"));
-					return;
+					continue;
 				}
 				URigVMPin* PropertyPin = GetterNode->FindPin(FRigVMDispatch_GetObjectProperty::PropertyName.ToString());
 				if(PropertyPin == nullptr)
 				{
 					InSettings.ReportError(TEXT("Could not find Property pin"));
-					return;
+					continue;
 				}
+
 				FString DefaultPathValue;
 				FAnimNextFieldPath AnimNextFieldPath;
-				AnimNextFieldPath.FieldPath = LocatorBinding.Property; 
+				AnimNextFieldPath.FieldPath = LocatorBinding.Property;
 				FAnimNextFieldPath::StaticStruct()->ExportText(DefaultPathValue, &AnimNextFieldPath, nullptr, nullptr, PPF_None, nullptr);
 				bool bPropertyDefaultValueSet = InArgs.Controller->SetPinDefaultValue(PropertyPin, DefaultPathValue, true, false, false);
 				if(!bPropertyDefaultValueSet)
 				{
 					InSettings.ReportError(TEXT("Could not set Property pin's default value"));
-					return;
+					continue;
+				}
+
+				URigVMPin* ValuePin = GetterNode->FindPin(FRigVMDispatch_GetObjectProperty::ValueName.ToString());
+				if (ValuePin == nullptr)
+				{
+					InSettings.ReportError(TEXT("Could not find Value pin"));
+					continue;
+				}
+
+				if (!InArgs.Controller->ResolveWildCardPin(ValuePin->GetPinPath(), FRigVMRegistry::Get().GetTypeIndex(*Input->CPPType, Input->CPPTypeObject)))
+				{
+					InSettings.ReportError(TEXT("Could not resolve wildcard pin to function's return type"));
+					continue;
 				}
+
 				break;
 			}
 			case FAnimNextUniversalObjectLocatorBindingType::Function:
@@ -692,7 +788,7 @@ void FUniversalObjectLocatorBindingType::BuildBindingGraphFragment(const FRigVMC
 				if(Function == nullptr)
 				{
 					InSettings.ReportErrorf(TEXT("Could not resolve function call %s"), *LocatorBinding.Function.ToString());
-					return;
+					continue;
 				}
 
 				// Determine the call specialization we need (native vs script)
@@ -710,13 +806,13 @@ void FUniversalObjectLocatorBindingType::BuildBindingGraphFragment(const FRigVMC
 				if(GetterNode == nullptr)
 				{
 					InSettings.ReportError(TEXT("Could not spawn Call Object Accessor Function node"));
-					return;
+					continue;
 				}
 				URigVMPin* FunctionPin = GetterNode->FindPin(FRigVMDispatch_CallObjectAccessorFunctionBase::FunctionName.ToString());
 				if(FunctionPin == nullptr)
 				{
 					InSettings.ReportError(TEXT("Could not find Function pin"));
-					return;
+					continue;
 				}
 				FString DefaultPtrValue;
 				FAnimNextSoftFunctionPtr AnimNextSoftFunctionPtr;
@@ -726,8 +822,22 @@ void FUniversalObjectLocatorBindingType::BuildBindingGraphFragment(const FRigVMC
 				if(!bFunctionDefaultValueSet)
 				{
 					InSettings.ReportError(TEXT("Could not set Function pin's default value"));
-					return;
+					continue;
+				}
+
+				URigVMPin* ValuePin = GetterNode->FindPin(FRigVMDispatch_CallObjectAccessorFunctionBase::ValueName.ToString());
+				if (ValuePin == nullptr)
+				{
+					InSettings.ReportError(TEXT("Could not find Value pin"));
+					continue;
 				}
+
+				if (!InArgs.Controller->ResolveWildCardPin(ValuePin->GetPinPath(), FRigVMRegistry::Get().GetTypeIndex(*Input->CPPType, Input->CPPTypeObject)))
+				{
+					InSettings.ReportError(TEXT("Could not resolve wildcard pin to function's return type"));
+					continue;
+				}
+
 				break;
 			}
 			case FAnimNextUniversalObjectLocatorBindingType::HoistedFunction:
@@ -736,7 +846,7 @@ void FUniversalObjectLocatorBindingType::BuildBindingGraphFragment(const FRigVMC
 				if(Function == nullptr)
 				{
 					InSettings.ReportErrorf(TEXT("Could not resolve hoisted function call %s"), *LocatorBinding.Function.ToString());
-					return;
+					continue;
 				}
 
 				// Determine the call specialization we need (native vs script)
@@ -754,13 +864,13 @@ void FUniversalObjectLocatorBindingType::BuildBindingGraphFragment(const FRigVMC
 				if(GetterNode == nullptr)
 				{
 					InSettings.ReportError(TEXT("Could not spawn Call Hoisted Accessor Function node"));
-					return;
+					continue;
 				}
-				URigVMPin* FunctionPin = GetterNode->FindPin(FRigVMDispatch_CallObjectAccessorFunctionBase::FunctionName.ToString());
+				URigVMPin* FunctionPin = GetterNode->FindPin(FRigVMDispatch_CallHoistedAccessorFunctionBase::FunctionName.ToString());
 				if(FunctionPin == nullptr)
 				{
 					InSettings.ReportError(TEXT("Could not find Function pin"));
-					return;
+					continue;
 				}
 				FString DefaultValue;
 				FAnimNextSoftFunctionPtr AnimNextSoftFunctionPtr;
@@ -770,8 +880,22 @@ void FUniversalObjectLocatorBindingType::BuildBindingGraphFragment(const FRigVMC
 				if(!bFunctionDefaultValueSet)
 				{
 					InSettings.ReportError(TEXT("Could not set Function pin's default value"));
-					return;
+					continue;
 				}
+
+				URigVMPin* ValuePin = GetterNode->FindPin(FRigVMDispatch_CallHoistedAccessorFunctionBase::ValueName.ToString());
+				if (ValuePin == nullptr)
+				{
+					InSettings.ReportError(TEXT("Could not find Value pin"));
+					continue;
+				}
+
+				if (!InArgs.Controller->ResolveWildCardPin(ValuePin->GetPinPath(), FRigVMRegistry::Get().GetTypeIndex(*Input->CPPType, Input->CPPTypeObject)))
+				{
+					InSettings.ReportError(TEXT("Could not resolve wildcard pin to function's return type"));
+					continue;
+				}
+
 				break;  
 			}
 			default:
@@ -787,19 +911,19 @@ void FUniversalObjectLocatorBindingType::BuildBindingGraphFragment(const FRigVMC
 				if(ObjectPin == nullptr)
 				{
 					InSettings.ReportError(TEXT("Could not find object input pin"));
-					return;
+					continue;
 				}
 				bool bObjectPinLinked = InArgs.Controller->AddLink(ResolvedObjectPin, ObjectPin, false);
 				if(!bObjectPinLinked)
 				{
 					InSettings.ReportError(TEXT("Could not link object pins"));
-					return;
+					continue;
 				}
 				OutputValuePin = GetterNode->FindPin(FRigVMDispatch_GetObjectProperty::ValueName.ToString());
 				if(OutputValuePin == nullptr)
 				{
 					InSettings.ReportError(TEXT("Could find output value pin"));
-					return;
+					continue;
 				}
 			}
 			else
@@ -812,32 +936,32 @@ void FUniversalObjectLocatorBindingType::BuildBindingGraphFragment(const FRigVMC
 			if(VariableSetNode == nullptr)
 			{
 				InSettings.ReportError(TEXT("Could not spawn Variable Set node"));
-				return;
+				continue;
 			}
 			URigVMPin* VariableExecPin = VariableSetNode->FindPin(FRigVMStruct::ExecuteContextName.ToString());
 			if(VariableExecPin == nullptr)
 			{
 				InSettings.ReportError(TEXT("Could not find Variable Set Execute pin"));
-				return;
+				continue;
 			}
 			bool bExecPinLinked = InArgs.Controller->AddLink(ExecTail, VariableExecPin, false);
 			if(!bExecPinLinked)
 			{
 				InSettings.ReportError(TEXT("Could not link Variable Set Execute pin"));
-				return;
+				continue;
 			}
 			ExecTail = VariableExecPin;
 			URigVMPin* VariableInputPin = VariableSetNode->GetValuePin();
 			if(VariableInputPin == nullptr)
 			{
 				InSettings.ReportError(TEXT("Could not find Variable Set Value pin"));
-				return;
+				continue;
 			}
 			bool bValuePinLinked = InArgs.Controller->AddLink(OutputValuePin, VariableInputPin, false, ERigVMPinDirection::Invalid, /*bCreateCastNode*/true);
 			if(!bValuePinLinked)
 			{
 				InSettings.ReportError(TEXT("Could not link Variable Set Value pin"));
-				return;
+				continue;
 			}
 
 			AccumulatedYOffset += 100.0f;
diff --git a/Engine/Plugins/Experimental/UAF/UAF/Source/UAFUncookedOnly/Private/Variables/UniversalObjectLocatorBindingType.h b/Engine/Plugins/Experimental/UAF/UAF/Source/UAFUncookedOnly/Private/Variables/UniversalObjectLocatorBindingType.h
--- a/Engine/Plugins/Experimental/UAF/UAF/Source/UAFUncookedOnly/Private/Variables/UniversalObjectLocatorBindingType.h
+++ b/Engine/Plugins/Experimental/UAF/UAF/Source/UAFUncookedOnly/Private/Variables/UniversalObjectLocatorBindingType.h
@@ -21,6 +21,7 @@ private:
 	virtual TSharedRef<SWidget> CreateEditWidget(const TSharedRef<IPropertyHandle>& InPropertyHandle, const FAnimNextParamType& InType) const override;
 	virtual FText GetDisplayText(TConstStructView<FAnimNextVariableBindingData> InBindingData) const override;
 	virtual FText GetTooltipText(TConstStructView<FAnimNextVariableBindingData> InBindingData) const override;
+	virtual bool Validate(TConstStructView<FAnimNextVariableBindingData> InBindingData, FText& OutError) const override;
 	virtual void BuildBindingGraphFragment(const FRigVMCompileSettings& InSettings, const FBindingGraphFragmentArgs& InArgs, URigVMPin*& OutExecTail, FVector2D& OutLocation) const override;
 
 	static const UClass* GetClass(TConstStructView<FAnimNextVariableBindingData> InBindingData);
diff --git a/Engine/Plugins/Experimental/UAF/UAF/Source/UAFUncookedOnly/Public/Variables/IVariableBindingType.h b/Engine/Plugins/Experimental/UAF/UAF/Source/UAFUncookedOnly/Public/Variables/IVariableBindingType.h
--- a/Engine/Plugins/Experimental/UAF/UAF/Source/UAFUncookedOnly/Public/Variables/IVariableBindingType.h
+++ b/Engine/Plugins/Experimental/UAF/UAF/Source/UAFUncookedOnly/Public/Variables/IVariableBindingType.h
@@ -5,6 +5,7 @@
 #include "CoreMinimal.h"
 #include "Param/ParamType.h"
 
+class IAnimNextRigVMVariableInterface;
 struct FRigVMCompileSettings;
 class IPropertyHandle;
 class URigVMController;
@@ -75,6 +76,9 @@ private:
 	// Get the tooltip text for the specified instance ID
 	virtual FText GetTooltipText(TConstStructView<FAnimNextVariableBindingData> InBindingData) const = 0;
 
+	// Validate the binding data
+	virtual bool Validate(TConstStructView<FAnimNextVariableBindingData> InBindingData, FText& OutError) const = 0;
+
 	// Transforms the inputs into graph fragments. Called to convert variable bindings (derived from FAnimNextVariableBindingData) into intermediate
 	// RigVM graphs for consumption by the compiler  
 	virtual void BuildBindingGraphFragment(const FRigVMCompileSettings& InSettings, const FBindingGraphFragmentArgs& InArgs, URigVMPin*& OutExecTail, FVector2D& OutLocation) const = 0;
