Add support for space conversion in UAF value runtime

#Synced-CL 48759454

diff --git a/Engine/Plugins/Experimental/UAF/UAF/Source/UAF/Private/UAF/Attributes/Transformers/BoneSpace.cpp b/Engine/Plugins/Experimental/UAF/UAF/Source/UAF/Private/UAF/Attributes/Transformers/BoneSpace.cpp
new file mode 100644
--- /dev/null
+++ b/Engine/Plugins/Experimental/UAF/UAF/Source/UAF/Private/UAF/Attributes/Transformers/BoneSpace.cpp
@@ -0,0 +1,429 @@
+// Copyright Epic Games, Inc. All Rights Reserved.
+
+#include "UAF/Attributes/Transformers/BoneSpace.h"
+
+#include "UAF/Attributes/AttributeCollection.h"
+#include "UAF/Attributes/EngineAttributes.h"
+
+namespace UE::UAF::Transformers
+{
+	void FBoneSpace::LocalToComponent(const FAttributeCollection& InputValues, FAttributeCollection& OutputValues)
+	{
+		checkf(InputValues.GetValueSpace().GetType() == EValueSpaceType::Local, TEXT("Input must be in local space"));
+		checkf(!InputValues.GetValueSpace().IsAdditive(), TEXT("Input space cannot be additive"));
+
+		if (&InputValues != &OutputValues)
+		{
+			// Copy our input, we'll modify the bone transforms in-place
+			OutputValues = InputValues;
+		}
+
+		OutputValues.SetValueSpace(FValueSpace(EValueSpaceType::Component));
+
+		// Modify in place
+		TAttributeSetMap<FBoneTransformAnimationAttribute>* BoneTransforms = OutputValues.FindSetMap<FBoneTransformAnimationAttribute>(FAttributeMappingKey::MakeFromTo<FBoneTransformAnimationAttribute>());
+		if (BoneTransforms == nullptr)
+		{
+			// No bone transforms, nothing to convert
+			return;
+		}
+
+		const FAttributeTypedSetPtr& BoneSet = BoneTransforms->GetTypedSet();
+		const int32 NumBones = BoneTransforms->Num();
+
+		// Root is equivalent in both spaces, skip it
+		for (FAttributeSetIndex BoneSetIndex(1); BoneSetIndex < NumBones; ++BoneSetIndex)
+		{
+			const FAttributeSetIndex ParentBoneSetIndex = BoneSet->GetParentIndex(BoneSetIndex);
+			if (ParentBoneSetIndex.IsValid())
+			{
+				const FBoneTransformAnimationAttribute& ParentBoneTransform = (*BoneTransforms)[ParentBoneSetIndex];
+				FBoneTransformAnimationAttribute& BoneTransform = (*BoneTransforms)[BoneSetIndex];
+
+				FTransform::Multiply(&BoneTransform.Value, &BoneTransform.Value, &ParentBoneTransform.Value);
+
+				BoneTransform.Value.NormalizeRotation();
+			}
+			else
+			{
+				// TODO: Use the ref pose in component space
+				checkf(false, TEXT("Bones must have their parent within the same set to convert from Local to Component space"));
+			}
+		}
+	}
+
+	void FBoneSpace::ComponentToLocal(const FAttributeCollection& InputValues, FAttributeCollection& OutputValues)
+	{
+		checkf(InputValues.GetValueSpace().GetType() == EValueSpaceType::Component, TEXT("Input must be in component space"));
+		checkf(!InputValues.GetValueSpace().IsAdditive(), TEXT("Input space cannot be additive"));
+
+		if (&InputValues != &OutputValues)
+		{
+			// Copy our input, we'll modify the bone transforms in-place
+			OutputValues = InputValues;
+		}
+
+		OutputValues.SetValueSpace(FValueSpace(EValueSpaceType::Local));
+
+		// Modify in place
+		TAttributeSetMap<FBoneTransformAnimationAttribute>* BoneTransforms = OutputValues.FindSetMap<FBoneTransformAnimationAttribute>(FAttributeMappingKey::MakeFromTo<FBoneTransformAnimationAttribute>());
+		if (BoneTransforms == nullptr)
+		{
+			// No bone transforms, nothing to convert
+			return;
+		}
+
+		const FAttributeTypedSetPtr& BoneSet = BoneTransforms->GetTypedSet();
+		const int32 NumBones = BoneTransforms->Num();
+
+		// Root is equivalent in both spaces, skip it
+		// Iterate backwards since we need the parent transforms intact
+		for (FAttributeSetIndex BoneSetIndex(NumBones - 1); BoneSetIndex > 0; --BoneSetIndex)
+		{
+			const FAttributeSetIndex ParentBoneSetIndex = BoneSet->GetParentIndex(BoneSetIndex);
+			if (ParentBoneSetIndex.IsValid())
+			{
+				const FBoneTransformAnimationAttribute& ParentBoneTransform = (*BoneTransforms)[ParentBoneSetIndex];
+				FBoneTransformAnimationAttribute& BoneTransform = (*BoneTransforms)[BoneSetIndex];
+
+				BoneTransform.Value.SetToRelativeTransform(ParentBoneTransform.Value);
+				BoneTransform.Value.NormalizeRotation();
+			}
+			else
+			{
+				// TODO: Use the ref pose in component space
+				checkf(false, TEXT("Bones must have their parent within the same set to convert from Component to Local space"));
+			}
+		}
+	}
+
+	void FBoneSpace::LocalToMeshRotation(const FAttributeCollection& InputValues, FAttributeCollection& OutputValues)
+	{
+		const FValueSpace InputSpace = InputValues.GetValueSpace();
+		checkf(InputSpace.GetType() == EValueSpaceType::Local || InputSpace.GetType() == EValueSpaceType::Mixed, TEXT("Input must be in local or mixed space"));
+		checkf(!InputSpace.IsAdditive(), TEXT("Input space cannot be additive"));
+
+		if (&InputValues != &OutputValues)
+		{
+			// Copy our input, we'll modify the bone transforms in-place
+			OutputValues = InputValues;
+		}
+
+		EMixedSpaceFlags MixedSpaceFlags = InputSpace.GetMixedSpaceFlags();
+		checkf(!(MixedSpaceFlags & EMixedSpaceFlags::MeshRotation), TEXT("Input is already in mesh rotation space"));
+		checkf(!(MixedSpaceFlags & EMixedSpaceFlags::RootRotation), TEXT("Input is already in root rotation space"));
+		EnumAddFlags(MixedSpaceFlags, EMixedSpaceFlags::MeshRotation);
+		OutputValues.SetValueSpace(FValueSpace(MixedSpaceFlags));
+
+		// Modify in place
+		TAttributeSetMap<FBoneTransformAnimationAttribute>* BoneTransforms = OutputValues.FindSetMap<FBoneTransformAnimationAttribute>(FAttributeMappingKey::MakeFromTo<FBoneTransformAnimationAttribute>());
+		if (BoneTransforms == nullptr)
+		{
+			// No bone transforms, nothing to convert
+			return;
+		}
+
+		const FAttributeTypedSetPtr& BoneSet = BoneTransforms->GetTypedSet();
+		const int32 NumBones = BoneTransforms->Num();
+
+		// Root is equivalent in both spaces, skip it
+		for (FAttributeSetIndex BoneSetIndex(1); BoneSetIndex < NumBones; ++BoneSetIndex)
+		{
+			const FAttributeSetIndex ParentBoneSetIndex = BoneSet->GetParentIndex(BoneSetIndex);
+			if (ParentBoneSetIndex.IsValid())
+			{
+				const FBoneTransformAnimationAttribute& ParentBoneTransform = (*BoneTransforms)[ParentBoneSetIndex];
+				FBoneTransformAnimationAttribute& BoneTransform = (*BoneTransforms)[BoneSetIndex];
+
+				BoneTransform.Value.SetRotation(ParentBoneTransform.Value.GetRotation() * BoneTransform.Value.GetRotation());
+				BoneTransform.Value.NormalizeRotation();
+			}
+			else
+			{
+				// TODO: Use the ref pose in component space
+				checkf(false, TEXT("Bones must have their parent within the same set to convert from Local to Mesh Rotation space"));
+			}
+		}
+	}
+
+	void FBoneSpace::MeshRotationToLocal(const FAttributeCollection& InputValues, FAttributeCollection& OutputValues)
+	{
+		const FValueSpace InputSpace = InputValues.GetValueSpace();
+		EMixedSpaceFlags MixedSpaceFlags = InputSpace.GetMixedSpaceFlags();
+		checkf(InputSpace.GetType() == EValueSpaceType::Mixed && !!(MixedSpaceFlags & EMixedSpaceFlags::MeshRotation), TEXT("Input must be in mesh rotation space"));
+		checkf(!InputSpace.IsAdditive(), TEXT("Input space cannot be additive"));
+
+		if (&InputValues != &OutputValues)
+		{
+			// Copy our input, we'll modify the bone transforms in-place
+			OutputValues = InputValues;
+		}
+
+		EnumRemoveFlags(MixedSpaceFlags, EMixedSpaceFlags::MeshRotation);
+
+		FValueSpace OutputSpace;
+		if (MixedSpaceFlags == EMixedSpaceFlags::None)
+		{
+			// We aren't mixed space anymore
+			OutputSpace = FValueSpace(EValueSpaceType::Local);
+		}
+		else
+		{
+			OutputSpace = FValueSpace(MixedSpaceFlags);
+		}
+		OutputValues.SetValueSpace(OutputSpace);
+
+		// Modify in place
+		TAttributeSetMap<FBoneTransformAnimationAttribute>* BoneTransforms = OutputValues.FindSetMap<FBoneTransformAnimationAttribute>(FAttributeMappingKey::MakeFromTo<FBoneTransformAnimationAttribute>());
+		if (BoneTransforms == nullptr)
+		{
+			// No bone transforms, nothing to convert
+			return;
+		}
+
+		const FAttributeTypedSetPtr& BoneSet = BoneTransforms->GetTypedSet();
+		const int32 NumBones = BoneTransforms->Num();
+
+		// Root is equivalent in both spaces, skip it
+		// Iterate backwards since we need the parent transforms intact
+		for (FAttributeSetIndex BoneSetIndex(NumBones - 1); BoneSetIndex > 0; --BoneSetIndex)
+		{
+			const FAttributeSetIndex ParentBoneSetIndex = BoneSet->GetParentIndex(BoneSetIndex);
+			if (ParentBoneSetIndex.IsValid())
+			{
+				const FBoneTransformAnimationAttribute& ParentBoneTransform = (*BoneTransforms)[ParentBoneSetIndex];
+				FBoneTransformAnimationAttribute& BoneTransform = (*BoneTransforms)[BoneSetIndex];
+
+				BoneTransform.Value.SetRotation(ParentBoneTransform.Value.GetRotation().Inverse() * BoneTransform.Value.GetRotation());
+				BoneTransform.Value.NormalizeRotation();
+			}
+			else
+			{
+				// TODO: Use the ref pose in component space
+				checkf(false, TEXT("Bones must have their parent within the same set to convert from Mesh Rotation to Local space"));
+			}
+		}
+	}
+
+	void FBoneSpace::LocalToRootRotation(const FAttributeCollection& InputValues, FAttributeCollection& OutputValues)
+	{
+		const FValueSpace InputSpace = InputValues.GetValueSpace();
+		checkf(InputSpace.GetType() == EValueSpaceType::Local || InputSpace.GetType() == EValueSpaceType::Mixed, TEXT("Input must be in local or mixed space"));
+		checkf(!InputSpace.IsAdditive(), TEXT("Input space cannot be additive"));
+
+		if (&InputValues != &OutputValues)
+		{
+			// Copy our input, we'll modify the bone transforms in-place
+			OutputValues = InputValues;
+		}
+
+		EMixedSpaceFlags MixedSpaceFlags = InputSpace.GetMixedSpaceFlags();
+		checkf(!(MixedSpaceFlags & EMixedSpaceFlags::MeshRotation), TEXT("Input is already in mesh rotation space"));
+		checkf(!(MixedSpaceFlags & EMixedSpaceFlags::RootRotation), TEXT("Input is already in root rotation space"));
+		EnumAddFlags(MixedSpaceFlags, EMixedSpaceFlags::RootRotation);
+		OutputValues.SetValueSpace(FValueSpace(MixedSpaceFlags));
+
+		// Modify in place
+		TAttributeSetMap<FBoneTransformAnimationAttribute>* BoneTransforms = OutputValues.FindSetMap<FBoneTransformAnimationAttribute>(FAttributeMappingKey::MakeFromTo<FBoneTransformAnimationAttribute>());
+		if (BoneTransforms == nullptr)
+		{
+			// No bone transforms, nothing to convert
+			return;
+		}
+
+		const FAttributeTypedSetPtr& BoneSet = BoneTransforms->GetTypedSet();
+		const int32 NumBones = BoneTransforms->Num();
+
+		const FBoneTransformAnimationAttribute Identity{ FTransform::Identity };
+
+		// Root is equivalent in both spaces, skip it
+		for (FAttributeSetIndex BoneSetIndex(1); BoneSetIndex < NumBones; ++BoneSetIndex)
+		{
+			const FAttributeSetIndex ParentBoneSetIndex = BoneSet->GetParentIndex(BoneSetIndex);
+			if (ParentBoneSetIndex.IsValid())
+			{
+				const FBoneTransformAnimationAttribute& ParentBoneTransform = ParentBoneSetIndex.IsRootBone() ? Identity : (*BoneTransforms)[ParentBoneSetIndex];
+				FBoneTransformAnimationAttribute& BoneTransform = (*BoneTransforms)[BoneSetIndex];
+
+				BoneTransform.Value.SetRotation(ParentBoneTransform.Value.GetRotation() * BoneTransform.Value.GetRotation());
+				BoneTransform.Value.NormalizeRotation();
+			}
+			else
+			{
+				// TODO: Use the ref pose in component space
+				checkf(false, TEXT("Bones must have their parent within the same set to convert from Local to Root Rotation space"));
+			}
+		}
+	}
+
+	void FBoneSpace::RootRotationToLocal(const FAttributeCollection& InputValues, FAttributeCollection& OutputValues)
+	{
+		const FValueSpace InputSpace = InputValues.GetValueSpace();
+		EMixedSpaceFlags MixedSpaceFlags = InputSpace.GetMixedSpaceFlags();
+		checkf(InputSpace.GetType() == EValueSpaceType::Mixed && !!(MixedSpaceFlags & EMixedSpaceFlags::RootRotation), TEXT("Input must be in root rotation space"));
+		checkf(!InputSpace.IsAdditive(), TEXT("Input space cannot be additive"));
+
+		if (&InputValues != &OutputValues)
+		{
+			// Copy our input, we'll modify the bone transforms in-place
+			OutputValues = InputValues;
+		}
+
+		EnumRemoveFlags(MixedSpaceFlags, EMixedSpaceFlags::RootRotation);
+
+		FValueSpace OutputSpace;
+		if (MixedSpaceFlags == EMixedSpaceFlags::None)
+		{
+			// We aren't mixed space anymore
+			OutputSpace = FValueSpace(EValueSpaceType::Local);
+		}
+		else
+		{
+			OutputSpace = FValueSpace(MixedSpaceFlags);
+		}
+		OutputValues.SetValueSpace(OutputSpace);
+
+		// Modify in place
+		TAttributeSetMap<FBoneTransformAnimationAttribute>* BoneTransforms = OutputValues.FindSetMap<FBoneTransformAnimationAttribute>(FAttributeMappingKey::MakeFromTo<FBoneTransformAnimationAttribute>());
+		if (BoneTransforms == nullptr)
+		{
+			// No bone transforms, nothing to convert
+			return;
+		}
+
+		const FAttributeTypedSetPtr& BoneSet = BoneTransforms->GetTypedSet();
+		const int32 NumBones = BoneTransforms->Num();
+
+		const FBoneTransformAnimationAttribute Identity{ FTransform::Identity };
+
+		// Root is equivalent in both spaces, skip it
+		// Iterate backwards since we need the parent transforms intact
+		for (FAttributeSetIndex BoneSetIndex(NumBones - 1); BoneSetIndex > 0; --BoneSetIndex)
+		{
+			const FAttributeSetIndex ParentBoneSetIndex = BoneSet->GetParentIndex(BoneSetIndex);
+			if (ParentBoneSetIndex.IsValid())
+			{
+				const FBoneTransformAnimationAttribute& ParentBoneTransform = ParentBoneSetIndex.IsRootBone() ? Identity : (*BoneTransforms)[ParentBoneSetIndex];
+				FBoneTransformAnimationAttribute& BoneTransform = (*BoneTransforms)[BoneSetIndex];
+
+				BoneTransform.Value.SetRotation(ParentBoneTransform.Value.GetRotation().Inverse() * BoneTransform.Value.GetRotation());
+				BoneTransform.Value.NormalizeRotation();
+			}
+			else
+			{
+				// TODO: Use the ref pose in component space
+				checkf(false, TEXT("Bones must have their parent within the same set to convert from Root Rotation to Local space"));
+			}
+		}
+	}
+
+	void FBoneSpace::LocalToMeshScale(const FAttributeCollection& InputValues, FAttributeCollection& OutputValues)
+	{
+		const FValueSpace InputSpace = InputValues.GetValueSpace();
+		checkf(InputSpace.GetType() == EValueSpaceType::Local || InputSpace.GetType() == EValueSpaceType::Mixed, TEXT("Input must be in local or mixed space"));
+		checkf(!InputSpace.IsAdditive(), TEXT("Input space cannot be additive"));
+
+		if (&InputValues != &OutputValues)
+		{
+			// Copy our input, we'll modify the bone transforms in-place
+			OutputValues = InputValues;
+		}
+
+		EMixedSpaceFlags MixedSpaceFlags = InputSpace.GetMixedSpaceFlags();
+		checkf(!(MixedSpaceFlags & EMixedSpaceFlags::MeshScale), TEXT("Input is already in mesh scale space"));
+		EnumAddFlags(MixedSpaceFlags, EMixedSpaceFlags::MeshScale);
+		OutputValues.SetValueSpace(FValueSpace(MixedSpaceFlags));
+
+		// Modify in place
+		TAttributeSetMap<FBoneTransformAnimationAttribute>* BoneTransforms = OutputValues.FindSetMap<FBoneTransformAnimationAttribute>(FAttributeMappingKey::MakeFromTo<FBoneTransformAnimationAttribute>());
+		if (BoneTransforms == nullptr)
+		{
+			// No bone transforms, nothing to convert
+			return;
+		}
+
+		const FAttributeTypedSetPtr& BoneSet = BoneTransforms->GetTypedSet();
+		const int32 NumBones = BoneTransforms->Num();
+
+		const FBoneTransformAnimationAttribute Identity{ FTransform::Identity };
+
+		// Root is equivalent in both spaces, skip it
+		for (FAttributeSetIndex BoneSetIndex(1); BoneSetIndex < NumBones; ++BoneSetIndex)
+		{
+			const FAttributeSetIndex ParentBoneSetIndex = BoneSet->GetParentIndex(BoneSetIndex);
+			if (ParentBoneSetIndex.IsValid())
+			{
+				const FBoneTransformAnimationAttribute& ParentBoneTransform = ParentBoneSetIndex.IsRootBone() ? Identity : (*BoneTransforms)[ParentBoneSetIndex];
+				FBoneTransformAnimationAttribute& BoneTransform = (*BoneTransforms)[BoneSetIndex];
+
+				BoneTransform.Value.SetScale3D(ParentBoneTransform.Value.GetScale3D() * BoneTransform.Value.GetScale3D());
+			}
+			else
+			{
+				// TODO: Use the ref pose in component space
+				checkf(false, TEXT("Bones must have their parent within the same set to convert from Local to Mesh Scale space"));
+			}
+		}
+	}
+
+	void FBoneSpace::MeshScaleToLocal(const FAttributeCollection& InputValues, FAttributeCollection& OutputValues)
+	{
+		const FValueSpace InputSpace = InputValues.GetValueSpace();
+		EMixedSpaceFlags MixedSpaceFlags = InputSpace.GetMixedSpaceFlags();
+		checkf(InputSpace.GetType() == EValueSpaceType::Mixed && !!(MixedSpaceFlags & EMixedSpaceFlags::MeshScale), TEXT("Input must be in mesh scale space"));
+		checkf(!InputSpace.IsAdditive(), TEXT("Input space cannot be additive"));
+
+		if (&InputValues != &OutputValues)
+		{
+			// Copy our input, we'll modify the bone transforms in-place
+			OutputValues = InputValues;
+		}
+
+		EnumRemoveFlags(MixedSpaceFlags, EMixedSpaceFlags::MeshScale);
+
+		FValueSpace OutputSpace;
+		if (MixedSpaceFlags == EMixedSpaceFlags::None)
+		{
+			// We aren't mixed space anymore
+			OutputSpace = FValueSpace(EValueSpaceType::Local);
+		}
+		else
+		{
+			OutputSpace = FValueSpace(MixedSpaceFlags);
+		}
+		OutputValues.SetValueSpace(OutputSpace);
+
+		// Modify in place
+		TAttributeSetMap<FBoneTransformAnimationAttribute>* BoneTransforms = OutputValues.FindSetMap<FBoneTransformAnimationAttribute>(FAttributeMappingKey::MakeFromTo<FBoneTransformAnimationAttribute>());
+		if (BoneTransforms == nullptr)
+		{
+			// No bone transforms, nothing to convert
+			return;
+		}
+
+		const FAttributeTypedSetPtr& BoneSet = BoneTransforms->GetTypedSet();
+		const int32 NumBones = BoneTransforms->Num();
+
+		const FBoneTransformAnimationAttribute Identity{ FTransform::Identity };
+
+		// Root is equivalent in both spaces, skip it
+		// Iterate backwards since we need the parent transforms intact
+		for (FAttributeSetIndex BoneSetIndex(NumBones - 1); BoneSetIndex > 0; --BoneSetIndex)
+		{
+			const FAttributeSetIndex ParentBoneSetIndex = BoneSet->GetParentIndex(BoneSetIndex);
+			if (ParentBoneSetIndex.IsValid())
+			{
+				const FBoneTransformAnimationAttribute& ParentBoneTransform = ParentBoneSetIndex.IsRootBone() ? Identity : (*BoneTransforms)[ParentBoneSetIndex];
+				FBoneTransformAnimationAttribute& BoneTransform = (*BoneTransforms)[BoneSetIndex];
+
+				const FVector ParentScaleInv = FTransform::GetSafeScaleReciprocal(ParentBoneTransform.Value.GetScale3D(), UE_SMALL_NUMBER);
+				BoneTransform.Value.SetScale3D(ParentScaleInv * BoneTransform.Value.GetScale3D());
+			}
+			else
+			{
+				// TODO: Use the ref pose in component space
+				checkf(false, TEXT("Bones must have their parent within the same set to convert from Local to Mesh Scale space"));
+			}
+		}
+	}
+}
diff --git a/Engine/Plugins/Experimental/UAF/UAF/Source/UAF/Public/UAF/Attributes/Transformers/BoneSpace.h b/Engine/Plugins/Experimental/UAF/UAF/Source/UAF/Public/UAF/Attributes/Transformers/BoneSpace.h
new file mode 100644
--- /dev/null
+++ b/Engine/Plugins/Experimental/UAF/UAF/Source/UAF/Public/UAF/Attributes/Transformers/BoneSpace.h
@@ -0,0 +1,50 @@
+// Copyright Epic Games, Inc. All Rights Reserved.
+
+#pragma once
+
+#include "CoreMinimal.h"
+
+#define UE_API UAF_API
+
+namespace UE::UAF
+{
+	class FAttributeCollection;
+}
+
+namespace UE::UAF::Transformers
+{
+	// Bone space conversion utility functions
+	// Not a transformer, but transformer-like
+	struct FBoneSpace final
+	{
+		// Converts bone transforms between Local Space and Component Space
+		// These functions support outputting to their input in place
+		static UE_API void LocalToComponent(const FAttributeCollection& InputValues, FAttributeCollection& OutputValues);
+		static UE_API void ComponentToLocal(const FAttributeCollection& InputValues, FAttributeCollection& OutputValues);
+
+		// Converts bone transforms between Local Space and Mesh Rotation Space
+		// Mesh Rotation Space assumes that bone transform scales are uniform and will yield incorrect results
+		// if non-uniform (aka 3D) scale is used. It transforms the rotations into component space by ignoring
+		// translation and scale. The output values will thus be in Mixed Space.
+		// These functions support outputting to their input in place
+		static UE_API void LocalToMeshRotation(const FAttributeCollection& InputValues, FAttributeCollection& OutputValues);
+		static UE_API void MeshRotationToLocal(const FAttributeCollection& InputValues, FAttributeCollection& OutputValues);
+
+		// Converts bone transforms between Local Space and Root Rotation Space
+		// Root Rotation Space is similar to Mesh Rotation space except that the root bone transform is ignored
+		// and assumed to be the identity transform (root corrections are ignored).
+		// These functions support outputting to their input in place
+		static UE_API void LocalToRootRotation(const FAttributeCollection& InputValues, FAttributeCollection& OutputValues);
+		static UE_API void RootRotationToLocal(const FAttributeCollection& InputValues, FAttributeCollection& OutputValues);
+
+		// Converts bone transforms between Local Space and Mesh Scale Space
+		// Mesh Scale Space assumes that bone transform scales are uniform and will yield incorrect results
+		// if non-uniform (aka 3D) scale is used. It transforms the scales into component space by ignoring
+		// rotation and translation. The output values will thus be in a Mixed Space.
+		// These functions support outputting to their input in place
+		static UE_API void LocalToMeshScale(const FAttributeCollection& InputValues, FAttributeCollection& OutputValues);
+		static UE_API void MeshScaleToLocal(const FAttributeCollection& InputValues, FAttributeCollection& OutputValues);
+	};
+}
+
+#undef UE_API
diff --git a/Engine/Plugins/Experimental/UAF/UAFAnimGraph/Source/UAFAnimGraph/Private/EvaluationVM/Tasks/ConvertRotationsToLocalSpace.cpp b/Engine/Plugins/Experimental/UAF/UAFAnimGraph/Source/UAFAnimGraph/Private/EvaluationVM/Tasks/ConvertRotationsToLocalSpace.cpp
--- a/Engine/Plugins/Experimental/UAF/UAFAnimGraph/Source/UAFAnimGraph/Private/EvaluationVM/Tasks/ConvertRotationsToLocalSpace.cpp
+++ b/Engine/Plugins/Experimental/UAF/UAFAnimGraph/Source/UAFAnimGraph/Private/EvaluationVM/Tasks/ConvertRotationsToLocalSpace.cpp
@@ -4,6 +4,7 @@
 
 #include "EvaluationVM/EvaluationVM.h"
 #include "EvaluationVM/Tasks/ConvertRotationSpaceDefs.h"
+#include "UAF/Attributes/Transformers/BoneSpace.h"
 
 #if UE_ANIM_CONVERT_ROTATION_SPACE_TASKS_ISPC
 #include "ConvertRotationSpace.ispc.generated.h"
@@ -25,43 +26,55 @@ void FAnimNextConvertRotationsToLocalSpaceTask::Execute(UE::UAF::FEvaluationVM&
 
 	if (VM.GetContext().IsValid())
 	{
-		// TODO: Implement with new attribute runtime
-		return;
-	}
+		for (int32 PoseIndex = 0; PoseIndex < NumPoses; ++PoseIndex)
+		{
+			TUniquePtr<FAttributeCollection>* Collection = VM.PeekValueMutable<TUniquePtr<FAttributeCollection>>(ATTRIBUTE_STACK_NAME, PoseIndex);
+			if (Collection == nullptr || !Collection->IsValid())
+			{
+				// No pose to convert, we are done
+				UE_LOG(LogAnimation, Warning, TEXT("FAnimNextConvertRotationsToLocalSpaceTask::Execute: Failed to fully execute - No pose data to convert."));
+				break;
+			}
 
-	if (!EnumHasAnyFlags(VM.GetFlags(), EEvaluationFlags::Bones))
-	{
-		return;
+			Transformers::FBoneSpace::MeshRotationToLocal(**Collection, **Collection);
+		}
 	}
-
-	for (int32 PoseIndex = 0; PoseIndex < NumPoses; ++PoseIndex)
+	else
 	{
-		TUniquePtr<FKeyframeState>* Keyframe = VM.PeekValueMutable<TUniquePtr<FKeyframeState>>(KEYFRAME_STACK_NAME, PoseIndex);
-		if (!Keyframe || !Keyframe->IsValid())
+		if (!EnumHasAnyFlags(VM.GetFlags(), EEvaluationFlags::Bones))
 		{
-			// No pose to convert
-			UE_LOG(LogAnimation, Warning, TEXT("FAnimNextConvertRotationsToLocalSpaceTask::Execute: Failed to fully execute - No pose data to convert."));
 			return;
 		}
 
-		FLODPoseStack& Pose = (*Keyframe)->Pose;
-		TArrayView<FQuat> PoseRotations = Pose.LocalTransforms.Rotations;
+		for (int32 PoseIndex = 0; PoseIndex < NumPoses; ++PoseIndex)
+		{
+			TUniquePtr<FKeyframeState>* Keyframe = VM.PeekValueMutable<TUniquePtr<FKeyframeState>>(KEYFRAME_STACK_NAME, PoseIndex);
+			if (!Keyframe || !Keyframe->IsValid())
+			{
+				// No pose to convert
+				UE_LOG(LogAnimation, Warning, TEXT("FAnimNextConvertRotationsToLocalSpaceTask::Execute: Failed to fully execute - No pose data to convert."));
+				return;
+			}
+
+			FLODPoseStack& Pose = (*Keyframe)->Pose;
+			TArrayView<FQuat> PoseRotations = Pose.LocalTransforms.Rotations;
 
 #if UE_ANIM_CONVERT_ROTATION_SPACE_TASKS_ISPC
-		static_assert(sizeof(FBoneIndexType) == sizeof(uint16)); // ISPC needs updating if assert fails
+			static_assert(sizeof(FBoneIndexType) == sizeof(uint16)); // ISPC needs updating if assert fails
 
-		const TArrayView<const FBoneIndexType> ParentIndices = Pose.GetRefPose().GetLODBoneIndexToParentLODBoneIndexMap(Pose.LODLevel);
+			const TArrayView<const FBoneIndexType> ParentIndices = Pose.GetRefPose().GetLODBoneIndexToParentLODBoneIndexMap(Pose.LODLevel);
 
-		ispc::ConvertRotationsToLocalSpace(
-			reinterpret_cast<ispc::FVector4*>(PoseRotations.GetData()),
-			reinterpret_cast<const uint16*>(ParentIndices.GetData()),
-			Pose.GetNumBones());
+			ispc::ConvertRotationsToLocalSpace(
+				reinterpret_cast<ispc::FVector4*>(PoseRotations.GetData()),
+				reinterpret_cast<const uint16*>(ParentIndices.GetData()),
+				Pose.GetNumBones());
 #else
-		for (FBoneIndexType Index = PoseRotations.Num() - 1; Index > 0; --Index)
-		{
-			const FBoneIndexType ParentIndex = Pose.GetLODBoneParentIndex(Index);
-			PoseRotations[Index] = PoseRotations[ParentIndex].Inverse() * PoseRotations[Index];
-		}
+			for (FBoneIndexType Index = PoseRotations.Num() - 1; Index > 0; --Index)
+			{
+				const FBoneIndexType ParentIndex = Pose.GetLODBoneParentIndex(Index);
+				PoseRotations[Index] = PoseRotations[ParentIndex].Inverse() * PoseRotations[Index];
+			}
 #endif // UE_ANIM_CONVERT_ROTATION_SPACE_TASKS_ISPC
+		}
 	}
 }
diff --git a/Engine/Plugins/Experimental/UAF/UAFAnimGraph/Source/UAFAnimGraph/Private/EvaluationVM/Tasks/ConvertRotationsToMeshSpace.cpp b/Engine/Plugins/Experimental/UAF/UAFAnimGraph/Source/UAFAnimGraph/Private/EvaluationVM/Tasks/ConvertRotationsToMeshSpace.cpp
--- a/Engine/Plugins/Experimental/UAF/UAFAnimGraph/Source/UAFAnimGraph/Private/EvaluationVM/Tasks/ConvertRotationsToMeshSpace.cpp
+++ b/Engine/Plugins/Experimental/UAF/UAFAnimGraph/Source/UAFAnimGraph/Private/EvaluationVM/Tasks/ConvertRotationsToMeshSpace.cpp
@@ -4,6 +4,7 @@
 
 #include "EvaluationVM/EvaluationVM.h"
 #include "EvaluationVM/Tasks/ConvertRotationSpaceDefs.h"
+#include "UAF/Attributes/Transformers/BoneSpace.h"
 
 #if UE_ANIM_CONVERT_ROTATION_SPACE_TASKS_ISPC
 #include "ConvertRotationSpace.ispc.generated.h"
@@ -25,43 +26,55 @@ void FAnimNextConvertRotationsToMeshSpaceTask::Execute(UE::UAF::FEvaluationVM& V
 
 	if (VM.GetContext().IsValid())
 	{
-		// TODO: Implement with new attribute runtime
-		return;
-	}
+		for (int32 PoseIndex = 0; PoseIndex < NumPoses; ++PoseIndex)
+		{
+			TUniquePtr<FAttributeCollection>* Collection = VM.PeekValueMutable<TUniquePtr<FAttributeCollection>>(ATTRIBUTE_STACK_NAME, PoseIndex);
+			if (Collection == nullptr || !Collection->IsValid())
+			{
+				// No pose to convert, we are done
+				UE_LOG(LogAnimation, Warning, TEXT("FAnimNextConvertRotationsToMeshSpaceTask::Execute: Failed to fully execute - No pose data to convert."));
+				break;
+			}
 
-	if (!EnumHasAnyFlags(VM.GetFlags(), EEvaluationFlags::Bones))
-	{
-		return;
+			Transformers::FBoneSpace::LocalToMeshRotation(**Collection, **Collection);
+		}
 	}
-
-	for (int32 PoseIndex = 0; PoseIndex < NumPoses; ++PoseIndex)
+	else
 	{
-		TUniquePtr<FKeyframeState>* Keyframe = VM.PeekValueMutable<TUniquePtr<FKeyframeState>>(KEYFRAME_STACK_NAME, PoseIndex);
-		if (!Keyframe || !Keyframe->IsValid())
+		if (!EnumHasAnyFlags(VM.GetFlags(), EEvaluationFlags::Bones))
 		{
-			// No pose to convert
-			UE_LOG(LogAnimation, Warning, TEXT("FAnimNextConvertRotationsToMeshSpaceTask::Execute: Failed to fully execute - No pose data to convert."));
 			return;
 		}
 
-		FLODPoseStack& Pose = (*Keyframe)->Pose;
-		TArrayView<FQuat> PoseRotations = Pose.LocalTransforms.Rotations;
+		for (int32 PoseIndex = 0; PoseIndex < NumPoses; ++PoseIndex)
+		{
+			TUniquePtr<FKeyframeState>* Keyframe = VM.PeekValueMutable<TUniquePtr<FKeyframeState>>(KEYFRAME_STACK_NAME, PoseIndex);
+			if (!Keyframe || !Keyframe->IsValid())
+			{
+				// No pose to convert
+				UE_LOG(LogAnimation, Warning, TEXT("FAnimNextConvertRotationsToMeshSpaceTask::Execute: Failed to fully execute - No pose data to convert."));
+				return;
+			}
+
+			FLODPoseStack& Pose = (*Keyframe)->Pose;
+			TArrayView<FQuat> PoseRotations = Pose.LocalTransforms.Rotations;
 
 #if UE_ANIM_CONVERT_ROTATION_SPACE_TASKS_ISPC
-		static_assert(sizeof(FBoneIndexType) == sizeof(uint16)); // ISPC needs updating if assert fails
+			static_assert(sizeof(FBoneIndexType) == sizeof(uint16)); // ISPC needs updating if assert fails
 
-		const TArrayView<const FBoneIndexType> ParentIndices = Pose.GetRefPose().GetLODBoneIndexToParentLODBoneIndexMap(Pose.LODLevel);
+			const TArrayView<const FBoneIndexType> ParentIndices = Pose.GetRefPose().GetLODBoneIndexToParentLODBoneIndexMap(Pose.LODLevel);
 
-		ispc::ConvertRotationsToMeshSpace(
-			reinterpret_cast<ispc::FVector4*>(PoseRotations.GetData()),
-			reinterpret_cast<const uint16*>(ParentIndices.GetData()),
-			Pose.GetNumBones());
+			ispc::ConvertRotationsToMeshSpace(
+				reinterpret_cast<ispc::FVector4*>(PoseRotations.GetData()),
+				reinterpret_cast<const uint16*>(ParentIndices.GetData()),
+				Pose.GetNumBones());
 #else
-		for (FBoneIndexType Index = 1; Index < PoseRotations.Num(); ++Index)
-		{
-			const FBoneIndexType ParentIndex = Pose.GetLODBoneParentIndex(Index);
-			PoseRotations[Index] = PoseRotations[ParentIndex] * PoseRotations[Index];
-		}
+			for (FBoneIndexType Index = 1; Index < PoseRotations.Num(); ++Index)
+			{
+				const FBoneIndexType ParentIndex = Pose.GetLODBoneParentIndex(Index);
+				PoseRotations[Index] = PoseRotations[ParentIndex] * PoseRotations[Index];
+			}
 #endif // UE_ANIM_CONVERT_ROTATION_SPACE_TASKS_ISPC
+		}
 	}
 }
