UAF: Skip input variables when generating programmatic variables & pins for event shims

Fixes compilation time regression now that functions have external variable pins

#Synced-CL 48758736

diff --git a/Engine/Plugins/Experimental/UAF/UAF/Source/UAFUncookedOnly/Private/AnimNextRigVMAssetEditorData.cpp b/Engine/Plugins/Experimental/UAF/UAF/Source/UAFUncookedOnly/Private/AnimNextRigVMAssetEditorData.cpp
--- a/Engine/Plugins/Experimental/UAF/UAF/Source/UAFUncookedOnly/Private/AnimNextRigVMAssetEditorData.cpp
+++ b/Engine/Plugins/Experimental/UAF/UAF/Source/UAFUncookedOnly/Private/AnimNextRigVMAssetEditorData.cpp
@@ -865,6 +865,11 @@ void UAnimNextRigVMAssetEditorData::BuildFunctionWrapperEventVariables(FAnimNext
 			{
 				if (Argument.IsValid())
 				{
+					if (Argument.bIsInputVariable)
+					{
+						continue;
+					}
+				
 					FRigVMGraphFunctionArgument InternallyNamedArgument = Argument;
 					InternallyNamedArgument.Name = FName(FUtils::MakeFunctionWrapperVariableName(FunctionHeader.Name, Argument.Name));
 					InContext.ProgrammaticVariables.Add(FAnimNextProgrammaticVariable::FromRigVMGraphFunctionArgument(InternallyNamedArgument));
@@ -959,7 +964,7 @@ void UAnimNextRigVMAssetEditorData::BuildFunctionWrapperEvents(FAnimNextRigVMAss
 		for (const FRigVMGraphFunctionArgument& Argument : FunctionHeader.Arguments)
 		{
 			// Execution context is captured as arg pins, skip those for internal variable gen
-			if (Argument.Direction == ERigVMPinDirection::IO || !Argument.IsValid())
+			if (Argument.Direction == ERigVMPinDirection::IO || !Argument.IsValid() || Argument.bIsInputVariable)
 			{
 				continue;
 			}
