Use correct reference skeleton (mesh one) in case skeleton remapping isn't available when retargetting using orient-scale
#TODO AbstractHierarchy code-path doesn't have acccess to ref-skeleton (USkeleton only) so this needs future fix-up

#Synced-CL 48761512

diff --git a/Engine/Plugins/Experimental/UAF/UAF/Source/UAF/Private/DecompressionTools.cpp b/Engine/Plugins/Experimental/UAF/UAF/Source/UAF/Private/DecompressionTools.cpp
--- a/Engine/Plugins/Experimental/UAF/UAF/Source/UAF/Private/DecompressionTools.cpp
+++ b/Engine/Plugins/Experimental/UAF/UAF/Source/UAF/Private/DecompressionTools.cpp
@@ -375,6 +375,8 @@ namespace Private
 		if (NumBonesToOrientAndScaleRetarget > 0)
 		{
 			const TArray<FTransform>& AuthoredOnRefSkeleton = RetargetTransforms;
+			
+			// TODO need to use Reference pose from SkeletalMesh its ref-skeleton (when available) to prevent incorrect proportion changes
 			const TArray<FTransform>& TargetRefSkeleton = TargetSkeleton->GetRefLocalPoses();
 			
 			for (const BoneTrackPair& BonePair : OrientAndScaleRetargetingPairs)
diff --git a/Engine/Plugins/Experimental/UAF/UAF/Source/UAF/Private/RetargetingTools.cpp b/Engine/Plugins/Experimental/UAF/UAF/Source/UAF/Private/RetargetingTools.cpp
--- a/Engine/Plugins/Experimental/UAF/UAF/Source/UAF/Private/RetargetingTools.cpp
+++ b/Engine/Plugins/Experimental/UAF/UAF/Source/UAF/Private/RetargetingTools.cpp
@@ -341,7 +341,8 @@ namespace UE::UAF
 		if (NumBonesToOrientAndScaleRetarget > 0)
 		{
 			const TArray<FTransform>& AuthoredOnRefSkeleton = SourceReferencePoseTransforms;
-			const TArray<FTransform>& TargetSkeletonRefPose = (SourceToTargetRemapping.IsValid() ? SourceToTargetRemapping.GetTargetSkeleton() : TargetReferencePose.GetSkeletonAsset())->GetReferenceSkeleton().GetRefBonePose();
+			const bool bValidSkeletonRemapping = SourceToTargetRemapping.IsValid();
+			const TArray<FTransform>* TargetSkeletonRefPose = bValidSkeletonRemapping ? &SourceToTargetRemapping.GetTargetSkeleton()->GetReferenceSkeleton().GetRefBonePose() : nullptr;
 			
 			for (const BoneTrackPair& BonePair : OrientAndScaleRetargetingPairs)
 			{
@@ -349,10 +350,13 @@ namespace UE::UAF
 				const int32 SourceSkeletonBoneIndex = BonePair.TrackIndex;
 				const int32 TargetSkeletonBoneIndex = SourceToTargetRemapping.IsValid() ? SourceToTargetRemapping.GetTargetSkeletonBoneIndex(SourceSkeletonBoneIndex) : SourceSkeletonBoneIndex;
 				
-				if(AuthoredOnRefSkeleton.IsValidIndex(SourceSkeletonBoneIndex) && TargetSkeletonRefPose.IsValidIndex(TargetSkeletonBoneIndex))
+				const bool bValidTargetSkeletonIndex = bValidSkeletonRemapping ? TargetSkeletonRefPose->IsValidIndex(TargetSkeletonBoneIndex) : (LODBoneIndex <TargetReferencePose.GetNumBonesForLOD(0) && LODBoneIndex >= 0);
+				if(AuthoredOnRefSkeleton.IsValidIndex(SourceSkeletonBoneIndex) && bValidTargetSkeletonIndex)
 				{
-					const FVector SourceSkelTrans = AuthoredOnRefSkeleton[SourceSkeletonBoneIndex].GetTranslation();
-					const FVector TargetSkelTrans = TargetSkeletonRefPose[TargetSkeletonBoneIndex].GetTranslation();
+					const FVector SourceSkelTrans = AuthoredOnRefSkeleton[SourceSkeletonBoneIndex].GetTranslation(); 
+
+					// Use Skeleton ref pose in case of valid remapping, otherwise use incoming reference pose (populated with either a SkeletalMesh or Skeleton ref-skeleton)
+					const FVector TargetSkelTrans = bValidSkeletonRemapping ? (*TargetSkeletonRefPose)[TargetSkeletonBoneIndex].GetTranslation() : TargetReferencePose.GetRefPoseTranslation(LODBoneIndex);
 			
 					// If translations are identical, we don't need to do any retargeting
 					if (!SourceSkelTrans.Equals(TargetSkelTrans, BONE_TRANS_RT_ORIENT_AND_SCALE_PRECISION))
