Fix FAttributeTypedSet::Init accessing uninitialized memory when searching for parent attribute names

#Synced-CL 48718022

diff --git a/Engine/Plugins/Experimental/UAF/UAF/Source/UAF/Private/UAF/Attributes/AttributeTypedSet.cpp b/Engine/Plugins/Experimental/UAF/UAF/Source/UAF/Private/UAF/Attributes/AttributeTypedSet.cpp
--- a/Engine/Plugins/Experimental/UAF/UAF/Source/UAF/Private/UAF/Attributes/AttributeTypedSet.cpp
+++ b/Engine/Plugins/Experimental/UAF/UAF/Source/UAF/Private/UAF/Attributes/AttributeTypedSet.cpp
@@ -296,7 +296,7 @@ namespace UE::UAF
 		AttributeKeys.Reserve(NumAttributes);
 		AttributeKeys.AddUninitialized(NumAttributes);
 		AttributeNames.Reserve(NumAttributes);
-		AttributeNames.AddUninitialized(NumAttributes);
+		AttributeNames.AddDefaulted(NumAttributes);
 		AttributeNameToIndexMap.Reserve(NumAttributes);
 		AttributeNameToIndexMap.AddUninitialized(NumAttributes);
 		ParentAttributeIndices.Reserve(NumAttributes);
