Add ConditionalPatchFunctionsOnLoad to RigVMAssetEditorData, its used to - conditionally - patch the function store when required. (issue at hand is a referencing asset trying to retrieve the correct function header). Used in ::HandlePackageDone() and ::GetRigVMGraphFunctionStore()
#misc strip export from header implemented virtual

#Synced-CL 48761512

diff --git a/Engine/Plugins/Experimental/UAF/UAF/Source/UAFUncookedOnly/Internal/AnimNextRigVMAssetEditorData.h b/Engine/Plugins/Experimental/UAF/UAF/Source/UAFUncookedOnly/Internal/AnimNextRigVMAssetEditorData.h
--- a/Engine/Plugins/Experimental/UAF/UAF/Source/UAFUncookedOnly/Internal/AnimNextRigVMAssetEditorData.h
+++ b/Engine/Plugins/Experimental/UAF/UAF/Source/UAFUncookedOnly/Internal/AnimNextRigVMAssetEditorData.h
@@ -653,7 +653,7 @@ protected:
 	UE_API virtual void BuildFunctionHeadersContext(const FRigVMCompileSettings& InSettings, FAnimNextGetFunctionHeaderCompileContext& OutCompileContext) const;
 
 	// Called pre programmatic variable population to allow population of compile context in a const manner. May be used by external assets querying our VM.
-	UE_API virtual void BuildProgrammaticVariablesContext(const FRigVMCompileSettings& InSettings, FAnimNextGetVariableCompileContext& OutCompileContext) const {}
+	virtual void BuildProgrammaticVariablesContext(const FRigVMCompileSettings& InSettings, FAnimNextGetVariableCompileContext& OutCompileContext) const {}
 
 	// Called during compilation to generate internal variables used for natively-callable RigVM functions
 	UE_API void BuildFunctionWrapperEventVariables(FAnimNextRigVMAssetCompileContext& InContext) const;
@@ -667,6 +667,9 @@ protected:
 	// Serialization/post-load fixup for old data
 	void UpgradeDataInterfaces();
 
+	// Conditionally calls PatchFunctionReferencesOnLoad and PatchFunctionsOnLoad on the RigVMClient	
+	void ConditionalPatchFunctionsOnLoad();
+
 	void SetVMRecompilationRequired(bool bInIsRequired);
 
 	/** All entries in this asset - not saved, discovered at load time and also contains InternalEntries */
@@ -781,6 +784,7 @@ protected:
 	bool bSuspendCompilationNotifications = false;
 	bool bSuspendCompilerReports = false;
 	bool bUpgradeDataInterfacesOnLoad = false;
+	bool bFunctionsPatched = false;
 };
 
 #undef UE_API
diff --git a/Engine/Plugins/Experimental/UAF/UAF/Source/UAFUncookedOnly/Private/AnimNextRigVMAssetEditorData.cpp b/Engine/Plugins/Experimental/UAF/UAF/Source/UAFUncookedOnly/Private/AnimNextRigVMAssetEditorData.cpp
--- a/Engine/Plugins/Experimental/UAF/UAF/Source/UAFUncookedOnly/Private/AnimNextRigVMAssetEditorData.cpp
+++ b/Engine/Plugins/Experimental/UAF/UAF/Source/UAFUncookedOnly/Private/AnimNextRigVMAssetEditorData.cpp
@@ -363,10 +363,7 @@ void UAnimNextRigVMAssetEditorData::HandlePackageDone()
 
 	GetRigVMClient()->RefreshAllModels(ERigVMLoadType::PostLoad, false, bIsCompiling);
 
-	GetRigVMClient()->PatchFunctionReferencesOnLoad();
-	TMap<URigVMLibraryNode*, FRigVMGraphFunctionHeader> OldHeaders;
-	TArray<FName> BackwardsCompatiblePublicFunctions;
-	GetRigVMClient()->PatchFunctionsOnLoad(this, BackwardsCompatiblePublicFunctions, OldHeaders);
+	ConditionalPatchFunctionsOnLoad();
 
 	// Register function references at RigVMBuildData
 	if (URigVMBuildData* BuildData = URigVMBuildData::Get())
@@ -794,6 +791,7 @@ URigVMGraph* UAnimNextRigVMAssetEditorData::GetRigVMGraphForEditorObject(UObject
 
 FRigVMGraphFunctionStore* UAnimNextRigVMAssetEditorData::GetRigVMGraphFunctionStore()
 {
+	ConditionalPatchFunctionsOnLoad();
 	return &GraphFunctionStore;
 }
 
@@ -3441,5 +3439,18 @@ void UAnimNextRigVMAssetEditorData::UpgradeDataInterfaces()
 	}
 }
 
+void UAnimNextRigVMAssetEditorData::ConditionalPatchFunctionsOnLoad()
+{
+	if (!bFunctionsPatched)
+	{
+		ensure(!HasAnyFlags(RF_NeedPostLoad));
+		bFunctionsPatched = true;
+		GetRigVMClient()->PatchFunctionReferencesOnLoad();
+		TMap<URigVMLibraryNode*, FRigVMGraphFunctionHeader> OldHeaders;
+		TArray<FName> BackwardsCompatiblePublicFunctions;
+		GetRigVMClient()->PatchFunctionsOnLoad(this, BackwardsCompatiblePublicFunctions, OldHeaders);
+	}
+}
+
 
 #undef LOCTEXT_NAMESPACE
