Fix refcount tracking issues in UAF value runtime
Add InitWithValueSpace on attribute collection to allow ref pose/identity initialization

#Synced-CL 48759454

diff --git a/Engine/Plugins/Experimental/MovieSceneAnimMixer/Source/MovieSceneAnimMixer/Private/Systems/MovieSceneMixedSkeletalAnimationSystem.cpp b/Engine/Plugins/Experimental/MovieSceneAnimMixer/Source/MovieSceneAnimMixer/Private/Systems/MovieSceneMixedSkeletalAnimationSystem.cpp
--- a/Engine/Plugins/Experimental/MovieSceneAnimMixer/Source/MovieSceneAnimMixer/Private/Systems/MovieSceneMixedSkeletalAnimationSystem.cpp
+++ b/Engine/Plugins/Experimental/MovieSceneAnimMixer/Source/MovieSceneAnimMixer/Private/Systems/MovieSceneMixedSkeletalAnimationSystem.cpp
@@ -408,9 +408,10 @@ void FMovieSceneSkeletalAnimationEvaluationTask::Execute(UE::UAF::FEvaluationVM&
 	if (VM.GetContext().IsValid())
 	{
 		// TODO: Implement with new attribute runtime
-		constexpr bool bIsAdditive = false;
-		constexpr bool bInitializeWithDefaultValues = true;
-		VM.PushValue(ATTRIBUTE_STACK_NAME, MakeUnique<FAttributeCollection>(VM.MakeAttributeCollection(bIsAdditive, bInitializeWithDefaultValues)));
+		FAttributeCollectionStack Collection(VM.GetContext().GetNamedSet());
+		Collection.InitWithValueSpace(FValueSpace(EValueSpaceType::Local));
+
+		VM.PushValue(ATTRIBUTE_STACK_NAME, MakeUnique<FAttributeCollection>(MoveTemp(Collection)));
 		return;
 	}
 
diff --git a/Engine/Plugins/Experimental/MovieScenePoseSearchTracks/Source/MovieScenePoseSearchTracks/Private/Sections/MovieSceneStitchAnimSection.cpp b/Engine/Plugins/Experimental/MovieScenePoseSearchTracks/Source/MovieScenePoseSearchTracks/Private/Sections/MovieSceneStitchAnimSection.cpp
--- a/Engine/Plugins/Experimental/MovieScenePoseSearchTracks/Source/MovieScenePoseSearchTracks/Private/Sections/MovieSceneStitchAnimSection.cpp
+++ b/Engine/Plugins/Experimental/MovieScenePoseSearchTracks/Source/MovieScenePoseSearchTracks/Private/Sections/MovieSceneStitchAnimSection.cpp
@@ -62,9 +62,10 @@ void FMovieSceneStitchAnimEvaluationTask::Execute(UE::UAF::FEvaluationVM& VM) co
 	if (VM.GetContext().IsValid())
 	{
 		// TODO: Implement with new attribute runtime
-		constexpr bool bIsAdditive = false;
-		constexpr bool bInitializeWithDefaultValues = true;
-		VM.PushValue(ATTRIBUTE_STACK_NAME, MakeUnique<FAttributeCollection>(VM.MakeAttributeCollection(bIsAdditive, bInitializeWithDefaultValues)));
+		FAttributeCollectionStack Collection(VM.GetContext().GetNamedSet());
+		Collection.InitWithValueSpace(FValueSpace(EValueSpaceType::Local));
+
+		VM.PushValue(ATTRIBUTE_STACK_NAME, MakeUnique<FAttributeCollection>(MoveTemp(Collection)));
 		return;
 	}
 
diff --git a/Engine/Plugins/Experimental/UAF/UAF/Source/UAF/Private/UAF/Attributes/AttributeBindingData.cpp b/Engine/Plugins/Experimental/UAF/UAF/Source/UAF/Private/UAF/Attributes/AttributeBindingData.cpp
--- a/Engine/Plugins/Experimental/UAF/UAF/Source/UAF/Private/UAF/Attributes/AttributeBindingData.cpp
+++ b/Engine/Plugins/Experimental/UAF/UAF/Source/UAF/Private/UAF/Attributes/AttributeBindingData.cpp
@@ -366,13 +366,19 @@ namespace UE::UAF
 
 						(*DefaultValues)[ValueIndex] = FBoneTransformAnimationAttribute{ RefPose[SkeletonBoneIndex.GetInt()] };
 					}
+
+					// Our set map has a reference to a typed set
+					NumInternalReferences++;
 				}
 
-				DefaultAttributeValuesByName.Add(SetName, DefaultAttributeValues);
-			}
+				DefaultAttributeValuesByName.Add(SetName, MoveTemp(DefaultAttributeValues));
 
-			NumInternalReferences += DefaultAttributeValuesByName.Num();
+				// Our collection has a reference to a named set
+				NumInternalReferences++;
+			}
 		}
+
+		BindingAsset = &Binding;
 	}
 
 	int32 FAttributeBindingData::NumNamedSets() const
@@ -436,14 +442,14 @@ namespace UE::UAF
 
 	uint32 FAttributeBindingData::Release() const
 	{
-		const uint32 RefsBeforeRelease = UE::Private::FQueryableRefCountedObject::Release();
-		if (RefsBeforeRelease == NumInternalReferences + 1)
+		const uint32 RefsAfterRelease = UE::Private::FQueryableRefCountedObject::Release();
+		if (RefsAfterRelease == NumInternalReferences)
 		{
 			// Our last external reference has been released, release our internal references to free the memory
 			const_cast<FAttributeBindingData*>(this)->ResetInternalReferences();
 		}
 
-		return RefsBeforeRelease;
+		return RefsAfterRelease;
 	}
 
 	void FAttributeBindingData::ResetInternalReferences()
diff --git a/Engine/Plugins/Experimental/UAF/UAF/Source/UAF/Private/UAF/Attributes/AttributeCollection.cpp b/Engine/Plugins/Experimental/UAF/UAF/Source/UAF/Private/UAF/Attributes/AttributeCollection.cpp
--- a/Engine/Plugins/Experimental/UAF/UAF/Source/UAF/Private/UAF/Attributes/AttributeCollection.cpp
+++ b/Engine/Plugins/Experimental/UAF/UAF/Source/UAF/Private/UAF/Attributes/AttributeCollection.cpp
@@ -77,29 +77,7 @@ namespace UE::UAF
 
 	FAttributeCollection::~FAttributeCollection()
 	{
-		if (SetMaps != nullptr)
-		{
-			for (int32 MapIndex = 0; MapIndex < SetMapCount; ++MapIndex)
-			{
-				ReleaseAttributeSetMap(SetMaps[MapIndex].Map);
-			}
-
-			(*ReallocFun)(reinterpret_cast<uint8*>(SetMaps), sizeof(FSetMapEntry) * SetMapCapacity, 0);
-			SetMaps = nullptr;
-			SetMapCount = 0;
-		}
-
-		if (Maps != nullptr)
-		{
-			for (int32 MapIndex = 0; MapIndex < MapCount; ++MapIndex)
-			{
-				ReleaseAttributeMap(Maps[MapIndex].Map);
-			}
-
-			(*ReallocFun)(reinterpret_cast<uint8*>(Maps), sizeof(FMapEntry) * MapCapacity, 0);
-			Maps = nullptr;
-			MapCount = 0;
-		}
+		Empty();
 	}
 
 	FAttributeCollection& FAttributeCollection::operator=(const FAttributeCollection& Other)
@@ -109,7 +87,7 @@ namespace UE::UAF
 			checkf(ReallocFun != nullptr, TEXT("When we copy from another collection, we retain our original allocator"));
 
 			// Free anything we might be holding
-			this->~FAttributeCollection();
+			Empty();
 
 			NamedSet = Other.NamedSet;
 			ValueSpace = Other.ValueSpace;
@@ -162,7 +140,7 @@ namespace UE::UAF
 			checkf(ReallocFun == Other.ReallocFun, TEXT("When we move from another collection, our allocators must match"));
 
 			// Free anything we might be holding
-			this->~FAttributeCollection();
+			Empty();
 
 			NamedSet = MoveTemp(Other.NamedSet);
 			ValueSpace = Other.ValueSpace;
@@ -204,6 +182,54 @@ namespace UE::UAF
 		ValueSpace = InValueSpace;
 	}
 
+	void FAttributeCollection::InitWithValueSpace(FValueSpace InValueSpace)
+	{
+		if (!NamedSet)
+		{
+			// Nothing to initialize if we have no named set
+			return;
+		}
+
+		// Empty the collection
+		Reset(NamedSet);
+
+		const bool bIsAdditive = InValueSpace.IsAdditive();
+
+		bool bIsInitialized = false;
+		if (!bIsAdditive)
+		{
+			const EValueSpaceType SpaceType = InValueSpace.GetType();
+			switch (SpaceType)
+			{
+			case EValueSpaceType::Local:
+				// From local space reference pose
+				*this = *NamedSet->GetDefaultAttributeValues();
+				bIsInitialized = true;
+				break;
+			case EValueSpaceType::Component:
+				// From component space reference pose
+				checkf(false, TEXT("Not yet supported"));
+				break;
+			}
+		}
+
+		if (!bIsInitialized)
+		{
+			for (auto It = NamedSet->CreateTypedSetIterator(); It; ++It)
+			{
+				const FAttributeTypedSetPtr TypedSet = *It;
+
+				// By default, attributes map to values of their type
+				const FAttributeMappingKey MappingKey = FAttributeMappingKey::MakeFromTo(TypedSet->GetType());
+
+				FAttributeSetMap* Copy = AddSetMap(MappingKey);
+				Copy->FillWithIdentity(bIsAdditive);
+			}
+		}
+
+		ValueSpace = InValueSpace;
+	}
+
 	bool FAttributeCollection::IsAdditive() const
 	{
 		return ValueSpace.IsAdditive();
@@ -228,6 +254,38 @@ namespace UE::UAF
 		ValueSpace = FValueSpace();
 	}
 
+	void FAttributeCollection::Empty()
+	{
+		if (SetMaps != nullptr)
+		{
+			for (int32 MapIndex = 0; MapIndex < SetMapCount; ++MapIndex)
+			{
+				ReleaseAttributeSetMap(SetMaps[MapIndex].Map);
+			}
+
+			(*ReallocFun)(reinterpret_cast<uint8*>(SetMaps), sizeof(FSetMapEntry) * SetMapCapacity, 0);
+			SetMaps = nullptr;
+			SetMapCount = 0;
+			SetMapCapacity = 0;
+		}
+
+		if (Maps != nullptr)
+		{
+			for (int32 MapIndex = 0; MapIndex < MapCount; ++MapIndex)
+			{
+				ReleaseAttributeMap(Maps[MapIndex].Map);
+			}
+
+			(*ReallocFun)(reinterpret_cast<uint8*>(Maps), sizeof(FMapEntry) * MapCapacity, 0);
+			Maps = nullptr;
+			MapCount = 0;
+			MapCapacity = 0;
+		}
+
+		NamedSet = nullptr;
+		ValueSpace = FValueSpace();
+	}
+
 	FReallocFun FAttributeCollection::GetAllocator() const
 	{
 		return ReallocFun;
@@ -258,7 +316,8 @@ namespace UE::UAF
 
 		// We don't know the concrete type of the set map, use indirection to construct it
 		constexpr bool bConstructSetMap = true;
-		return *AddSetMapImpl(Args, bConstructSetMap);
+		FAttributeSetMap** SetMap = AddSetMapImpl(Args, bConstructSetMap);
+		return SetMap != nullptr ? *SetMap : nullptr;
 	}
 
 	FAttributeSetMap** FAttributeCollection::AddSetMapImpl(const FAttributeSetMap::FConstructArgs& Args, bool bConstructSetMap)
@@ -480,7 +539,8 @@ namespace UE::UAF
 
 		// We don't know the concrete type of the attribute map, use indirection to construct it
 		constexpr bool bConstructMap = true;
-		return *AddMapImpl(ValueType, bConstructMap);
+		FAttributeMap** Map = AddMapImpl(ValueType, bConstructMap);
+		return Map != nullptr ? *Map : nullptr;
 	}
 
 	FAttributeMap** FAttributeCollection::AddMapImpl(UScriptStruct* ValueType, bool bConstructMap)
diff --git a/Engine/Plugins/Experimental/UAF/UAF/Source/UAF/Private/UAF/Attributes/AttributeEvaluationContext.cpp b/Engine/Plugins/Experimental/UAF/UAF/Source/UAF/Private/UAF/Attributes/AttributeEvaluationContext.cpp
--- a/Engine/Plugins/Experimental/UAF/UAF/Source/UAF/Private/UAF/Attributes/AttributeEvaluationContext.cpp
+++ b/Engine/Plugins/Experimental/UAF/UAF/Source/UAF/Private/UAF/Attributes/AttributeEvaluationContext.cpp
@@ -15,15 +15,15 @@ namespace UE::UAF
 			return;
 		}
 
-		if (FAttributeBindingDataPtr InBindingData = InSetBinding->GetAttributeBindingData())
+		if (const FAttributeBindingDataPtr& InBindingData = InSetBinding->GetAttributeBindingData())
 		{
 			if (FAttributeNamedSetPtr InNamedSet = InBindingData->FindNamedSet(InSetName))
 			{
 				SetBinding = InSetBinding;
 				Skeleton = InSetBinding->GetSkeleton();
 				BindingData = InBindingData;
-				NamedSet = InNamedSet;
-				DefaultAttributeValues = InNamedSet->GetDefaultAttributeValues();
+				NamedSet = MoveTemp(InNamedSet);
+				DefaultAttributeValues = NamedSet->GetDefaultAttributeValues();
 			}
 		}
 	}
diff --git a/Engine/Plugins/Experimental/UAF/UAF/Source/UAF/Private/UAF/UAFSkeletonUserData.cpp b/Engine/Plugins/Experimental/UAF/UAF/Source/UAF/Private/UAF/UAFSkeletonUserData.cpp
--- a/Engine/Plugins/Experimental/UAF/UAF/Source/UAF/Private/UAF/UAFSkeletonUserData.cpp
+++ b/Engine/Plugins/Experimental/UAF/UAF/Source/UAF/Private/UAF/UAFSkeletonUserData.cpp
@@ -9,7 +9,7 @@
 #include "UAF/AbstractSkeleton/AbstractSkeletonSetBinding.h"
 #include "UObject/UObjectGlobals.h"
 
-namespace
+namespace UE::UAF::Private
 {
 	struct FCachedUAFSkeletonUserData
 	{
@@ -160,6 +160,8 @@ namespace
 
 UUAFSkeletonUserData* UUAFSkeletonUserData::FromSkeleton(USkeleton& Skeleton)
 {
+	using namespace UE::UAF::Private;
+
 	UUAFSkeletonUserData* UserData = Cast<UUAFSkeletonUserData>(Skeleton.GetAssetUserDataOfClass(UUAFSkeletonUserData::StaticClass()));
 	if (UserData != nullptr && UserData->GetSetBinding() != nullptr)
 	{
diff --git a/Engine/Plugins/Experimental/UAF/UAF/Source/UAF/Public/UAF/Attributes/AttributeBindingData.h b/Engine/Plugins/Experimental/UAF/UAF/Source/UAF/Public/UAF/Attributes/AttributeBindingData.h
--- a/Engine/Plugins/Experimental/UAF/UAF/Source/UAF/Public/UAF/Attributes/AttributeBindingData.h
+++ b/Engine/Plugins/Experimental/UAF/UAF/Source/UAF/Public/UAF/Attributes/AttributeBindingData.h
@@ -6,6 +6,7 @@
 #include "Templates/RefCounting.h"
 #include "UAF/Attributes/AttributeNamedSet.h"
 #include "UAF/Attributes/AttributeTypedSet.h"
+#include "UObject/WeakObjectPtr.h"
 
 #define UE_API UAF_API
 
@@ -124,6 +125,9 @@ namespace UE::UAF
 		// Default attribute values for each named set
 		TMap<FName, FAttributeCollectionPtr> DefaultAttributeValuesByName;
 
+		// The source binding asset used to create this binding data
+		TWeakObjectPtr<const UAbstractSkeletonSetBinding> BindingAsset;
+
 		// The number of internal references for the purposes of reference counting
 		uint32 NumInternalReferences = 0;
 
diff --git a/Engine/Plugins/Experimental/UAF/UAF/Source/UAF/Public/UAF/Attributes/AttributeCollection.h b/Engine/Plugins/Experimental/UAF/UAF/Source/UAF/Public/UAF/Attributes/AttributeCollection.h
--- a/Engine/Plugins/Experimental/UAF/UAF/Source/UAF/Public/UAF/Attributes/AttributeCollection.h
+++ b/Engine/Plugins/Experimental/UAF/UAF/Source/UAF/Public/UAF/Attributes/AttributeCollection.h
@@ -145,12 +145,22 @@ namespace UE::UAF
 		// Note that care must be taken when changing the value space as all values within the collection must be migrated accordingly
 		UE_API void SetValueSpace(FValueSpace ValueSpace);
 
+		// Sets the space our values line in and default initializes this collection
+		// For local/component space, we initialize with the reference pose
+		// For additive space, we initialize with the additive identity
+		// Other spaces initialize to their identity
+		// Note that only bound values are default initialized
+		UE_API void InitWithValueSpace(FValueSpace ValueSpace);
+
 		// Returns whether or not the space our values live in is additive
 		[[nodiscard]] UE_API bool IsAdditive() const;
 
 		// Resets the collection with the specified named set
 		UE_API void Reset(const FAttributeNamedSetPtr& NamedSet);
 
+		// Clears the collection of any content
+		UE_API void Empty();
+
 		// Returns the allocator function this collection was initialized with
 		[[nodiscard]] UE_API FReallocFun GetAllocator() const;
 
diff --git a/Engine/Plugins/Experimental/UAF/UAF/Source/UAF/Public/UAF/Attributes/AttributeSetMap.h b/Engine/Plugins/Experimental/UAF/UAF/Source/UAF/Public/UAF/Attributes/AttributeSetMap.h
--- a/Engine/Plugins/Experimental/UAF/UAF/Source/UAF/Public/UAF/Attributes/AttributeSetMap.h
+++ b/Engine/Plugins/Experimental/UAF/UAF/Source/UAF/Public/UAF/Attributes/AttributeSetMap.h
@@ -6,6 +6,7 @@
 #include "Templates/MemoryOps.h"
 #include "UAF/Attributes/AttributeMappingKey.h"
 #include "UAF/Attributes/AttributeTypedSet.h"
+#include "UAF/Attributes/EngineAttributes.h"
 #include "UAF/Attributes/IndirectAllocator.h"
 
 #define UE_API UAF_API
@@ -61,6 +62,9 @@ namespace UE::UAF
 		// Returns a copy of this set map using the specified allocator
 		[[nodiscard]] virtual FAttributeSetMap* Duplicate(FReallocFun InReallocFun) const = 0;
 
+		// Initializes every value to its identity, optionally additive
+		virtual void FillWithIdentity(bool bIsAdditive) = 0;
+
 		// Copies the content of this set map into another
 		// Both must have matching typed sets and value types
 		virtual void CopyTo(FAttributeSetMap& Other) const = 0;
@@ -136,6 +140,9 @@ namespace UE::UAF
 		// Returns a copy of this set map using the specified allocator
 		[[nodiscard]] virtual FAttributeSetMap* Duplicate(FReallocFun InReallocFun) const override;
 
+		// Initializes every value to its identity, optionally additive
+		virtual void FillWithIdentity(bool bIsAdditive) override;
+
 		// Copies the content of this set map into another
 		// Both must have matching typed sets and value types
 		virtual void CopyTo(FAttributeSetMap& Other) const override;
@@ -274,6 +281,38 @@ namespace UE::UAF
 		return Copy;
 	}
 
+	template<class InValueType>
+	inline void TAttributeSetMap<InValueType>::FillWithIdentity(bool bIsAdditive)
+	{
+		InValueType Identity{};
+
+		if constexpr (std::is_base_of_v<FTransformAnimationAttribute, InValueType>)
+		{
+			// Transforms have a special additive identity
+			if (bIsAdditive)
+			{
+				Identity = FTransformAnimationAttribute{ FTransform(FQuat::Identity, FVector::ZeroVector, FVector::ZeroVector) };
+			}
+			else
+			{
+				Identity = FTransformAnimationAttribute{ FTransform(FQuat::Identity, FVector::ZeroVector, FVector::OneVector) };
+			}
+		}
+		else
+		{
+			// Everything else uses the default value
+		}
+
+		InValueType* ValuePtr = GetData();
+		InValueType* ValuePtrEnd = ValuePtr + NumValues;
+		while (ValuePtr < ValuePtrEnd)
+		{
+			*ValuePtr = Identity;
+
+			ValuePtr++;
+		}
+	}
+
 	template<class InValueType>
 	inline void TAttributeSetMap<InValueType>::CopyTo(FAttributeSetMap& Other) const
 	{
diff --git a/Engine/Plugins/Experimental/UAF/UAF/Source/UAF/Public/UAF/Attributes/EngineAttributes.h b/Engine/Plugins/Experimental/UAF/UAF/Source/UAF/Public/UAF/Attributes/EngineAttributes.h
--- a/Engine/Plugins/Experimental/UAF/UAF/Source/UAF/Public/UAF/Attributes/EngineAttributes.h
+++ b/Engine/Plugins/Experimental/UAF/UAF/Source/UAF/Public/UAF/Attributes/EngineAttributes.h
@@ -14,6 +14,9 @@ USTRUCT(BlueprintType)
 struct FBoneTransformAnimationAttribute : public FTransformAnimationAttribute
 {
 	GENERATED_BODY()
+
+	// Allow promotion from plain transform to bone transform
+	using FTransformAnimationAttribute::operator=;
 };
 
 namespace UE::Anim
diff --git a/Engine/Plugins/Experimental/UAF/UAFAnimGraph/Source/UAFAnimGraph/Private/EvaluationVM/EvaluationVM.cpp b/Engine/Plugins/Experimental/UAF/UAFAnimGraph/Source/UAFAnimGraph/Private/EvaluationVM/EvaluationVM.cpp
--- a/Engine/Plugins/Experimental/UAF/UAFAnimGraph/Source/UAFAnimGraph/Private/EvaluationVM/EvaluationVM.cpp
+++ b/Engine/Plugins/Experimental/UAF/UAFAnimGraph/Source/UAFAnimGraph/Private/EvaluationVM/EvaluationVM.cpp
@@ -130,42 +130,6 @@ namespace UE::UAF
 		return Keyframe;
 	}
 
-	FAttributeCollectionStack FEvaluationVM::MakeAttributeCollection(bool bAdditive, bool bInitialize) const
-	{
-		if (!InternalContext.IsValid())
-		{
-			return FAttributeCollectionStack();
-		}
-
-		FAttributeCollectionStack Collection(InternalContext.GetNamedSet());
-
-		// TODO: How to handle additive attributes? We need to store an alternate set of default values
-		// Perhaps that should be a parameter we pass to AddSetMap, an initialization space (or to the collection)
-		// and when we create the map, instead of copying defaults, we default construct instead using that space to decide on our value
-
-		// Iterate over our named set attribute types and initialize our collection with its default values
-		for (auto It = InternalContext.GetNamedSet()->CreateTypedSetIterator(); It; ++It)
-		{
-			FAttributeTypedSetPtr TypedSet = *It;
-			UScriptStruct* AttributeType = TypedSet->GetType();
-
-			// By default, attributes map to values of their type
-			const FAttributeMappingKey MappingKey = FAttributeMappingKey::MakeFromTo(AttributeType);
-
-			FAttributeSetMap* Copy = Collection.AddSetMap(MappingKey);
-
-			if (bInitialize)
-			{
-				if (const FAttributeSetMap* DefaultValuesMap = InternalContext.GetDefaultAttributeValues()->FindSetMap(MappingKey))
-				{
-					DefaultValuesMap->CopyTo(*Copy);
-				}
-			}
-		}
-
-		return Collection;
-	}
-
 	const FAttributeEvaluationContext& FEvaluationVM::GetContext() const
 	{
 		return InternalContext;
diff --git a/Engine/Plugins/Experimental/UAF/UAFAnimGraph/Source/UAFAnimGraph/Private/EvaluationVM/Tasks/PushAnimSequenceKeyframe.cpp b/Engine/Plugins/Experimental/UAF/UAFAnimGraph/Source/UAFAnimGraph/Private/EvaluationVM/Tasks/PushAnimSequenceKeyframe.cpp
--- a/Engine/Plugins/Experimental/UAF/UAFAnimGraph/Source/UAFAnimGraph/Private/EvaluationVM/Tasks/PushAnimSequenceKeyframe.cpp
+++ b/Engine/Plugins/Experimental/UAF/UAFAnimGraph/Source/UAFAnimGraph/Private/EvaluationVM/Tasks/PushAnimSequenceKeyframe.cpp
@@ -34,11 +34,12 @@ void FAnimNextAnimSequenceKeyframeTask::Execute(UE::UAF::FEvaluationVM& VM) cons
 {
 	using namespace UE::UAF;
 	
+	const FAttributeEvaluationContext& VMContext = VM.GetContext();
+
 	if(const UAnimSequence* AnimSequencePtr = AnimSequence.Get())
 	{
 		const bool bExtractRootMotion = bExtractTrajectory;
 		const FAnimExtractContext ExtractionContext(SampleTime, bExtractRootMotion, DeltaTimeRecord, bLooping);
-		const FAttributeEvaluationContext& VMContext = VM.GetContext();
 
 		const bool bIsAdditive = AnimSequencePtr->IsValidAdditive();
 		const EAdditiveAnimationType AdditiveType = AnimSequencePtr->GetAdditiveAnimType();
@@ -46,9 +47,6 @@ void FAnimNextAnimSequenceKeyframeTask::Execute(UE::UAF::FEvaluationVM& VM) cons
 
 		if (VMContext.IsValid())
 		{
-			constexpr bool bInitializeWithDefaultValues = true;
-			FAttributeCollectionStack Collection = VM.MakeAttributeCollection(bIsAdditive, bInitializeWithDefaultValues);
-
 			FValueSpace ValueSpace;
 			if (bIsAdditive && AdditiveType == AAT_RotationOffsetMeshSpace)
 			{
@@ -58,7 +56,9 @@ void FAnimNextAnimSequenceKeyframeTask::Execute(UE::UAF::FEvaluationVM& VM) cons
 			{
 				ValueSpace = FValueSpace(EValueSpaceType::Local, bIsAdditive);
 			}
-			Collection.SetValueSpace(ValueSpace);
+
+			FAttributeCollectionStack Collection(VM.GetContext().GetNamedSet());
+			Collection.InitWithValueSpace(ValueSpace);
 
 			FDecompressionTools::GetAnimationPose(AnimSequencePtr, ExtractionContext, VMContext, Collection, bUseRawData);
 			FDecompressionTools::GetAnimationCurves(AnimSequencePtr, ExtractionContext, VMContext, Collection, bUseRawData);
@@ -126,12 +126,18 @@ void FAnimNextAnimSequenceKeyframeTask::Execute(UE::UAF::FEvaluationVM& VM) cons
 	}
 	else
 	{
-		constexpr bool bIsAdditive = false;
-		FKeyframeState Keyframe = VM.MakeReferenceKeyframe(bIsAdditive);
-		VM.PushValue(KEYFRAME_STACK_NAME, MakeUnique<FKeyframeState>(MoveTemp(Keyframe)));
+		if (VMContext.IsValid())
+		{
+			FAttributeCollectionStack Collection(VM.GetContext().GetNamedSet());
+			Collection.InitWithValueSpace(FValueSpace(EValueSpaceType::Local));
 
-		constexpr bool bInitializeWithDefaultValues = true;
-		FAttributeCollectionStack Collection = VM.MakeAttributeCollection(bIsAdditive, bInitializeWithDefaultValues);
-		VM.PushValue(ATTRIBUTE_STACK_NAME, MakeUnique<FAttributeCollection>(MoveTemp(Collection)));
+			VM.PushValue(ATTRIBUTE_STACK_NAME, MakeUnique<FAttributeCollection>(MoveTemp(Collection)));
+		}
+		else
+		{
+			constexpr bool bIsAdditive = false;
+			FKeyframeState Keyframe = VM.MakeReferenceKeyframe(bIsAdditive);
+			VM.PushValue(KEYFRAME_STACK_NAME, MakeUnique<FKeyframeState>(MoveTemp(Keyframe)));
+		}
 	}
 }
diff --git a/Engine/Plugins/Experimental/UAF/UAFAnimGraph/Source/UAFAnimGraph/Private/EvaluationVM/Tasks/PushPose.cpp b/Engine/Plugins/Experimental/UAF/UAFAnimGraph/Source/UAFAnimGraph/Private/EvaluationVM/Tasks/PushPose.cpp
--- a/Engine/Plugins/Experimental/UAF/UAFAnimGraph/Source/UAFAnimGraph/Private/EvaluationVM/Tasks/PushPose.cpp
+++ b/Engine/Plugins/Experimental/UAF/UAFAnimGraph/Source/UAFAnimGraph/Private/EvaluationVM/Tasks/PushPose.cpp
@@ -32,9 +32,10 @@ void FAnimNextPushPoseTask::Execute(UE::UAF::FEvaluationVM& VM) const
 	if (VM.GetContext().IsValid())
 	{
 		// TODO: Implement with new attribute runtime
-		constexpr bool bIsAdditive = false;
-		constexpr bool bInitializeWithDefaultValues = true;
-		VM.PushValue(ATTRIBUTE_STACK_NAME, MakeUnique<FAttributeCollection>(VM.MakeAttributeCollection(bIsAdditive, bInitializeWithDefaultValues)));
+		FAttributeCollectionStack Collection(VM.GetContext().GetNamedSet());
+		Collection.InitWithValueSpace(FValueSpace(EValueSpaceType::Local));
+
+		VM.PushValue(ATTRIBUTE_STACK_NAME, MakeUnique<FAttributeCollection>(MoveTemp(Collection)));
 		return;
 	}
 
diff --git a/Engine/Plugins/Experimental/UAF/UAFAnimGraph/Source/UAFAnimGraph/Private/EvaluationVM/Tasks/PushReferenceKeyframe.cpp b/Engine/Plugins/Experimental/UAF/UAFAnimGraph/Source/UAFAnimGraph/Private/EvaluationVM/Tasks/PushReferenceKeyframe.cpp
--- a/Engine/Plugins/Experimental/UAF/UAFAnimGraph/Source/UAFAnimGraph/Private/EvaluationVM/Tasks/PushReferenceKeyframe.cpp
+++ b/Engine/Plugins/Experimental/UAF/UAFAnimGraph/Source/UAFAnimGraph/Private/EvaluationVM/Tasks/PushReferenceKeyframe.cpp
@@ -25,8 +25,10 @@ void FAnimNextPushReferenceKeyframeTask::Execute(UE::UAF::FEvaluationVM& VM) con
 
 	if (VM.GetContext().IsValid())
 	{
-		constexpr bool bInitializeWithDefaultValues = true;
-		VM.PushValue(ATTRIBUTE_STACK_NAME, MakeUnique<FAttributeCollection>(VM.MakeAttributeCollection(bIsAdditive, bInitializeWithDefaultValues)));
+		FAttributeCollectionStack Collection(VM.GetContext().GetNamedSet());
+		Collection.InitWithValueSpace(FValueSpace(EValueSpaceType::Local, bIsAdditive));
+
+		VM.PushValue(ATTRIBUTE_STACK_NAME, MakeUnique<FAttributeCollection>(MoveTemp(Collection)));
 	}
 	else
 	{
diff --git a/Engine/Plugins/Experimental/UAF/UAFAnimGraph/Source/UAFAnimGraph/Private/EvaluationVM/Tasks/StoreKeyframe.cpp b/Engine/Plugins/Experimental/UAF/UAFAnimGraph/Source/UAFAnimGraph/Private/EvaluationVM/Tasks/StoreKeyframe.cpp
--- a/Engine/Plugins/Experimental/UAF/UAFAnimGraph/Source/UAFAnimGraph/Private/EvaluationVM/Tasks/StoreKeyframe.cpp
+++ b/Engine/Plugins/Experimental/UAF/UAFAnimGraph/Source/UAFAnimGraph/Private/EvaluationVM/Tasks/StoreKeyframe.cpp
@@ -64,9 +64,10 @@ void FAnimNextDuplicateTopKeyframeTask::Execute(UE::UAF::FEvaluationVM& VM) cons
 	if (VM.GetContext().IsValid())
 	{
 		// TODO: Implement with new attribute runtime
-		constexpr bool bIsAdditive = false;
-		constexpr bool bInitializeWithDefaultValues = true;
-		VM.PushValue(ATTRIBUTE_STACK_NAME, MakeUnique<FAttributeCollection>(VM.MakeAttributeCollection(bIsAdditive, bInitializeWithDefaultValues)));
+		FAttributeCollectionStack Collection(VM.GetContext().GetNamedSet());
+		Collection.InitWithValueSpace(FValueSpace(EValueSpaceType::Local));
+
+		VM.PushValue(ATTRIBUTE_STACK_NAME, MakeUnique<FAttributeCollection>(MoveTemp(Collection)));
 		return;
 	}
 
diff --git a/Engine/Plugins/Experimental/UAF/UAFAnimGraph/Source/UAFAnimGraph/Public/EvaluationVM/EvaluationVM.h b/Engine/Plugins/Experimental/UAF/UAFAnimGraph/Source/UAFAnimGraph/Public/EvaluationVM/EvaluationVM.h
--- a/Engine/Plugins/Experimental/UAF/UAFAnimGraph/Source/UAFAnimGraph/Public/EvaluationVM/EvaluationVM.h
+++ b/Engine/Plugins/Experimental/UAF/UAFAnimGraph/Source/UAFAnimGraph/Public/EvaluationVM/EvaluationVM.h
@@ -198,9 +198,6 @@ namespace UE::UAF
 		// Returns an uninitialized keyframe with memory pre-allocated
 		[[nodiscard]] UE_API FKeyframeState MakeUninitializedKeyframe(bool bAdditiveKeyframe) const;
 
-		// Returns an attribute collection sized for the current evaluation mask and LOD
-		[[nodiscard]] UE_API FAttributeCollectionStack MakeAttributeCollection(bool bAdditive, bool bInitialize) const;
-
 		// Returns the VM context that controls what state we produce
 		[[nodiscard]] UE_API const FAttributeEvaluationContext& GetContext() const;
 
