UAF: Refresh bindings customization on construction

Fixes distracting 1-frame flicker as the variables tab is scrolled

#Synced-CL 48758736

diff --git a/Engine/Plugins/Experimental/UAF/UAF/Source/UAFEditor/Private/Variables/VariableBindingPropertyCustomization.cpp b/Engine/Plugins/Experimental/UAF/UAF/Source/UAFEditor/Private/Variables/VariableBindingPropertyCustomization.cpp
--- a/Engine/Plugins/Experimental/UAF/UAF/Source/UAFEditor/Private/Variables/VariableBindingPropertyCustomization.cpp
+++ b/Engine/Plugins/Experimental/UAF/UAF/Source/UAFEditor/Private/Variables/VariableBindingPropertyCustomization.cpp
@@ -134,7 +134,7 @@ void FVariableBindingPropertyCustomization::CustomizeHeader(TSharedRef<IProperty
 		]
 	];
 
-	RequestRefresh();
+	Refresh();
 }
 
 void FVariableBindingPropertyCustomization::CustomizeChildren(TSharedRef<IPropertyHandle> InPropertyHandle, IDetailChildrenBuilder& ChildBuilder, IPropertyTypeCustomizationUtils& CustomizationUtils)
