Add additive space support to UAF value runtime

#Synced-CL 48769542

diff --git a/Engine/Plugins/Experimental/UAF/UAF/Source/UAF/Private/DecompressionTools.cpp b/Engine/Plugins/Experimental/UAF/UAF/Source/UAF/Private/DecompressionTools.cpp
--- a/Engine/Plugins/Experimental/UAF/UAF/Source/UAF/Private/DecompressionTools.cpp
+++ b/Engine/Plugins/Experimental/UAF/UAF/Source/UAF/Private/DecompressionTools.cpp
@@ -106,9 +106,6 @@ namespace Private
 	{
 		const FAttributeTypedSetPtr& BoneSet = OutBoneTransforms.GetTypedSet();
 		const int32 NumBoneTransforms = OutBoneTransforms.Num();
-
-		const FTransform RootRefTransform = OutBoneTransforms.GetData()->Value;
-
 		const int32 NumTracks = CompressedData.CompressedTrackToSkeletonMapTable.Num();
 
 		const USkeleton* SourceSkeleton = DecompressionContext.GetSourceSkeleton();
@@ -139,11 +136,6 @@ namespace Private
 		if (bIsMeshSpaceAdditive)
 		{
 			AnimatedCompactRotations.Init(false, NumBoneTransforms);
-			//OutAnimationPoseData.Flags |= ELODPoseFlags::MeshSpaceAdditive;			// TODO: Should be set when we create the collection by the caller
-		}
-		else if (DecompressionContext.GetAdditiveType() == AAT_LocalSpaceBase)
-		{
-			//OutAnimationPoseData.Flags |= ELODPoseFlags::LocalSpaceAdditive;
 		}
 
 		// this is not guaranteed for AnimSequences though... If Root is not animated, Track will not exist.
@@ -224,15 +216,6 @@ namespace Private
 
 			if (!RotationScalePairs.IsEmpty())
 			{
-#if DEFAULT_SOA
-				// get the remaining bone atoms
-				//CompressedData.BoneCompressionCodec->DecompressPose(DecompressionContext, UE::Anim::FAnimPoseDecompressionData(RotationScalePairs, TranslationPairs, RotationScalePairs, OutAnimationPoseData.LocalTransformsView.Rotations, OutAnimationPoseData.LocalTransformsView.Translations, OutAnimationPoseData.LocalTransformsView.Scales3D));
-#else
-				// get the remaining bone atoms
-				//TArrayView<FTransform> OutPoseBones = OutAnimationPoseData.LocalTransforms.Transforms;
-				//CompressedData.BoneCompressionCodec->DecompressPose(DecompressionContext, RotationScalePairs, TranslationPairs, RotationScalePairs, OutPoseBones);
-#endif
-
 				static_assert(sizeof(FBoneTransformAnimationAttribute) == sizeof(FTransform), "Sizes should match since we coerce between the two");
 
 				TArrayView<FTransform> TransformsView(reinterpret_cast<FTransform*>(OutBoneTransforms.GetData()), NumBoneTransforms);
@@ -304,6 +287,16 @@ namespace Private
 			// Once pose has been extracted, snap root bone back to first frame if we are extracting root motion.
 			if ((ExtractionContext.bExtractRootMotion && RootMotionReset.bEnableRootMotion) || RootMotionReset.bForceRootLock)
 			{
+				// TODO: Speculatively fetching the reference pose isn't ideal, it might not even be needed
+				// It also doesn't make sense if the animation is additive, the reference root transform isn't in additive space
+				// and we have nothing to subtract it from, we should probably use the additive identity
+				const FAttributeNamedSetPtr& NamedSet = OutBoneTransforms.GetTypedSet()->GetNamedSet();
+				const FAttributeCollectionPtr ReferenceValues = NamedSet->GetDefaultAttributeValues();
+				const TAttributeSetMap<FBoneTransformAnimationAttribute>* ReferenceBoneTransforms = ReferenceValues->FindSetMap<FBoneTransformAnimationAttribute>(OutBoneTransforms.GetMappingKey());
+				check(ReferenceBoneTransforms != nullptr && !ReferenceBoneTransforms->IsEmpty());
+
+				const FTransform& RootRefTransform = ReferenceBoneTransforms->GetData()->Value;
+
 				FTransform RootTransform = OutBoneTransforms.GetData()->Value;
 				RootMotionReset.ResetRootBoneForRootMotion(RootTransform, RootRefTransform);
 				OutBoneTransforms.GetData()->Value = RootTransform;
diff --git a/Engine/Plugins/Experimental/UAF/UAF/Source/UAF/Private/UAF/Attributes/EngineAttributes.cpp b/Engine/Plugins/Experimental/UAF/UAF/Source/UAF/Private/UAF/Attributes/EngineAttributes.cpp
--- a/Engine/Plugins/Experimental/UAF/UAF/Source/UAF/Private/UAF/Attributes/EngineAttributes.cpp
+++ b/Engine/Plugins/Experimental/UAF/UAF/Source/UAF/Private/UAF/Attributes/EngineAttributes.cpp
@@ -4,6 +4,7 @@
 
 #include "UAF/Attributes/AttributeRuntimeRegistry.h"
 #include "UAF/Attributes/Transformers/Accumulate.h"
+#include "UAF/Attributes/Transformers/AdditiveSpace.h"
 #include "UAF/Attributes/Transformers/Interpolate.h"
 #include "UAF/Attributes/Transformers/Overwrite.h"
 #include "UAF/Attributes/Transformers/Sanitize.h"
@@ -26,8 +27,12 @@ namespace UE::UAF
 
 		RuntimeRegistry.RegisterAttributeValueTransformer<Transformers::Private::FAccumulate_BoneTransformAttribute>();
 		RuntimeRegistry.RegisterAttributeValueTransformer<Transformers::Private::FAccumulate_FloatAttribute>();
+		RuntimeRegistry.RegisterAttributeValueTransformer<Transformers::Private::FApplyAdditiveSpace_BoneTransformAttribute>();
+		RuntimeRegistry.RegisterAttributeValueTransformer<Transformers::Private::FApplyAdditiveSpace_FloatAttribute>();
 		RuntimeRegistry.RegisterAttributeValueTransformer<Transformers::Private::FInterpolate_BoneTransformAttribute>();
 		RuntimeRegistry.RegisterAttributeValueTransformer<Transformers::Private::FInterpolate_FloatAttribute>();
+		RuntimeRegistry.RegisterAttributeValueTransformer<Transformers::Private::FMakeAdditiveSpace_BoneTransformAttribute>();
+		RuntimeRegistry.RegisterAttributeValueTransformer<Transformers::Private::FMakeAdditiveSpace_FloatAttribute>();
 		RuntimeRegistry.RegisterAttributeValueTransformer<Transformers::Private::FOverwrite_BoneTransformAttribute>();
 		RuntimeRegistry.RegisterAttributeValueTransformer<Transformers::Private::FOverwrite_FloatAttribute>();
 		RuntimeRegistry.RegisterAttributeValueTransformer<Transformers::Private::FSanitize_BoneTransformAttribute>();
diff --git a/Engine/Plugins/Experimental/UAF/UAF/Source/UAF/Private/UAF/Attributes/Transformers/AdditiveSpace.cpp b/Engine/Plugins/Experimental/UAF/UAF/Source/UAF/Private/UAF/Attributes/Transformers/AdditiveSpace.cpp
new file mode 100644
--- /dev/null
+++ b/Engine/Plugins/Experimental/UAF/UAF/Source/UAF/Private/UAF/Attributes/Transformers/AdditiveSpace.cpp
@@ -0,0 +1,833 @@
+// Copyright Epic Games, Inc. All Rights Reserved.
+
+#include "UAF/Attributes/Transformers/AdditiveSpace.h"
+
+#include "UAF/Attributes/AttributeBindingData.h"
+#include "UAF/Attributes/IteratorUtils.h"
+
+namespace UE::UAF::Transformers
+{
+	FName FMakeAdditiveSpace::TransformerName = TEXT("UAF::Transformers::MakeAdditive");
+	FName FApplyAdditiveSpace::TransformerName = TEXT("UAF::Transformers::ApplyAdditive");
+
+	void FMakeAdditiveSpace::Apply(
+		const FAttributeValueTransformerMapPtr& TransformerMap,
+		const FAttributeCollection& Base,
+		const FAttributeCollection& Input,
+		FAttributeCollection& AdditiveOutput)
+	{
+		// Named sets must match as it ensures that our inputs have the same sizes/shapes
+		checkf(Base.GetNamedSet() == Input.GetNamedSet(), TEXT("Base and Input collections must use the same named set"));
+		checkf(&Base != &AdditiveOutput, TEXT("Cannot output to the Base collection when making an additive space collection"));
+		checkf(&Input != &AdditiveOutput, TEXT("Cannot output to the Input collection when making an additive space collection"));
+
+		const FAttributeValueTransformerList* TransformerList = TransformerMap->Find(TransformerName);
+		if (TransformerList == nullptr)
+		{
+			// No transformer implementations have been registered
+			return;
+		}
+
+		// Setup our output using the value space of our base and making it additive
+		FValueSpace OutputSpace = Base.GetValueSpace();
+		checkf(OutputSpace.GetType() == EValueSpaceType::Local
+			|| (OutputSpace.GetType() == EValueSpaceType::Mixed && OutputSpace.GetMixedSpaceFlags() == EMixedSpaceFlags::MeshRotation),
+			TEXT("We can only create local and mesh space additives"));
+		checkf(!OutputSpace.IsAdditive() && !Input.GetValueSpace().IsAdditive(), TEXT("Inputs cannot be additive"));
+		OutputSpace.SetIsAdditive(true);
+
+		AdditiveOutput.Reset(Input.GetNamedSet());
+		AdditiveOutput.SetValueSpace(OutputSpace);
+
+		// Set maps
+		OuterJoinBy(
+			FAttributeValueTransformerListWithAttributeCollectionSetMapJoinOp(),
+			[&AdditiveOutput](FRawTransformerFunc Transformer, const FAttributeSetMap* BaseMap, const FAttributeSetMap* InputMap)
+			{
+				if (InputMap == nullptr)
+				{
+					// The base doesn't have a matching pair in our input, nothing to convert
+					return;
+				}
+
+				if (Transformer == nullptr)
+				{
+					// No transformer was found for this value type, retain the input as-is
+					FAttributeSetMap* SetMapCopy = InputMap->Duplicate(AdditiveOutput.GetAllocator());
+					AdditiveOutput.AppendSetMap(SetMapCopy);
+					return;
+				}
+
+				// Transform our matching pair, note that the base map might be missing
+				// If it is missing, transformer specializations should use a default value
+				FAttributeSetMap* OutputSetMap = nullptr;
+
+				FTransformSetMapFunc TypedTransformer = reinterpret_cast<FTransformSetMapFunc>(Transformer);
+				(*TypedTransformer)(BaseMap, InputMap, OutputSetMap, AdditiveOutput.GetAllocator());
+
+				checkf(OutputSetMap != nullptr, TEXT("Transformer function is responsible for allocating an output buffer"));
+				checkf(OutputSetMap->GetAllocator() == AdditiveOutput.GetAllocator(), TEXT("Transformer function is responsible for allocating using the output allocator"));
+				AdditiveOutput.AppendSetMap(OutputSetMap);
+			},
+			TransformerList->CreateSetMapTransformerIterator(),
+			Base.CreateConstSetMapIterator(),
+			Input.CreateConstSetMapIterator());
+
+		// Name/value maps
+		OuterJoinBy(
+			FAttributeValueTransformerListWithAttributeCollectionMapJoinOp(),
+			[&AdditiveOutput](FRawTransformerFunc Transformer, const FAttributeMap* BaseMap, const FAttributeMap* InputMap)
+			{
+				if (InputMap == nullptr)
+				{
+					// The base doesn't have a matching pair in our input, nothing to convert
+					return;
+				}
+
+				if (Transformer == nullptr)
+				{
+					// No transformer was found for this value type, retain the input as-is
+					FAttributeMap* MapCopy = InputMap->Duplicate(AdditiveOutput.GetAllocator());
+					AdditiveOutput.AppendMap(MapCopy);
+					return;
+				}
+
+				// Transform our matching pair, note that the base map might be missing
+				// If it is missing, transformer specializations should use a default value
+				FAttributeMap* OutputMap = nullptr;
+
+				FTransformMapFunc TypedTransformer = reinterpret_cast<FTransformMapFunc>(Transformer);
+				(*TypedTransformer)(BaseMap, InputMap, OutputMap, AdditiveOutput.GetAllocator());
+
+				checkf(OutputMap != nullptr, TEXT("Transformer function is responsible for allocating an output buffer"));
+				checkf(OutputMap->GetAllocator() == AdditiveOutput.GetAllocator(), TEXT("Transformer function is responsible for allocating using the output allocator"));
+				AdditiveOutput.AppendMap(OutputMap);
+			},
+			TransformerList->CreateMapTransformerIterator(),
+			Base.CreateConstMapIterator(),
+			Input.CreateConstMapIterator());
+	}
+
+	void FApplyAdditiveSpace::Apply(
+		const FAttributeValueTransformerMapPtr& TransformerMap,
+		const FAttributeCollection& Base,
+		const FAttributeCollection& Additive,
+		float AdditiveWeight,
+		FAttributeCollection& Output)
+	{
+		// Named sets must match as it ensures that our inputs have the same sizes / shapes
+		checkf(Base.GetNamedSet() == Additive.GetNamedSet(), TEXT("Base and Additive collections must use the same named set"));
+
+		const FValueSpace BaseSpace = Base.GetValueSpace();
+		const FValueSpace AdditiveSpace = Additive.GetValueSpace();
+		checkf(BaseSpace.GetType() == EValueSpaceType::Local
+			|| (BaseSpace.GetType() == EValueSpaceType::Mixed && !!(BaseSpace.GetMixedSpaceFlags() & EMixedSpaceFlags::MeshRotation)), TEXT("Base must be in local or mesh space to apply an additive"));
+		checkf(!BaseSpace.IsAdditive(), TEXT("Base input cannot be in additive space"));
+		checkf(AdditiveSpace.IsAdditive(), TEXT("Additive input must be in additive space"));
+		checkf(BaseSpace.GetType() == AdditiveSpace.GetType(), TEXT("Additive value space must match the base (local or mesh space additive)"));
+
+		const bool bIsOutputBase = &Base == &Output;
+		if (bIsOutputBase && !FAnimWeight::IsRelevant(AdditiveWeight))
+		{
+			// Our input base is our output and the additive input is not relevant, nothing to do
+			return;
+		}
+
+		const FAttributeValueTransformerList* TransformerList = TransformerMap->Find(TransformerName);
+		if (TransformerList == nullptr)
+		{
+			// No transformer implementations have been registered
+			return;
+		}
+
+		const bool bIsOutputAdditive = &Additive == &Output;
+		if (bIsOutputBase || bIsOutputAdditive)
+		{
+			// Our output collection matches one of our inputs, we'll modify it in place
+			FReallocFun OutputAllocator = Output.GetAllocator();
+
+			// Set maps
+			{
+				// Keep a list of set maps that needs to be appended to our output
+				TArray<FAttributeSetMap*, TInlineAllocator<8>> PendingSetMaps;
+
+				OuterJoinBy(
+					FAttributeValueTransformerListWithAttributeCollectionSetMapJoinOp(),
+					[AdditiveWeight, bIsOutputBase, OutputAllocator, &PendingSetMaps](FRawTransformerFunc Transformer, FAttributeSetMap* BaseMap, FAttributeSetMap* AdditiveMap)
+					{
+						if (AdditiveMap == nullptr)
+						{
+							// The base doesn't have a matching pair in our input, nothing to apply
+							return;
+						}
+
+						if (Transformer == nullptr)
+						{
+							// No transformer was found for this value type, retain the one with the largest weight
+							if (BaseMap == nullptr)
+							{
+								// We have no base, duplicate the additive
+								FAttributeSetMap* SetMapCopy = AdditiveMap->Duplicate(OutputAllocator);
+								PendingSetMaps.Add(SetMapCopy);
+							}
+							else if (AdditiveWeight >= 0.5f)
+							{
+								// Additive has more weight
+								if (bIsOutputBase)
+								{
+									// Base is our output, copy it over
+									AdditiveMap->CopyTo(*BaseMap);
+								}
+								else
+								{
+									// Additive is our output, nothing to do
+								}
+							}
+							else
+							{
+								// Base has more weight
+								if (bIsOutputBase)
+								{
+									// Base is our output, nothing to do
+								}
+								else
+								{
+									// Additive is our output, copy it over
+									BaseMap->CopyTo(*AdditiveMap);
+								}
+							}
+							return;
+						}
+
+						// Transform our matching pair, note that the base map might be missing
+						// If it is missing, transformer specializations should use a default value
+						FAttributeSetMap* OutputSetMap = bIsOutputBase ? BaseMap : AdditiveMap;
+						const bool bIsOutputMissing = OutputSetMap == nullptr;
+
+						FTransformSetMapFunc TypedTransformer = reinterpret_cast<FTransformSetMapFunc>(Transformer);
+						(*TypedTransformer)(BaseMap, AdditiveMap, AdditiveWeight, OutputSetMap, OutputAllocator);
+
+						if (bIsOutputMissing)
+						{
+							// Our output was missing, we should have allocated a buffer for it
+							checkf(OutputSetMap != nullptr, TEXT("Transformer function is responsible for allocating an output buffer"));
+							checkf(OutputSetMap->GetAllocator() == OutputAllocator, TEXT("Transformer function is responsible for allocating using the output allocator"));
+							PendingSetMaps.Add(OutputSetMap);
+						}
+					},
+					TransformerList->CreateSetMapTransformerIterator(),
+					// Remove the const qualifier since we'll mutate one in place as our output
+					const_cast<FAttributeCollection&>(Base).CreateSetMapIterator(),
+					const_cast<FAttributeCollection&>(Additive).CreateSetMapIterator());
+
+				// Append our pending entries
+				for (FAttributeSetMap* SetMap : PendingSetMaps)
+				{
+					Output.AddSetMap(SetMap);
+				}
+			}
+
+			// Name/value maps
+			{
+				// Keep a list of maps that needs to be appended to our output
+				TArray<FAttributeMap*, TInlineAllocator<8>> PendingMaps;
+
+				OuterJoinBy(
+					FAttributeValueTransformerListWithAttributeCollectionMapJoinOp(),
+					[AdditiveWeight, bIsOutputBase, OutputAllocator, &PendingMaps](FRawTransformerFunc Transformer, FAttributeMap* BaseMap, FAttributeMap* AdditiveMap)
+					{
+						if (AdditiveMap == nullptr)
+						{
+							// The base doesn't have a matching pair in our input, nothing to apply
+							return;
+						}
+
+						if (Transformer == nullptr)
+						{
+							// No transformer was found for this value type, retain the one with the largest weight
+							if (BaseMap == nullptr)
+							{
+								// We have no base
+								if (bIsOutputBase)
+								{
+									// Duplicate the additive
+									FAttributeMap* MapCopy = AdditiveMap->Duplicate(OutputAllocator);
+									PendingMaps.Add(MapCopy);
+								}
+								else
+								{
+									// Our output is the additive, it already contains it
+								}
+							}
+							else if (AdditiveWeight >= 0.5f)
+							{
+								// Additive has more weight or we have no base
+								if (bIsOutputBase)
+								{
+									// Base is our output, copy it over
+									AdditiveMap->CopyTo(*BaseMap);
+								}
+								else
+								{
+									// Additive is our output, nothing to do
+								}
+							}
+							else
+							{
+								// Base has more weight
+								if (bIsOutputBase)
+								{
+									// Base is our output, nothing to do
+								}
+								else
+								{
+									// Additive is our output, copy it over
+									BaseMap->CopyTo(*AdditiveMap);
+								}
+							}
+							return;
+						}
+
+						// Transform our matching pair, note that the base map might be missing
+						// If it is missing, transformer specializations should use a default value
+						FAttributeMap* OutputMap = nullptr;
+
+						FTransformMapFunc TypedTransformer = reinterpret_cast<FTransformMapFunc>(Transformer);
+						(*TypedTransformer)(BaseMap, AdditiveMap, AdditiveWeight, OutputMap, OutputAllocator);
+
+						checkf(OutputMap != nullptr, TEXT("Transformer function is responsible for allocating an output buffer"));
+						checkf(OutputMap->GetAllocator() == OutputAllocator, TEXT("Transformer function is responsible for allocating using the output allocator"));
+
+						if (bIsOutputBase)
+						{
+							if (BaseMap != nullptr)
+							{
+								// Overwrite our old base
+								OutputMap->MoveTo(*BaseMap);
+
+								// Destroy our temporary output map
+								ReleaseAttributeMap(OutputMap);
+							}
+							else
+							{
+								// Keep our new output
+								PendingMaps.Add(OutputMap);
+							}
+						}
+						else
+						{
+							// Overwrite our old additive
+							OutputMap->MoveTo(*AdditiveMap);
+
+							// Destroy our temporary output map
+							ReleaseAttributeMap(OutputMap);
+						}
+					},
+					TransformerList->CreateMapTransformerIterator(),
+					// Remove the const qualifier since we'll mutate one in place as our output
+					const_cast<FAttributeCollection&>(Base).CreateMapIterator(),
+					const_cast<FAttributeCollection&>(Additive).CreateMapIterator());
+
+				// Append our pending entries
+				for (FAttributeMap* Map : PendingMaps)
+				{
+					Output.AddMap(Map);
+				}
+			}
+		}
+		else
+		{
+			// Our output collection isn't either of our inputs, reset it
+			Output.Reset(Base.GetNamedSet());
+			Output.SetValueSpace(BaseSpace);
+
+			// Set maps
+			OuterJoinBy(
+				FAttributeValueTransformerListWithAttributeCollectionSetMapJoinOp(),
+				[AdditiveWeight, &Output](FRawTransformerFunc Transformer, const FAttributeSetMap* BaseMap, const FAttributeSetMap* AdditiveMap)
+				{
+					if (AdditiveMap == nullptr)
+					{
+						// The base doesn't have a matching pair in our input, nothing to apply
+						return;
+					}
+
+					if (Transformer == nullptr)
+					{
+						// No transformer was found for this value type, retain the one with the largest weight
+						FAttributeSetMap* MapCopy;
+						if (BaseMap == nullptr || AdditiveWeight >= 0.5f)
+						{
+							// Additive map has more weight or we have no base
+							MapCopy = AdditiveMap->Duplicate(Output.GetAllocator());
+						}
+						else
+						{
+							// Base map has more weight
+							MapCopy = BaseMap->Duplicate(Output.GetAllocator());
+						}
+
+						Output.AppendSetMap(MapCopy);
+						return;
+					}
+
+					// Transform our matching pair, note that the base map might be missing
+					// If it is missing, transformer specializations should use a default value
+					FAttributeSetMap* OutputSetMap = nullptr;
+
+					FTransformSetMapFunc TypedTransformer = reinterpret_cast<FTransformSetMapFunc>(Transformer);
+					(*TypedTransformer)(BaseMap, AdditiveMap, AdditiveWeight, OutputSetMap, Output.GetAllocator());
+
+					checkf(OutputSetMap != nullptr, TEXT("Transformer function is responsible for allocating an output buffer"));
+					checkf(OutputSetMap->GetAllocator() == Output.GetAllocator(), TEXT("Transformer function is responsible for allocating using the output allocator"));
+					Output.AppendSetMap(OutputSetMap);
+				},
+				TransformerList->CreateSetMapTransformerIterator(),
+				Base.CreateConstSetMapIterator(),
+				Additive.CreateConstSetMapIterator());
+
+			// Name/value maps
+			OuterJoinBy(
+				FAttributeValueTransformerListWithAttributeCollectionMapJoinOp(),
+				[AdditiveWeight, &Output](FRawTransformerFunc Transformer, const FAttributeMap* BaseMap, const FAttributeMap* AdditiveMap)
+				{
+					if (AdditiveMap == nullptr)
+					{
+						// The base doesn't have a matching pair in our input, nothing to apply
+						return;
+					}
+
+					if (Transformer == nullptr)
+					{
+						// No transformer was found for this value type, retain the one with the largest weight
+						FAttributeMap* MapCopy;
+						if (BaseMap == nullptr || AdditiveWeight >= 0.5f)
+						{
+							// Additive map has more weight or we have no base
+							MapCopy = AdditiveMap->Duplicate(Output.GetAllocator());
+						}
+						else
+						{
+							// Base map has more weight
+							MapCopy = BaseMap->Duplicate(Output.GetAllocator());
+						}
+
+						Output.AppendMap(MapCopy);
+						return;
+					}
+
+					// Transform our matching pair, note that the base map might be missing
+					// If it is missing, transformer specializations should use a default value
+					FAttributeMap* OutputMap = nullptr;
+
+					FTransformMapFunc TypedTransformer = reinterpret_cast<FTransformMapFunc>(Transformer);
+					(*TypedTransformer)(BaseMap, AdditiveMap, AdditiveWeight, OutputMap, Output.GetAllocator());
+
+					checkf(OutputMap != nullptr, TEXT("Transformer function is responsible for allocating an output buffer"));
+					checkf(OutputMap->GetAllocator() == Output.GetAllocator(), TEXT("Transformer function is responsible for allocating using the output allocator"));
+					Output.AppendMap(OutputMap);
+				},
+				TransformerList->CreateMapTransformerIterator(),
+				Base.CreateConstMapIterator(),
+				Additive.CreateConstMapIterator());
+		}
+	}
+
+	namespace Private
+	{
+		//////////////////////////////////////////////////////////////////////////
+		// MakeAdditiveSpace transformers
+
+		void FMakeAdditiveSpace_BoneTransformAttribute::TransformSetMap(const FAttributeSetMap* Base, const FAttributeSetMap* Input, FAttributeSetMap*& Output, FReallocFun OutputAllocator)
+		{
+			if (Base == nullptr)
+			{
+				// Our base is missing, use the reference pose
+				const FAttributeTypedSetPtr& TypedSet = Input->GetTypedSet();
+				const FAttributeCollectionPtr ReferenceValues = TypedSet->GetBindingData()->GetDefaultAttributeValues(TypedSet->GetNamedSet()->GetName());
+				Base = ReferenceValues->FindSetMap(Input->GetMappingKey());
+			}
+
+			if (Output == nullptr)
+			{
+				const FAttributeSetMap::FConstructArgs Args(Input->GetTypedSet(), Input->GetValueType(), OutputAllocator);
+				Output = MakeAttributeSetMap<FBoneTransformAnimationAttribute>(Args);
+			}
+
+			if (Base == nullptr)
+			{
+				UE_LOG(LogAnimation, Warning, TEXT("FMakeAdditiveSpace_BoneTransformAttribute::TransformSetMap: Could not transform set map - no base provided and no reference found. Setting to identity instead."));
+				Output->FillWithIdentity(true);
+				return;
+			}
+
+			// Our output should have the same shape as our base
+			check(Base->GetTypedSet() == Output->GetTypedSet());
+			check(Base->Num() == Output->Num());
+
+			// Our base and input should have the same shape
+			check(Base->GetTypedSet() == Input->GetTypedSet());
+			check(Base->Num() == Input->Num());
+
+			const TAttributeSetMap<FBoneTransformAnimationAttribute>* BaseBoneTransforms = CastChecked<FBoneTransformAnimationAttribute>(Base);
+			const TAttributeSetMap<FBoneTransformAnimationAttribute>* InputBoneTransforms = CastChecked<FBoneTransformAnimationAttribute>(Input);
+			TAttributeSetMap<FBoneTransformAnimationAttribute>* OutputBoneTransforms = CastChecked<FBoneTransformAnimationAttribute>(Output);
+
+			const FBoneTransformAnimationAttribute* BaseTransformPtr = BaseBoneTransforms->GetData();
+			const FBoneTransformAnimationAttribute* InputTransformPtr = InputBoneTransforms->GetData();
+			FBoneTransformAnimationAttribute* OutputTransformPtr = OutputBoneTransforms->GetData();
+			const FBoneTransformAnimationAttribute* BaseTransformPtrEnd = BaseTransformPtr + BaseBoneTransforms->Num();
+
+			while (BaseTransformPtr < BaseTransformPtrEnd)
+			{
+				OutputTransformPtr->Value.SetRotation(InputTransformPtr->Value.GetRotation() * BaseTransformPtr->Value.GetRotation().Inverse());
+				OutputTransformPtr->Value.SetTranslation(InputTransformPtr->Value.GetTranslation() - BaseTransformPtr->Value.GetTranslation());
+				// additive scale considers how much it grow or lower
+				// in order to support blending between different additive scale, we save [(target scale)/(source scale) - 1.f], and this can blend with 
+				// other delta scale value
+				// when we apply to the another scale, we apply scale * (1 + [additive scale])
+				OutputTransformPtr->Value.SetScale3D(InputTransformPtr->Value.GetScale3D() * FTransform::GetSafeScaleReciprocal(BaseTransformPtr->Value.GetScale3D()) - 1.0f);
+				OutputTransformPtr->Value.NormalizeRotation();
+
+				BaseTransformPtr++;
+				InputTransformPtr++;
+				OutputTransformPtr++;
+			}
+		}
+
+		void FMakeAdditiveSpace_BoneTransformAttribute::TransformMap(const FAttributeMap* Base, const FAttributeMap* Input, FAttributeMap*& Output, FReallocFun OutputAllocator)
+		{
+			check(Output == nullptr);
+
+			const TAttributeMap<FBoneTransformAnimationAttribute>* BaseBoneTransforms = Cast<FBoneTransformAnimationAttribute>(Base);
+			const TAttributeMap<FBoneTransformAnimationAttribute>* InputBoneTransforms = CastChecked<FBoneTransformAnimationAttribute>(Input);
+
+			TAttributeMap<FBoneTransformAnimationAttribute>* OutputBoneTransforms = MakeAttributeMap<FBoneTransformAnimationAttribute>(OutputAllocator);
+
+			OuterJoinBy(
+				FAttributeMapJoinOp(),
+				[OutputBoneTransforms](FName Name, const FBoneTransformAnimationAttribute* BaseValue, const FBoneTransformAnimationAttribute* InputValue)
+				{
+					if (InputValue == nullptr)
+					{
+						// We have no input to subtract, skip it
+						return;
+					}
+
+					// Missing values are set to the identity
+					const FTransform& BaseTransform = BaseValue != nullptr ? BaseValue->Value : FTransform::Identity;
+
+					FBoneTransformAnimationAttribute Output;
+					Output.Value.SetRotation(InputValue->Value.GetRotation() * BaseTransform.GetRotation().Inverse());
+					Output.Value.SetTranslation(InputValue->Value.GetTranslation() - BaseTransform.GetTranslation());
+					// additive scale considers how much it grow or lower
+					// in order to support blending between different additive scale, we save [(target scale)/(source scale) - 1.f], and this can blend with 
+					// other delta scale value
+					// when we apply to the another scale, we apply scale * (1 + [additive scale])
+					Output.Value.SetScale3D(InputValue->Value.GetScale3D() * FTransform::GetSafeScaleReciprocal(BaseTransform.GetScale3D()) - 1.0f);
+					Output.Value.NormalizeRotation();
+
+					OutputBoneTransforms->Append(Name, Output);
+				},
+				BaseBoneTransforms != nullptr ? BaseBoneTransforms->CreateConstIterator() : TAttributeMap<FBoneTransformAnimationAttribute>::FConstIterator(),
+				InputBoneTransforms->CreateConstIterator());
+
+			Output = OutputBoneTransforms;
+		}
+
+		void FMakeAdditiveSpace_FloatAttribute::TransformSetMap(const FAttributeSetMap* Base, const FAttributeSetMap* Input, FAttributeSetMap*& Output, FReallocFun OutputAllocator)
+		{
+			if (Base == nullptr)
+			{
+				// Our base is missing, use the user authored default values
+				const FAttributeTypedSetPtr& TypedSet = Input->GetTypedSet();
+				const FAttributeCollectionPtr ReferenceValues = TypedSet->GetBindingData()->GetDefaultAttributeValues(TypedSet->GetNamedSet()->GetName());
+				Base = ReferenceValues->FindSetMap(Input->GetMappingKey());
+			}
+
+			if (Output == nullptr)
+			{
+				const FAttributeSetMap::FConstructArgs Args(Input->GetTypedSet(), Input->GetValueType(), OutputAllocator);
+				Output = MakeAttributeSetMap<FFloatAnimationAttribute>(Args);
+			}
+
+			if (Base == nullptr)
+			{
+				UE_LOG(LogAnimation, Warning, TEXT("FMakeAdditiveSpace_FloatAttribute::TransformSetMap: Could not transform set map - no base provided and no reference found. Setting to identity instead."));
+				Output->FillWithIdentity(true);
+				return;
+			}
+
+			// Our output should have the same shape as our base
+			check(Base->GetTypedSet() == Output->GetTypedSet());
+			check(Base->Num() == Output->Num());
+
+			// Our base and input should have the same shape
+			check(Base->GetTypedSet() == Input->GetTypedSet());
+			check(Base->Num() == Input->Num());
+
+			const TAttributeSetMap<FFloatAnimationAttribute>* BaseFloats = CastChecked<FFloatAnimationAttribute>(Base);
+			const TAttributeSetMap<FFloatAnimationAttribute>* InputFloats = CastChecked<FFloatAnimationAttribute>(Input);
+			TAttributeSetMap<FFloatAnimationAttribute>* OutputFloats = CastChecked<FFloatAnimationAttribute>(Output);
+
+			const FFloatAnimationAttribute* BaseFloatsPtr = BaseFloats->GetData();
+			const FFloatAnimationAttribute* InputFloatsPtr = InputFloats->GetData();
+			FFloatAnimationAttribute* OutputFloatsPtr = OutputFloats->GetData();
+			const FFloatAnimationAttribute* BaseFloatsPtrEnd = BaseFloatsPtr + BaseFloats->Num();
+
+			while (BaseFloatsPtr < BaseFloatsPtrEnd)
+			{
+				OutputFloatsPtr->Value = InputFloatsPtr->Value - BaseFloatsPtr->Value;
+
+				BaseFloatsPtr++;
+				InputFloatsPtr++;
+				OutputFloatsPtr++;
+			}
+		}
+
+		void FMakeAdditiveSpace_FloatAttribute::TransformMap(const FAttributeMap* Base, const FAttributeMap* Input, FAttributeMap*& Output, FReallocFun OutputAllocator)
+		{
+			check(Output == nullptr);
+
+			const TAttributeMap<FFloatAnimationAttribute>* BaseFloats = Cast<FFloatAnimationAttribute>(Base);
+			const TAttributeMap<FFloatAnimationAttribute>* InputFloats = CastChecked<FFloatAnimationAttribute>(Input);
+
+			TAttributeMap<FFloatAnimationAttribute>* OutputFloats = MakeAttributeMap<FFloatAnimationAttribute>(OutputAllocator);
+
+			OuterJoinBy(
+				FAttributeMapJoinOp(),
+				[OutputFloats](FName Name, const FFloatAnimationAttribute* BaseValue, const FFloatAnimationAttribute* InputValue)
+				{
+					if (InputValue == nullptr)
+					{
+						// We have no input to subtract, skip it
+						return;
+					}
+
+					// Missing values are set to 0.0
+					const float BaseFloat = BaseValue != nullptr ? BaseValue->Value : 0.0f;
+
+					const FFloatAnimationAttribute Output{ InputValue->Value - BaseFloat };
+					OutputFloats->Append(Name, Output);
+				},
+				BaseFloats != nullptr ? BaseFloats->CreateConstIterator() : TAttributeMap<FFloatAnimationAttribute>::FConstIterator(),
+				InputFloats->CreateConstIterator());
+
+			Output = OutputFloats;
+		}
+
+		//////////////////////////////////////////////////////////////////////////
+		// ApplyAdditiveSpace transformers
+
+		void FApplyAdditiveSpace_BoneTransformAttribute::TransformSetMap(const FAttributeSetMap* Base, const FAttributeSetMap* Additive, float AdditiveWeight, FAttributeSetMap*& Output, FReallocFun OutputAllocator)
+		{
+			if (Base == nullptr)
+			{
+				// Our base is missing, use the reference pose
+				const FAttributeTypedSetPtr& TypedSet = Additive->GetTypedSet();
+				const FAttributeCollectionPtr ReferenceValues = TypedSet->GetBindingData()->GetDefaultAttributeValues(TypedSet->GetNamedSet()->GetName());
+				Base = ReferenceValues->FindSetMap(Additive->GetMappingKey());
+			}
+
+			if (Output == nullptr)
+			{
+				const FAttributeSetMap::FConstructArgs Args(Additive->GetTypedSet(), Additive->GetValueType(), OutputAllocator);
+				Output = MakeAttributeSetMap<FBoneTransformAnimationAttribute>(Args);
+			}
+
+			if (Base == nullptr)
+			{
+				UE_LOG(LogAnimation, Warning, TEXT("FApplyAdditiveSpace_BoneTransformAttribute::TransformSetMap: Could not transform set map - no base provided and no reference found. Setting to identity instead."));
+				Output->FillWithIdentity(false);
+				return;
+			}
+
+			// Our output should have the same shape as our base
+			check(Base->GetTypedSet() == Output->GetTypedSet());
+			check(Base->Num() == Output->Num());
+
+			// Our base and additive should have the same shape
+			check(Base->GetTypedSet() == Additive->GetTypedSet());
+			check(Base->Num() == Additive->Num());
+
+			const TAttributeSetMap<FBoneTransformAnimationAttribute>* BaseBoneTransforms = CastChecked<FBoneTransformAnimationAttribute>(Base);
+			const TAttributeSetMap<FBoneTransformAnimationAttribute>* AdditiveBoneTransforms = CastChecked<FBoneTransformAnimationAttribute>(Additive);
+			TAttributeSetMap<FBoneTransformAnimationAttribute>* OutputBoneTransforms = CastChecked<FBoneTransformAnimationAttribute>(Output);
+
+			const FBoneTransformAnimationAttribute* BaseTransformPtr = BaseBoneTransforms->GetData();
+			const FBoneTransformAnimationAttribute* AdditiveTransformPtr = AdditiveBoneTransforms->GetData();
+			FBoneTransformAnimationAttribute* OutputTransformPtr = OutputBoneTransforms->GetData();
+			const FBoneTransformAnimationAttribute* BaseTransformPtrEnd = BaseTransformPtr + BaseBoneTransforms->Num();
+
+			const ScalarRegister VAdditiveWeight(AdditiveWeight);
+
+			if (FAnimWeight::IsFullWeight(AdditiveWeight))
+			{
+				// Our additive has full weight, no need to interpolate with identity
+				while (BaseTransformPtr < BaseTransformPtrEnd)
+				{
+					// TODO: The VAdditiveWeight is practically 1.0, can assume that it is?
+					FTransform Result = BaseTransformPtr->Value;
+					Result.AccumulateWithAdditiveScale(AdditiveTransformPtr->Value, VAdditiveWeight);
+
+					OutputTransformPtr->Value = Result;
+
+					BaseTransformPtr++;
+					AdditiveTransformPtr++;
+					OutputTransformPtr++;
+				}
+			}
+			else
+			{
+				// Interpolate with identity before adding
+				while (BaseTransformPtr < BaseTransformPtrEnd)
+				{
+					FTransform Result = BaseTransformPtr->Value;
+					FTransform::BlendFromIdentityAndAccumulate(Result, AdditiveTransformPtr->Value, VAdditiveWeight);
+
+					OutputTransformPtr->Value = Result;
+
+					BaseTransformPtr++;
+					AdditiveTransformPtr++;
+					OutputTransformPtr++;
+				}
+			}
+		}
+
+		void FApplyAdditiveSpace_BoneTransformAttribute::TransformMap(const FAttributeMap* Base, const FAttributeMap* Additive, float AdditiveWeight, FAttributeMap*& Output, FReallocFun OutputAllocator)
+		{
+			check(Output == nullptr);
+
+			const TAttributeMap<FBoneTransformAnimationAttribute>* BaseBoneTransforms = Cast<FBoneTransformAnimationAttribute>(Base);
+			const TAttributeMap<FBoneTransformAnimationAttribute>* AdditiveBoneTransforms = CastChecked<FBoneTransformAnimationAttribute>(Additive);
+
+			TAttributeMap<FBoneTransformAnimationAttribute>* OutputBoneTransforms = MakeAttributeMap<FBoneTransformAnimationAttribute>(OutputAllocator);
+
+			const ScalarRegister VAdditiveWeight(AdditiveWeight);
+
+			if (FAnimWeight::IsFullWeight(AdditiveWeight))
+			{
+				// Our additive has full weight, no need to interpolate with identity
+				OuterJoinBy(
+					FAttributeMapJoinOp(),
+					[VAdditiveWeight, OutputBoneTransforms](FName Name, const FBoneTransformAnimationAttribute* BaseValue, const FBoneTransformAnimationAttribute* AdditiveValue)
+					{
+						// Missing values are set to the identity
+						const FTransform& BaseTransform = BaseValue != nullptr ? BaseValue->Value : FTransform::Identity;
+
+						FBoneTransformAnimationAttribute Output{ BaseTransform };
+
+						if (AdditiveValue != nullptr)
+						{
+							// TODO: The VAdditiveWeight is practically 1.0, can assume that it is?
+							Output.Value.AccumulateWithAdditiveScale(AdditiveValue->Value, VAdditiveWeight);
+						}
+
+						OutputBoneTransforms->Append(Name, Output);
+					},
+					BaseBoneTransforms != nullptr ? BaseBoneTransforms->CreateConstIterator() : TAttributeMap<FBoneTransformAnimationAttribute>::FConstIterator(),
+					AdditiveBoneTransforms->CreateConstIterator());
+			}
+			else
+			{
+				// Interpolate with identity before adding
+				OuterJoinBy(
+					FAttributeMapJoinOp(),
+					[VAdditiveWeight, OutputBoneTransforms](FName Name, const FBoneTransformAnimationAttribute* BaseValue, const FBoneTransformAnimationAttribute* AdditiveValue)
+					{
+						// Missing values are set to the identity
+						const FTransform& BaseTransform = BaseValue != nullptr ? BaseValue->Value : FTransform::Identity;
+
+						FBoneTransformAnimationAttribute Output{ BaseTransform };
+
+						if (AdditiveValue != nullptr)
+						{
+							FTransform::BlendFromIdentityAndAccumulate(Output.Value, AdditiveValue->Value, VAdditiveWeight);
+						}
+
+						OutputBoneTransforms->Append(Name, Output);
+					},
+					BaseBoneTransforms != nullptr ? BaseBoneTransforms->CreateConstIterator() : TAttributeMap<FBoneTransformAnimationAttribute>::FConstIterator(),
+					AdditiveBoneTransforms->CreateConstIterator());
+			}
+
+			Output = OutputBoneTransforms;
+		}
+
+		void FApplyAdditiveSpace_FloatAttribute::TransformSetMap(const FAttributeSetMap* Base, const FAttributeSetMap* Additive, float AdditiveWeight, FAttributeSetMap*& Output, FReallocFun OutputAllocator)
+		{
+			if (Base == nullptr)
+			{
+				// Our base is missing, use the user authored default values
+				const FAttributeTypedSetPtr& TypedSet = Additive->GetTypedSet();
+				const FAttributeCollectionPtr ReferenceValues = TypedSet->GetBindingData()->GetDefaultAttributeValues(TypedSet->GetNamedSet()->GetName());
+				Base = ReferenceValues->FindSetMap(Additive->GetMappingKey());
+			}
+
+			if (Output == nullptr)
+			{
+				const FAttributeSetMap::FConstructArgs Args(Additive->GetTypedSet(), Additive->GetValueType(), OutputAllocator);
+				Output = MakeAttributeSetMap<FFloatAnimationAttribute>(Args);
+			}
+
+			if (Base == nullptr)
+			{
+				UE_LOG(LogAnimation, Warning, TEXT("FApplyAdditiveSpace_FloatAttribute::TransformSetMap: Could not transform set map - no base provided and no reference found. Setting to identity instead."));
+				Output->FillWithIdentity(false);
+				return;
+			}
+
+			// Our output should have the same shape as our base
+			check(Base->GetTypedSet() == Output->GetTypedSet());
+			check(Base->Num() == Output->Num());
+
+			// Our base and additive should have the same shape
+			check(Base->GetTypedSet() == Additive->GetTypedSet());
+			check(Base->Num() == Additive->Num());
+
+			const TAttributeSetMap<FFloatAnimationAttribute>* BaseFloats = CastChecked<FFloatAnimationAttribute>(Base);
+			const TAttributeSetMap<FFloatAnimationAttribute>* AdditiveFloats = CastChecked<FFloatAnimationAttribute>(Additive);
+			TAttributeSetMap<FFloatAnimationAttribute>* OutputFloats = CastChecked<FFloatAnimationAttribute>(Output);
+
+			const FFloatAnimationAttribute* BaseFloatsPtr = BaseFloats->GetData();
+			const FFloatAnimationAttribute* AdditiveFloatsPtr = AdditiveFloats->GetData();
+			FFloatAnimationAttribute* OutputFloatsPtr = OutputFloats->GetData();
+			const FFloatAnimationAttribute* BaseFloatsPtrEnd = BaseFloatsPtr + BaseFloats->Num();
+
+			while (BaseFloatsPtr < BaseFloatsPtrEnd)
+			{
+				OutputFloatsPtr->Value = BaseFloatsPtr->Value + (AdditiveWeight * AdditiveFloatsPtr->Value);
+
+				BaseFloatsPtr++;
+				AdditiveFloatsPtr++;
+				OutputFloatsPtr++;
+			}
+		}
+
+		void FApplyAdditiveSpace_FloatAttribute::TransformMap(const FAttributeMap* Base, const FAttributeMap* Additive, float AdditiveWeight, FAttributeMap*& Output, FReallocFun OutputAllocator)
+		{
+			check(Output == nullptr);
+
+			const TAttributeMap<FFloatAnimationAttribute>* BaseFloats = Cast<FFloatAnimationAttribute>(Base);
+			const TAttributeMap<FFloatAnimationAttribute>* AdditiveFloats = CastChecked<FFloatAnimationAttribute>(Additive);
+
+			TAttributeMap<FFloatAnimationAttribute>* OutputFloats = MakeAttributeMap<FFloatAnimationAttribute>(OutputAllocator);
+
+			OuterJoinBy(
+				FAttributeMapJoinOp(),
+				[AdditiveWeight, OutputFloats](FName Name, const FFloatAnimationAttribute* BaseValue, const FFloatAnimationAttribute* AdditiveValue)
+				{
+					// Missing values are set to 0.0
+					const float BaseFloat = BaseValue != nullptr ? BaseValue->Value : 0.0f;
+					const float AdditiveFloat = AdditiveValue != nullptr ? AdditiveValue->Value : 0.0f;
+
+					const FFloatAnimationAttribute Output{ BaseFloat + (AdditiveWeight * AdditiveFloat) };
+					OutputFloats->Append(Name, Output);
+				},
+				BaseFloats != nullptr ? BaseFloats->CreateConstIterator() : TAttributeMap<FFloatAnimationAttribute>::FConstIterator(),
+				AdditiveFloats->CreateConstIterator());
+
+			Output = OutputFloats;
+		}
+	}
+}
diff --git a/Engine/Plugins/Experimental/UAF/UAF/Source/UAF/Public/UAF/Attributes/Transformers/AdditiveSpace.h b/Engine/Plugins/Experimental/UAF/UAF/Source/UAF/Public/UAF/Attributes/Transformers/AdditiveSpace.h
new file mode 100644
--- /dev/null
+++ b/Engine/Plugins/Experimental/UAF/UAF/Source/UAF/Public/UAF/Attributes/Transformers/AdditiveSpace.h
@@ -0,0 +1,95 @@
+// Copyright Epic Games, Inc. All Rights Reserved.
+
+#pragma once
+
+#include "CoreMinimal.h"
+
+#include "UAF/Attributes/EngineAttributes.h"
+#include "UAF/Attributes/Transformers/Transformer.h"
+
+#define UE_API UAF_API
+
+namespace UE::UAF::Transformers
+{
+	// Attribute value transformer to create additive values
+	struct FMakeAdditiveSpace final : public FAttributeValueTransformer
+	{
+		static UE_API FName TransformerName;
+
+		// If the output buffer is nullptr, an allocator is provided and the implementation is responsible for allocating the output buffer and returning it
+		// If the base map is missing, transformers should use a default value
+		using FTransformSetMapFunc = void (*)(const FAttributeSetMap* Base, const FAttributeSetMap* Input, FAttributeSetMap*& Output, FReallocFun OutputAllocator);
+		using FTransformMapFunc = void (*)(const FAttributeMap* Base, const FAttributeMap* Input, FAttributeMap*& Output, FReallocFun OutputAllocator);
+
+		// Creates a new additive collection by subtracting the base from the specified Input
+		static UE_API void Apply(
+			const FAttributeValueTransformerMapPtr& TransformerMap,
+			const FAttributeCollection& Base,
+			const FAttributeCollection& Input,
+			FAttributeCollection& AdditiveOutput);
+	};
+
+	// Attribute value transformer to apply additive values onto a base
+	struct FApplyAdditiveSpace final : public FAttributeValueTransformer
+	{
+		static UE_API FName TransformerName;
+
+		// If the output buffer is nullptr, an allocator is provided and the implementation is responsible for allocating the output buffer and returning it
+		// If the base map is missing, transformers should use a default value
+		using FTransformSetMapFunc = void (*)(const FAttributeSetMap* Base, const FAttributeSetMap* Additive, float AdditiveWeight, FAttributeSetMap*& Output, FReallocFun OutputAllocator);
+		using FTransformMapFunc = void (*)(const FAttributeMap* Base, const FAttributeMap* Additive, float AdditiveWeight, FAttributeMap*& Output, FReallocFun OutputAllocator);
+
+		// Applies an additive collection onto a base by adding it
+		// This function supports outputting to either input
+		static UE_API void Apply(
+			const FAttributeValueTransformerMapPtr& TransformerMap,
+			const FAttributeCollection& Base,
+			const FAttributeCollection& Additive,
+			float AdditiveWeight,
+			FAttributeCollection& Output);
+	};
+
+	//////////////////////////////////////////////////////////////////////////
+	// Engine built-in value type specializations
+
+	namespace Private
+	{
+		struct FMakeAdditiveSpace_BoneTransformAttribute
+		{
+			using TransformerType = FMakeAdditiveSpace;
+			using ValueType = FBoneTransformAnimationAttribute;
+
+			static void TransformSetMap(const FAttributeSetMap* Base, const FAttributeSetMap* Input, FAttributeSetMap*& Output, FReallocFun OutputAllocator);
+			static void TransformMap(const FAttributeMap* Base, const FAttributeMap* Input, FAttributeMap*& Output, FReallocFun OutputAllocator);
+		};
+
+		struct FMakeAdditiveSpace_FloatAttribute
+		{
+			using TransformerType = FMakeAdditiveSpace;
+			using ValueType = FFloatAnimationAttribute;
+
+			static void TransformSetMap(const FAttributeSetMap* Base, const FAttributeSetMap* Input, FAttributeSetMap*& Output, FReallocFun OutputAllocator);
+			static void TransformMap(const FAttributeMap* Base, const FAttributeMap* Input, FAttributeMap*& Output, FReallocFun OutputAllocator);
+		};
+
+		struct FApplyAdditiveSpace_BoneTransformAttribute
+		{
+			using TransformerType = FApplyAdditiveSpace;
+			using ValueType = FBoneTransformAnimationAttribute;
+
+			static void TransformSetMap(const FAttributeSetMap* Base, const FAttributeSetMap* Additive, float AdditiveWeight, FAttributeSetMap*& Output, FReallocFun OutputAllocator);
+			static void TransformMap(const FAttributeMap* Base, const FAttributeMap* Additive, float AdditiveWeight, FAttributeMap*& Output, FReallocFun OutputAllocator);
+		};
+
+		struct FApplyAdditiveSpace_FloatAttribute
+		{
+			using TransformerType = FApplyAdditiveSpace;
+			using ValueType = FFloatAnimationAttribute;
+
+			static void TransformSetMap(const FAttributeSetMap* Base, const FAttributeSetMap* Additive, float AdditiveWeight, FAttributeSetMap*& Output, FReallocFun OutputAllocator);
+			static void TransformMap(const FAttributeMap* Base, const FAttributeMap* Additive, float AdditiveWeight, FAttributeMap*& Output, FReallocFun OutputAllocator);
+		};
+	}
+}
+
+#undef UE_API
diff --git a/Engine/Plugins/Experimental/UAF/UAFAnimGraph/Source/UAFAnimGraph/Private/EvaluationVM/Tasks/ApplyAdditiveKeyframe.cpp b/Engine/Plugins/Experimental/UAF/UAFAnimGraph/Source/UAFAnimGraph/Private/EvaluationVM/Tasks/ApplyAdditiveKeyframe.cpp
--- a/Engine/Plugins/Experimental/UAF/UAFAnimGraph/Source/UAFAnimGraph/Private/EvaluationVM/Tasks/ApplyAdditiveKeyframe.cpp
+++ b/Engine/Plugins/Experimental/UAF/UAFAnimGraph/Source/UAFAnimGraph/Private/EvaluationVM/Tasks/ApplyAdditiveKeyframe.cpp
@@ -5,6 +5,9 @@
 #include "EvaluationVM/EvaluationVM.h"
 #include "EvaluationVM/KeyframeState.h"
 #include "TransformArrayOperations.h"
+#include "UAF/Attributes/AttributeCollection.h"
+#include "UAF/Attributes/Transformers/AdditiveSpace.h"
+#include "UAF/Attributes/Transformers/BoneSpace.h"
 
 #include UE_INLINE_GENERATED_CPP_BY_NAME(ApplyAdditiveKeyframe)
 
@@ -31,72 +34,132 @@ void FAnimNextApplyAdditiveKeyframeTask::Execute(UE::UAF::FEvaluationVM& VM) con
 
 	if (VM.GetContext().IsValid())
 	{
-		// TODO: Implement with new attribute runtime
-		TUniquePtr<FAttributeCollection> DiscardedCollection;
-		(void)VM.PopValue(ATTRIBUTE_STACK_NAME, DiscardedCollection);
-		return;
-	}
+		TUniquePtr<FAttributeCollection> AdditiveCollection;
+		if (!VM.PopValue(ATTRIBUTE_STACK_NAME, AdditiveCollection))
+		{
+			// We have no inputs, nothing to do
+			FAttributeCollectionStack Collection(VM.GetContext().GetNamedSet());
+			Collection.InitWithValueSpace(FValueSpace(EValueSpaceType::Local));
 
-	// Pop our top two poses, we'll re-use the top keyframe for our result
+			VM.PushValue(ATTRIBUTE_STACK_NAME, MakeUnique<FAttributeCollection>(MoveTemp(Collection)));
+			UE_LOG(LogAnimation, Warning, TEXT("FAnimNextApplyAdditiveKeyframeTask::Execute: Could not apply additive keyframe - no inputs provided. Pushing a ref pose instead."));
+			return;
+		}
 
-	TUniquePtr<FKeyframeState> AdditiveKeyframe;
-	if (!VM.PopValue(KEYFRAME_STACK_NAME, AdditiveKeyframe))
-	{
-		// We have no inputs, nothing to do
-		return;
-	}
+		TUniquePtr<FAttributeCollection> BaseCollection;
+		if (!VM.PopValue(ATTRIBUTE_STACK_NAME, BaseCollection))
+		{
+			// We have a single input, discard it since it must be the additive pose, either way something went wrong
+			// Push the reference pose since we'll expect a non-additive pose
+			FAttributeCollectionStack Collection(VM.GetContext().GetNamedSet());
+			Collection.InitWithValueSpace(FValueSpace(EValueSpaceType::Local));
+
+			VM.PushValue(ATTRIBUTE_STACK_NAME, MakeUnique<FAttributeCollection>(MoveTemp(Collection)));
+			UE_LOG(LogAnimation, Warning, TEXT("FAnimNextApplyAdditiveKeyframeTask::Execute: Could not apply additive keyframe - only a single input was provided. Pushing a ref pose instead."));
+			return;
+		}
 
-	TUniquePtr<FKeyframeState> BaseKeyframe;
-	if (!VM.PopValue(KEYFRAME_STACK_NAME, BaseKeyframe))
-	{
-		// We have a single input, discard it since it must be the additive pose, either way something went wrong
-		// Push the reference pose since we'll expect a non-additive pose
-		VM.PushValue(KEYFRAME_STACK_NAME, MakeUnique<FKeyframeState>(VM.MakeReferenceKeyframe(false)));
-		UE_LOG(LogAnimation, Warning, TEXT("FAnimNextApplyAdditiveKeyframeTask::Execute: Could not apply additive keyframe - only a single input was provided. Pushing a ref pose instead."));
-		return;
-	}
+		const FValueSpace AdditiveSpace = AdditiveCollection->GetValueSpace();
+		if (!AdditiveSpace.IsAdditive())
+		{
+			// Additive input must be additive, push reference pose if not the case
+			FAttributeCollectionStack Collection(VM.GetContext().GetNamedSet());
+			Collection.InitWithValueSpace(FValueSpace(EValueSpaceType::Local));
 
-	if (!AdditiveKeyframe->Pose.IsAdditive())
-	{
-		// Additive must be addtive type, push reference pose if not the case
-		VM.PushValue(KEYFRAME_STACK_NAME, MakeUnique<FKeyframeState>(VM.MakeReferenceKeyframe(false)));
-		UE_LOG(LogAnimation, Warning, TEXT("FAnimNextApplyAdditiveKeyframeTask::Execute: Could not apply additive keyframe - The expected additive pose was not of an additive type. Pushing a ref pose instead."));
-		return;
-	}
+			VM.PushValue(ATTRIBUTE_STACK_NAME, MakeUnique<FAttributeCollection>(MoveTemp(Collection)));
+			UE_LOG(LogAnimation, Warning, TEXT("FAnimNextApplyAdditiveKeyframeTask::Execute: Could not apply additive keyframe - The expected additive pose was not of an additive type. Pushing a ref pose instead."));
+			return;
+		}
+
+		const float AdditiveWeight = GetInterpolationAlpha(*BaseCollection, *AdditiveCollection);
+
+		bool bConvertBaseToMeshSpace = false;
+		if (AdditiveSpace.GetMixedSpaceFlags() == EMixedSpaceFlags::MeshRotation)
+		{
+			// If our additive input is in mesh rotation space, we ensure our base is in that space as well
+			const FValueSpace BaseSpace = BaseCollection->GetValueSpace();
+			bConvertBaseToMeshSpace = BaseSpace.GetMixedSpaceFlags() != EMixedSpaceFlags::MeshRotation;
+		}
+
+		if (bConvertBaseToMeshSpace)
+		{
+			Transformers::FBoneSpace::LocalToMeshRotation(*BaseCollection, *BaseCollection);
+		}
+
+		Transformers::FApplyAdditiveSpace::Apply(VM.GetTransformerMap(), *BaseCollection, *AdditiveCollection, AdditiveWeight, *BaseCollection);
 
-	const float AdditiveWeight = GetInterpolationAlpha(BaseKeyframe.Get(), AdditiveKeyframe.Get());
+		if (bConvertBaseToMeshSpace)
+		{
+			Transformers::FBoneSpace::MeshRotationToLocal(*BaseCollection, *BaseCollection);
+		}
 
-	if (EnumHasAnyFlags(VM.GetFlags(), EEvaluationFlags::Bones))
+		VM.PushValue(ATTRIBUTE_STACK_NAME, MoveTemp(BaseCollection));
+	}
+	else
 	{
-		check(BaseKeyframe->Pose.GetNumBones() == AdditiveKeyframe->Pose.GetNumBones());
+		// Pop our top two poses, we'll re-use the top keyframe for our result
 
-		const FTransformArrayView BaseTransformsView = BaseKeyframe->Pose.LocalTransforms.GetView();
+		TUniquePtr<FKeyframeState> AdditiveKeyframe;
+		if (!VM.PopValue(KEYFRAME_STACK_NAME, AdditiveKeyframe))
+		{
+			// We have no inputs, nothing to do
+			VM.PushValue(KEYFRAME_STACK_NAME, MakeUnique<FKeyframeState>(VM.MakeReferenceKeyframe(false)));
+			UE_LOG(LogAnimation, Warning, TEXT("FAnimNextApplyAdditiveKeyframeTask::Execute: Could not apply additive keyframe - no inputs provided. Pushing a ref pose instead."));
+			return;
+		}
 
-		if (AdditiveKeyframe->Pose.IsMeshSpaceAdditive())
+		TUniquePtr<FKeyframeState> BaseKeyframe;
+		if (!VM.PopValue(KEYFRAME_STACK_NAME, BaseKeyframe))
 		{
-			BlendWithIdentityAndAccumulateMesh(
-				BaseTransformsView, AdditiveKeyframe->Pose.LocalTransforms.GetConstView(), 
-				AdditiveKeyframe->Pose.GetLODBoneIndexToParentLODBoneIndexMap(), AdditiveWeight);
+			// We have a single input, discard it since it must be the additive pose, either way something went wrong
+			// Push the reference pose since we'll expect a non-additive pose
+			VM.PushValue(KEYFRAME_STACK_NAME, MakeUnique<FKeyframeState>(VM.MakeReferenceKeyframe(false)));
+			UE_LOG(LogAnimation, Warning, TEXT("FAnimNextApplyAdditiveKeyframeTask::Execute: Could not apply additive keyframe - only a single input was provided. Pushing a ref pose instead."));
+			return;
 		}
-		else
+
+		if (!AdditiveKeyframe->Pose.IsAdditive())
 		{
-			BlendWithIdentityAndAccumulate(BaseTransformsView, AdditiveKeyframe->Pose.LocalTransforms.GetConstView(), AdditiveWeight);
+			// Additive must be addtive type, push reference pose if not the case
+			VM.PushValue(KEYFRAME_STACK_NAME, MakeUnique<FKeyframeState>(VM.MakeReferenceKeyframe(false)));
+			UE_LOG(LogAnimation, Warning, TEXT("FAnimNextApplyAdditiveKeyframeTask::Execute: Could not apply additive keyframe - The expected additive pose was not of an additive type. Pushing a ref pose instead."));
+			return;
 		}
 
-		NormalizeRotations(BaseTransformsView);
-	}
+		const float AdditiveWeight = GetInterpolationAlpha(BaseKeyframe.Get(), AdditiveKeyframe.Get());
 
-	if (EnumHasAnyFlags(VM.GetFlags(), EEvaluationFlags::Curves))
-	{
-		BaseKeyframe->Curves.Accumulate(AdditiveKeyframe->Curves, AdditiveWeight);
-	}
+		if (EnumHasAnyFlags(VM.GetFlags(), EEvaluationFlags::Bones))
+		{
+			check(BaseKeyframe->Pose.GetNumBones() == AdditiveKeyframe->Pose.GetNumBones());
 
-	if (EnumHasAnyFlags(VM.GetFlags(), EEvaluationFlags::Attributes))
-	{
-		UE::Anim::Attributes::AccumulateAttributes(AdditiveKeyframe->Attributes, BaseKeyframe->Attributes, AdditiveWeight, AAT_None);
-	}
+			const FTransformArrayView BaseTransformsView = BaseKeyframe->Pose.LocalTransforms.GetView();
+
+			if (AdditiveKeyframe->Pose.IsMeshSpaceAdditive())
+			{
+				BlendWithIdentityAndAccumulateMesh(
+					BaseTransformsView, AdditiveKeyframe->Pose.LocalTransforms.GetConstView(), 
+					AdditiveKeyframe->Pose.GetLODBoneIndexToParentLODBoneIndexMap(), AdditiveWeight);
+			}
+			else
+			{
+				BlendWithIdentityAndAccumulate(BaseTransformsView, AdditiveKeyframe->Pose.LocalTransforms.GetConstView(), AdditiveWeight);
+			}
+
+			NormalizeRotations(BaseTransformsView);
+		}
+
+		if (EnumHasAnyFlags(VM.GetFlags(), EEvaluationFlags::Curves))
+		{
+			BaseKeyframe->Curves.Accumulate(AdditiveKeyframe->Curves, AdditiveWeight);
+		}
+
+		if (EnumHasAnyFlags(VM.GetFlags(), EEvaluationFlags::Attributes))
+		{
+			UE::Anim::Attributes::AccumulateAttributes(AdditiveKeyframe->Attributes, BaseKeyframe->Attributes, AdditiveWeight, AAT_None);
+		}
 
-	VM.PushValue(KEYFRAME_STACK_NAME, MoveTemp(BaseKeyframe));
+		VM.PushValue(KEYFRAME_STACK_NAME, MoveTemp(BaseKeyframe));
+	}
 }
 
 float FAnimNextApplyAdditiveKeyframeTask::GetInterpolationAlpha(const UE::UAF::FKeyframeState* KeyframeA, const UE::UAF::FKeyframeState* KeyframeB) const
@@ -119,3 +182,41 @@ float FAnimNextApplyAdditiveKeyframeTask::GetInterpolationAlpha(const UE::UAF::F
 
 	return FMath::Clamp(Alpha, 0.0f, 1.0f);
 }
+
+float FAnimNextApplyAdditiveKeyframeTask::GetInterpolationAlpha(const UE::UAF::FAttributeCollection& CollectionA, const UE::UAF::FAttributeCollection& CollectionB) const
+{
+	using namespace UE::UAF;
+
+	float Alpha = BlendWeight;
+
+	if (AlphaSourceCurveName != NAME_None && AlphaCurveInputIndex != INDEX_NONE)
+	{
+		const FAttributeMappingKey MappingKey = FAttributeMappingKey::MakeFromTo<FFloatAnimationAttribute>();
+		const FAttributeCollection& AlphaCurveCollection = AlphaCurveInputIndex == 0 ? CollectionA : CollectionB;
+		if (const TAttributeSetMap<FFloatAnimationAttribute>* FloatSetMap = AlphaCurveCollection.FindSetMap<FFloatAnimationAttribute>(MappingKey))
+		{
+			const FAttributeTypedSetPtr& TypedSet = FloatSetMap->GetTypedSet();
+			if (FAttributeSetIndex CurveIndex = TypedSet->FindIndex(AlphaSourceCurveName))
+			{
+				Alpha = (*FloatSetMap)[CurveIndex].Value;
+			}
+			else
+			{
+				// Curve name wasn't found
+				Alpha = 0.0f;
+			}
+		}
+		else
+		{
+			// No float curves found
+			Alpha = 0.0f;
+		}
+
+		if (InputScaleBiasClampFn.IsSet())
+		{
+			Alpha = InputScaleBiasClampFn(Alpha);
+		}
+	}
+
+	return FMath::Clamp(Alpha, 0.0f, 1.0f);
+}
diff --git a/Engine/Plugins/Experimental/UAF/UAFAnimGraph/Source/UAFAnimGraph/Public/EvaluationVM/Tasks/ApplyAdditiveKeyframe.h b/Engine/Plugins/Experimental/UAF/UAFAnimGraph/Source/UAFAnimGraph/Public/EvaluationVM/Tasks/ApplyAdditiveKeyframe.h
--- a/Engine/Plugins/Experimental/UAF/UAFAnimGraph/Source/UAFAnimGraph/Public/EvaluationVM/Tasks/ApplyAdditiveKeyframe.h
+++ b/Engine/Plugins/Experimental/UAF/UAFAnimGraph/Source/UAFAnimGraph/Public/EvaluationVM/Tasks/ApplyAdditiveKeyframe.h
@@ -13,6 +13,7 @@ struct FInputScaleBiasClamp;
 
 namespace UE::UAF
 {
+	class FAttributeCollection;
 	struct FKeyframeState;
 }
 
@@ -39,6 +40,7 @@ struct FAnimNextApplyAdditiveKeyframeTask : public FAnimNextEvaluationTask
 	// @TODO: Move this into shared method
 	// Compute the current alpha value for curves or return the interpolation alpha otherwise
 	UE_API float GetInterpolationAlpha(const UE::UAF::FKeyframeState* KeyframeA, const UE::UAF::FKeyframeState* KeyframeB) const;
+	UE_API float GetInterpolationAlpha(const UE::UAF::FAttributeCollection& CollectionA, const UE::UAF::FAttributeCollection& CollectionB) const;
 
 	// How much weight between the additive identity and the additive pose to apply: lerp(identity, additive, weight)
 	UPROPERTY()
