Animation.AnimNext.Variables test needs renaming and fails after rename run.
#fix ensure we copy default values when combining property bags
#misc change how we generate the combined property-bag (now done in one pass rather than incrementally)
#Synced-CL 48786226

diff --git a/Engine/Plugins/Experimental/UAF/UAF/Source/UAFUncookedOnly/Private/AnimNextRigVMAssetEditorData.cpp b/Engine/Plugins/Experimental/UAF/UAF/Source/UAFUncookedOnly/Private/AnimNextRigVMAssetEditorData.cpp
--- a/Engine/Plugins/Experimental/UAF/UAF/Source/UAFUncookedOnly/Private/AnimNextRigVMAssetEditorData.cpp
+++ b/Engine/Plugins/Experimental/UAF/UAF/Source/UAFUncookedOnly/Private/AnimNextRigVMAssetEditorData.cpp
@@ -3211,35 +3211,48 @@ FInstancedPropertyBag UAnimNextRigVMAssetEditorData::GenerateCombinedPropertyBag
 
 	UAnimNextRigVMAsset* Asset = UE::UAF::UncookedOnly::FUtils::GetAsset<UAnimNextRigVMAsset>(this);
 
-	FInstancedPropertyBag CombinedPropertyBag;
+	TArray<FPropertyBagPropertyDesc> PropertyDescriptions;
+	TArray<TConstArrayView<uint8>> PropertyDefaultValues;
+	// Required to keep default values from being freed
+	TArray<FInstancedPropertyBag> GeneratedPropertyBags;
+
+	// Flag set in case any properties to be combined overlap, if true will prevent combined property bag from being populated
+	bool bHasOverlappingProperties = false;	
 
-	auto AddInstancedPropertyBag = [this, &CombinedPropertyBag, &InCompileContext, Asset](const FInstancedPropertyBag& InPropertyBag, const UAnimNextRigVMAsset* SourceAsset, bool bValidateProperties = false)
+	// Checks whether or not the provided property name has already been added to PropertyDescriptions, in which case there is an overlap that cannot be resolved
+	auto HasOverlappingPropertyName = [&PropertyDescriptions, &InCompileContext, &bHasOverlappingProperties, Asset](const FName& PropertyNameToAdd, const UObject* SourceObject) -> bool
+	{
+		const bool bPropertyWithSameNameAlreadyExists = PropertyDescriptions.ContainsByPredicate([&PropertyNameToAdd](const FPropertyBagPropertyDesc& ExistingDesc)
+		{
+			return ExistingDesc.Name == PropertyNameToAdd;
+		});
+		
+		if (bPropertyWithSameNameAlreadyExists)
+		{
+			bHasOverlappingProperties = true;
+			InCompileContext.GetAssetCompileContext().Error(Asset, LOCTEXT("DuplicateVariableEntries", "Variable(s) with overlapping Name {0} found in {1}"), FText::FromName(PropertyNameToAdd), FText::FromString(SourceObject->GetPathName()));
+		}
+		
+		return bPropertyWithSameNameAlreadyExists;
+	};
+
+	auto AddInstancedPropertyBag = [this, &PropertyDescriptions, &PropertyDefaultValues, &InCompileContext, Asset, &HasOverlappingPropertyName](const FInstancedPropertyBag& InPropertyBag, const UAnimNextRigVMAsset* SourceAsset, bool bValidateProperties = false)
 	{
 		if(const UPropertyBag* PropertyBag = InPropertyBag.GetPropertyBagStruct())
 		{
 			TConstArrayView<FPropertyBagPropertyDesc> VariableDescs = PropertyBag->GetPropertyDescs();
 			if(VariableDescs.Num() != 0)
 			{
-				if (bValidateProperties)
+				TArray<TConstArrayView<uint8>> Values;
+				
+				for (const FPropertyBagPropertyDesc& Desc : VariableDescs)
 				{
-					bool bInvalidProperties = false;
-					for (const FPropertyBagPropertyDesc& Desc : VariableDescs)
+					if (!bValidateProperties || !HasOverlappingPropertyName(Desc.Name, SourceAsset))
 					{
-						if (const FPropertyBagPropertyDesc* ExistingDesc = CombinedPropertyBag.FindPropertyDescByName(Desc.Name))
-						{
-							InCompileContext.GetAssetCompileContext().Error(Asset, LOCTEXT("DuplicateVariableEntries", "Variable(s) with overlapping Name {0} found in {1}"), FText::FromName(Desc.Name), FText::FromString(SourceAsset->GetPathName()));
-
-							bInvalidProperties = true;
-						}
-					}
-
-					if (bInvalidProperties)
-					{
-						return;
+						PropertyDescriptions.Add(Desc);
+						PropertyDefaultValues.Add(TConstArrayView<uint8>(InPropertyBag.GetValue().GetMemory() + Desc.CachedProperty->GetOffset_ForInternal(), Desc.CachedProperty->GetSize()));
 					}
 				}
-				
-				CombinedPropertyBag.AddProperties(VariableDescs);
 			}
 		}
 	};
@@ -3248,7 +3261,7 @@ FInstancedPropertyBag UAnimNextRigVMAssetEditorData::GenerateCombinedPropertyBag
 	// It must match that of the variables FUAFInstanceVariableData::Initialize
 
 	// First add internal variables
-	FInstancedPropertyBag AssetBag = UE::UAF::UncookedOnly::FUtils::MakePropertyBagForEditorData(this, InCompileContext);
+	FInstancedPropertyBag& AssetBag = GeneratedPropertyBags.Add_GetRef(UE::UAF::UncookedOnly::FUtils::MakePropertyBagForEditorData(this, InCompileContext));
 	AddInstancedPropertyBag(AssetBag, Asset, false);
 
 	// Next add shared variables
@@ -3262,13 +3275,19 @@ FInstancedPropertyBag UAnimNextRigVMAssetEditorData::GenerateCombinedPropertyBag
 		const UAnimNextRigVMAssetEditorData* ReferencedEditorData = UE::UAF::UncookedOnly::FUtils::GetEditorData<UAnimNextRigVMAssetEditorData>(ReferencedVariableAsset);
 
 		FAnimNextRigVMAssetCompileContext CompileContext = { ReferencedEditorData };
-		AssetBag = UE::UAF::UncookedOnly::FUtils::MakePropertyBagForEditorData(ReferencedEditorData, ReferencedEditorData->GetVariableCompileContext(InSettings, CompileContext));
+		AssetBag = GeneratedPropertyBags.Add_GetRef(UE::UAF::UncookedOnly::FUtils::MakePropertyBagForEditorData(ReferencedEditorData, ReferencedEditorData->GetVariableCompileContext(InSettings, CompileContext)));
 		AddInstancedPropertyBag(AssetBag, ReferencedVariableAsset, true);
 	}
 
+	// Required to keep default struct property values from being freed
+	TArray<FInstancedStruct> InstancedReferenceStructs;
 	// Next add native structs
 	for (const UScriptStruct* ReferencedStruct : Asset->ReferencedVariableStructs)
 	{
+		// Allocate default values for populating property bag entries with
+		FInstancedStruct& StructInstance = InstancedReferenceStructs.AddDefaulted_GetRef();
+		StructInstance.InitializeAs(ReferencedStruct);
+
 		TSharedRef<UE::UAF::FStructDataCache> StructData = UE::UAF::FStructDataCache::GetStructInfo(ReferencedStruct);
 		for (const UE::UAF::FStructDataCache::FPropertyInfo& PropertyInfo : StructData->GetProperties())
 		{
@@ -3279,32 +3298,28 @@ FInstancedPropertyBag UAnimNextRigVMAssetEditorData::GenerateCombinedPropertyBag
 				InCompileContext.GetAssetCompileContext().Error(Asset, LOCTEXT("InvalidNativeVariableType", "@@ Variable '{0}' from '{1}' has unsupported variable type '{2}'"), FText::FromName(PropertyInfo.Property->GetFName()),  FText::FromName(ReferencedStruct->GetFName()), FText::FromString(ParamType.ToString()));
 				continue;
 			}
-			
-			constexpr bool bOverrideProperty = false;
-			const EPropertyBagAlterationResult Result = CombinedPropertyBag.AddProperty(PropertyInfo.Property->GetFName(), PropertyInfo.Property, bOverrideProperty);
 
-			if (Result != EPropertyBagAlterationResult::Success)
+			if (!HasOverlappingPropertyName(PropertyInfo.Property->GetFName(), ReferencedStruct))
 			{
-				if (Result == EPropertyBagAlterationResult::TargetPropertyAlreadyExists)
-				{
-					InCompileContext.GetAssetCompileContext().Error(
-						Asset,
-						LOCTEXT("AddNativeSharedVariablePropertyFailure_NameOverlaps", "The property for variable '{0}' inside Native Shared Variables '{1}' could not be generated as the name overlaps with another variable"),
-						FText::FromName(PropertyInfo.Property->GetFName()),
-						FText::FromName(ReferencedStruct->GetFName()));
-				}
-				else
-				{
-					InCompileContext.GetAssetCompileContext().Error(
-						Asset,
-						LOCTEXT("AddNativeSharedVariablePropertyFailure_Other", "The property for variable '{0}' inside Native Shared Variables '{1}' could not be generated"),
-						FText::FromName(PropertyInfo.Property->GetFName()),
-						FText::FromName(ReferencedStruct->GetFName()));
-				}
+				PropertyDescriptions.Add(FPropertyBagPropertyDesc(PropertyInfo.Property->GetFName(), PropertyInfo.Property));
+				PropertyDefaultValues.Add(TConstArrayView<uint8>(StructInstance.GetMemory() + PropertyInfo.Property->GetOffset_ForInternal(), PropertyInfo.Property->GetSize()));
 			}
 		}
 	}
 
+	FInstancedPropertyBag CombinedPropertyBag;
+	if (!bHasOverlappingProperties)
+	{
+		const EPropertyBagResult Result = CombinedPropertyBag.ReplaceAllPropertiesAndValues(PropertyDescriptions, PropertyDefaultValues);
+		if (Result != EPropertyBagResult::Success)
+		{
+			InCompileContext.GetAssetCompileContext().Error(
+				Asset,
+				LOCTEXT("FailedToGenerateCombinedPropertyBag", "Failed to generate combines property bag for %s"), FText::FromName(Asset->GetFName())
+			);
+		}
+	 }
+
 	return CombinedPropertyBag;
 }
 
diff --git a/Engine/Plugins/Experimental/UAF/UAFTests/Source/UAFTestSuite/Private/AnimNextVariablesTest.cpp b/Engine/Plugins/Experimental/UAF/UAFTests/Source/UAFTestSuite/Private/AnimNextVariablesTest.cpp
--- a/Engine/Plugins/Experimental/UAF/UAFTests/Source/UAFTestSuite/Private/AnimNextVariablesTest.cpp
+++ b/Engine/Plugins/Experimental/UAF/UAFTests/Source/UAFTestSuite/Private/AnimNextVariablesTest.cpp
@@ -222,7 +222,7 @@ bool FVariableTypesTest::RunTest(const FString& InParameters)
 	return true;
 }
 
-IMPLEMENT_SIMPLE_AUTOMATION_TEST(FVariables, "Animation.AnimNext.Variables", EAutomationTestFlags::EditorContext | EAutomationTestFlags::EngineFilter)
+IMPLEMENT_SIMPLE_AUTOMATION_TEST(FVariables, "Animation.AnimNext.Variables.Variables", EAutomationTestFlags::EditorContext | EAutomationTestFlags::EngineFilter)
 
 bool FVariables::RunTest(const FString& InParameters)
 {
