UAF: Remove logspam warning for systems without bindings

#Synced-CL 48809816

diff --git a/Engine/Plugins/Experimental/UAF/UAF/Source/UAF/Private/Module/AnimNextModuleInstance.cpp b/Engine/Plugins/Experimental/UAF/UAF/Source/UAF/Private/Module/AnimNextModuleInstance.cpp
--- a/Engine/Plugins/Experimental/UAF/UAF/Source/UAF/Private/Module/AnimNextModuleInstance.cpp
+++ b/Engine/Plugins/Experimental/UAF/UAF/Source/UAF/Private/Module/AnimNextModuleInstance.cpp
@@ -590,8 +590,6 @@ void FAnimNextModuleInstance::RunScriptEvent(const UE::UAF::FScriptEvent& InEven
 {
 	if (!InEvent.IsCallable())
 	{
-		UE_LOG(LogAnimation, Warning, TEXT("FAnimNextModuleInstance::RunScriptEvent: Could not run script event [%s] - The event is not callable."),
-			*InEvent.GetEventName().ToString());
 		return;
 	}
 
