Fix-up error log message spelling mistake

#Synced-CL 48786226

diff --git a/Engine/Plugins/Experimental/UAF/UAF/Source/UAFUncookedOnly/Private/AnimNextRigVMAssetEditorData.cpp b/Engine/Plugins/Experimental/UAF/UAF/Source/UAFUncookedOnly/Private/AnimNextRigVMAssetEditorData.cpp
--- a/Engine/Plugins/Experimental/UAF/UAF/Source/UAFUncookedOnly/Private/AnimNextRigVMAssetEditorData.cpp
+++ b/Engine/Plugins/Experimental/UAF/UAF/Source/UAFUncookedOnly/Private/AnimNextRigVMAssetEditorData.cpp
@@ -3315,7 +3315,7 @@ FInstancedPropertyBag UAnimNextRigVMAssetEditorData::GenerateCombinedPropertyBag
 		{
 			InCompileContext.GetAssetCompileContext().Error(
 				Asset,
-				LOCTEXT("FailedToGenerateCombinedPropertyBag", "Failed to generate combines property bag for %s"), FText::FromName(Asset->GetFName())
+				LOCTEXT("FailedToGenerateCombinedPropertyBag", "Failed to generate combined property bag for %s"), FText::FromName(Asset->GetFName())
 			);
 		}
 	 }
