Fix crash attempting to reinitialize a previously marked-for-release FAnimNextModuleInstance
#misc added flag to mark an instance as pending release

#Synced-CL 48786226

diff --git a/Engine/Plugins/Experimental/UAF/UAF/Source/UAF/Internal/Module/AnimNextModuleInstance.h b/Engine/Plugins/Experimental/UAF/UAF/Source/UAF/Internal/Module/AnimNextModuleInstance.h
--- a/Engine/Plugins/Experimental/UAF/UAF/Source/UAF/Internal/Module/AnimNextModuleInstance.h
+++ b/Engine/Plugins/Experimental/UAF/UAF/Source/UAF/Internal/Module/AnimNextModuleInstance.h
@@ -290,6 +290,9 @@ private:
 #if WITH_EDITOR
 	// Whether we are currently recreating this instance because of compilation/reinstancing
 	bool bIsRecreatingOnCompile : 1 = false;
+	
+	// Whether ore not the module instance has had its custom destructor called (queueing it up for release at the end of the frame)
+	bool bHasBeenQueuedForRelease : 1 = false;
 #endif
 
 private:
@@ -314,4 +317,4 @@ struct TStructOpsTypeTraits<FAnimNextModuleInstance> : public TStructOpsTypeTrai
 #define TRACE_UAF_MODULE(ModuleInstance)
 #endif
 
-#undef UE_API
\ No newline at end of file
+#undef UE_API
diff --git a/Engine/Plugins/Experimental/UAF/UAF/Source/UAF/Private/Module/SystemReference.cpp b/Engine/Plugins/Experimental/UAF/UAF/Source/UAF/Private/Module/SystemReference.cpp
--- a/Engine/Plugins/Experimental/UAF/UAF/Source/UAF/Private/Module/SystemReference.cpp
+++ b/Engine/Plugins/Experimental/UAF/UAF/Source/UAF/Private/Module/SystemReference.cpp
@@ -94,6 +94,9 @@ struct FCustomSystemDeleter
 		if (InInstance)
 		{
 			check(InInstance == &Private::GetPool().SystemInstances.Get(Handle));
+#if WITH_EDITOR 
+			InInstance->bHasBeenQueuedForRelease = true;
+#endif // WITH_EDITOR
 
 			Private::GetPool().PendingActions.Enqueue([Handle = Handle]()
 			{
@@ -505,15 +508,18 @@ void UUAFEngineSubsystem::OnCompileJobFinished(UAnimNextRigVMAsset* InAsset)
 {
 	using namespace UE::UAF;
 
+	// Flush any pending (queued up) instance release/destroy actions 
+	Private::GetPool().FlushPendingActions();
+
 	check(IsInGameThread());
 	UE::TReadScopeLock InstancesLockScope(Private::GetPool().SystemInstancesLock);
 
 	for (FAnimNextModuleInstance& Instance : Private::GetPool().SystemInstances)
 	{
-		if (Instance.GetSystem() == InAsset)
+		if (Instance.GetSystem() == InAsset && ensure(Instance.bHasBeenQueuedForRelease == false))
 		{
 			Instance.OnCompileJobFinished();
 		}
 	}
 }
-#endif
\ No newline at end of file
+#endif
