Fix trying to add unbound attribute bindings in PartitionAttributes
* In CollectAttributeTypes we correctly guard against attributes that're unbound however this same check is missing later on when adding attributes.

#Synced-CL 48779969

diff --git a/Engine/Plugins/Experimental/UAF/UAF/Source/UAF/Private/UAF/Attributes/AttributeBindingData.cpp b/Engine/Plugins/Experimental/UAF/UAF/Source/UAF/Private/UAF/Attributes/AttributeBindingData.cpp
--- a/Engine/Plugins/Experimental/UAF/UAF/Source/UAF/Private/UAF/Attributes/AttributeBindingData.cpp
+++ b/Engine/Plugins/Experimental/UAF/UAF/Source/UAF/Private/UAF/Attributes/AttributeBindingData.cpp
@@ -65,10 +65,7 @@ namespace UE::UAF
 
 			for (const FAbstractSkeleton_AttributeBinding& AttributeBinding : AttributeBindings)
 			{
-				if (!AttributeBinding.SetName.IsNone())
-				{
-					AttributeTypes.AddUnique(AttributeBinding.Attribute.GetType());
-				}
+				AttributeTypes.AddUnique(AttributeBinding.Attribute.GetType());
 			}
 
 			return AttributeTypes;
