Comment typo fixup

#Synced-CL 48907971

diff --git a/Engine/Plugins/Experimental/UAF/UAF/Source/UAF/Internal/Module/AnimNextModuleInstance.h b/Engine/Plugins/Experimental/UAF/UAF/Source/UAF/Internal/Module/AnimNextModuleInstance.h
--- a/Engine/Plugins/Experimental/UAF/UAF/Source/UAF/Internal/Module/AnimNextModuleInstance.h
+++ b/Engine/Plugins/Experimental/UAF/UAF/Source/UAF/Internal/Module/AnimNextModuleInstance.h
@@ -291,7 +291,7 @@ private:
 	// Whether we are currently recreating this instance because of compilation/reinstancing
 	bool bIsRecreatingOnCompile : 1 = false;
 	
-	// Whether ore not the module instance has had its custom destructor called (queueing it up for release at the end of the frame)
+	// Whether the module instance has had its custom destructor called (queueing it up for release at the end of the frame)
 	bool bHasBeenQueuedForRelease : 1 = false;
 #endif
 
