RigVM: allow to hide pins via the ed graph schema

#rnx

#Synced-CL 48908485

diff --git a/Engine/Plugins/Experimental/UAF/UAF/Source/UAFUncookedOnly/Internal/AnimNextEdGraphSchema.h b/Engine/Plugins/Experimental/UAF/UAF/Source/UAFUncookedOnly/Internal/AnimNextEdGraphSchema.h
--- a/Engine/Plugins/Experimental/UAF/UAF/Source/UAFUncookedOnly/Internal/AnimNextEdGraphSchema.h
+++ b/Engine/Plugins/Experimental/UAF/UAF/Source/UAFUncookedOnly/Internal/AnimNextEdGraphSchema.h
@@ -18,6 +18,7 @@ class UAnimNextEdGraphSchema : public URigVMEdGraphSchema
 
 	// URigVMEdGraphSchema interface
 	virtual TSubclassOf<URigVMEdGraphNode> GetGraphNodeClass(const URigVMEdGraph* InGraph) const override { return UAnimNextEdGraphNode::StaticClass(); }
+	UE_API virtual bool ShouldHidePinOnNode(const URigVMPin* InModelPin) const override;
 
 	// UEdGraphSchema interface
 	UE_API virtual void GetGraphDisplayInformation(const UEdGraph& Graph, /*out*/ FGraphDisplayInfo& DisplayInfo) const override;
diff --git a/Engine/Plugins/Experimental/UAF/UAF/Source/UAFUncookedOnly/Private/AnimNextEdGraphSchema.cpp b/Engine/Plugins/Experimental/UAF/UAF/Source/UAFUncookedOnly/Private/AnimNextEdGraphSchema.cpp
--- a/Engine/Plugins/Experimental/UAF/UAF/Source/UAFUncookedOnly/Private/AnimNextEdGraphSchema.cpp
+++ b/Engine/Plugins/Experimental/UAF/UAF/Source/UAFUncookedOnly/Private/AnimNextEdGraphSchema.cpp
@@ -11,6 +11,21 @@
 
 #define LOCTEXT_NAMESPACE "AnimNextEdGraphSchema"
 
+bool UAnimNextEdGraphSchema::ShouldHidePinOnNode(const URigVMPin* InModelPin) const
+{
+	if (InModelPin)
+	{
+		if (const URigVMPin* RootPin = InModelPin->GetRootPin())
+		{
+			if (RootPin->IsDefinedAsInputVariable() && !RootPin->GetBoundVariableName().IsEmpty())
+			{
+				return true;
+			}
+		}
+	}
+	return Super::ShouldHidePinOnNode(InModelPin);
+}
+
 void UAnimNextEdGraphSchema::GetGraphDisplayInformation(const UEdGraph& Graph, /*out*/ FGraphDisplayInfo& DisplayInfo) const
 {
 	Super::GetGraphDisplayInformation(Graph, DisplayInfo);
diff --git a/Engine/Plugins/Runtime/RigVM/Source/RigVMDeveloper/Private/EdGraph/RigVMEdGraphNode.cpp b/Engine/Plugins/Runtime/RigVM/Source/RigVMDeveloper/Private/EdGraph/RigVMEdGraphNode.cpp
--- a/Engine/Plugins/Runtime/RigVM/Source/RigVMDeveloper/Private/EdGraph/RigVMEdGraphNode.cpp
+++ b/Engine/Plugins/Runtime/RigVM/Source/RigVMDeveloper/Private/EdGraph/RigVMEdGraphNode.cpp
@@ -773,6 +773,15 @@ bool URigVMEdGraphNode::CreateGraphPinFromModelPin(const URigVMPin* InModelPin,
 		return false;
 	}
 
+	// obey to the schema. certain pins may be hidden as per schema contract.
+	if (const URigVMEdGraphSchema* Schema = Cast<URigVMEdGraphSchema>(GetSchema()))
+	{
+		if (Schema->ShouldHidePinOnNode(InModelPin))
+		{
+			return false;
+		}
+	}
+
 	const FString Category = InModelPin->GetCategory();
 	if(!Category.IsEmpty())
 	{
diff --git a/Engine/Plugins/Runtime/RigVM/Source/RigVMDeveloper/Private/EdGraph/RigVMEdGraphSchema.cpp b/Engine/Plugins/Runtime/RigVM/Source/RigVMDeveloper/Private/EdGraph/RigVMEdGraphSchema.cpp
--- a/Engine/Plugins/Runtime/RigVM/Source/RigVMDeveloper/Private/EdGraph/RigVMEdGraphSchema.cpp
+++ b/Engine/Plugins/Runtime/RigVM/Source/RigVMDeveloper/Private/EdGraph/RigVMEdGraphSchema.cpp
@@ -2702,5 +2702,10 @@ bool URigVMEdGraphSchema::IsRigVMDefaultEvent(const FName& InEventName) const
 	return false;
 }
 
+bool URigVMEdGraphSchema::ShouldHidePinOnNode(const URigVMPin* InModelPin) const
+{
+	return false;
+}
+
 #undef LOCTEXT_NAMESPACE
 
diff --git a/Engine/Plugins/Runtime/RigVM/Source/RigVMDeveloper/Public/EdGraph/RigVMEdGraphSchema.h b/Engine/Plugins/Runtime/RigVM/Source/RigVMDeveloper/Public/EdGraph/RigVMEdGraphSchema.h
--- a/Engine/Plugins/Runtime/RigVM/Source/RigVMDeveloper/Public/EdGraph/RigVMEdGraphSchema.h
+++ b/Engine/Plugins/Runtime/RigVM/Source/RigVMDeveloper/Public/EdGraph/RigVMEdGraphSchema.h
@@ -397,6 +397,9 @@ public:
 
 	UE_API virtual bool IsRigVMDefaultEvent(const FName& InEventName) const;
 
+	// Allows to hide an input variable pin from node node
+	UE_API virtual bool ShouldHidePinOnNode(const URigVMPin* InModelPin) const;
+
 private:
 
 	const UEdGraphPin* LastPinForCompatibleCheck = nullptr;
