Fix for Skeleton mode for translational retargeting

#Synced-CL 48901052

diff --git a/Engine/Plugins/Experimental/UAF/UAF/Source/UAF/Private/DecompressionTools.cpp b/Engine/Plugins/Experimental/UAF/UAF/Source/UAF/Private/DecompressionTools.cpp
--- a/Engine/Plugins/Experimental/UAF/UAF/Source/UAF/Private/DecompressionTools.cpp
+++ b/Engine/Plugins/Experimental/UAF/UAF/Source/UAF/Private/DecompressionTools.cpp
@@ -45,6 +45,7 @@ static bool GetForceRawData()
 
 struct FGetBonePoseScratchArea : public TThreadSingleton<FGetBonePoseScratchArea>
 {
+	BoneTrackArray SkeletonPairs;
 	BoneTrackArray RotationScalePairs;
 	BoneTrackArray TranslationPairs;
 	BoneTrackArray AnimScaleRetargetingPairs;
@@ -1067,6 +1068,7 @@ void FDecompressionTools::DecompressPose(FLODPose& OutAnimationPoseData,
 	const bool bDisableRetargeting = OutAnimationPoseData.GetDisableRetargeting();
 
 	FGetBonePoseScratchArea& ScratchArea = FGetBonePoseScratchArea::Get();
+	BoneTrackArray& SkeletonPairs = ScratchArea.SkeletonPairs;
 	BoneTrackArray& RotationScalePairs = ScratchArea.RotationScalePairs;
 	BoneTrackArray& TranslationPairs = ScratchArea.TranslationPairs;
 	BoneTrackArray& AnimScaleRetargetingPairs = ScratchArea.AnimScaleRetargetingPairs;
@@ -1074,6 +1076,7 @@ void FDecompressionTools::DecompressPose(FLODPose& OutAnimationPoseData,
 	BoneTrackArray& OrientAndScaleRetargetingPairs = ScratchArea.OrientAndScaleRetargetingPairs;
 
 	// build a list of desired bones
+	SkeletonPairs.Reset();
 	RotationScalePairs.Reset();
 	TranslationPairs.Reset();
 	AnimScaleRetargetingPairs.Reset();
@@ -1131,6 +1134,9 @@ void FDecompressionTools::DecompressPose(FLODPose& OutAnimationPoseData,
 					// Skip extracting translation component for EBoneTranslationRetargetingMode::Skeleton.
 					switch (RetargetMode)
 					{
+					case EBoneTranslationRetargetingMode::Skeleton:
+						SkeletonPairs.Add(BoneTrackPair(LODBoneIndex, TrackIndex));
+						break;
 					case EBoneTranslationRetargetingMode::Animation:
 						TranslationPairs.Add(BoneTrackPair(LODBoneIndex, TrackIndex));
 						break;
@@ -1289,6 +1295,7 @@ void FDecompressionTools::DecompressPose(FLODPose& OutAnimationPoseData,
 	FRetargetingTools::RetargetPose(
 		ReferencePose,
 		SkeletonRemapping,
+		SkeletonPairs,
 		AnimScaleRetargetingPairs,
 		AnimRelativeRetargetingPairs,
 		OrientAndScaleRetargetingPairs,
diff --git a/Engine/Plugins/Experimental/UAF/UAF/Source/UAF/Private/Graph/RigUnit_AnimNextRetargetPose.cpp b/Engine/Plugins/Experimental/UAF/UAF/Source/UAF/Private/Graph/RigUnit_AnimNextRetargetPose.cpp
--- a/Engine/Plugins/Experimental/UAF/UAF/Source/UAF/Private/Graph/RigUnit_AnimNextRetargetPose.cpp
+++ b/Engine/Plugins/Experimental/UAF/UAF/Source/UAF/Private/Graph/RigUnit_AnimNextRetargetPose.cpp
@@ -82,7 +82,7 @@ FRigUnit_AnimNextRetargetPose_Execute()
 			SourceSkeleton, TargetSkeleton,
 			bUseSourceRetargetModes,
 			bDisableRetargeting,
-			RetargetPoseData.AnimScalePairs, RetargetPoseData.AnimRelativePairs, RetargetPoseData.OrientAndScalePairs);
+			RetargetPoseData.SkeletonPairs, RetargetPoseData.AnimScalePairs, RetargetPoseData.AnimRelativePairs, RetargetPoseData.OrientAndScalePairs);
 
 		RetargetPoseData.PairsInitForSourcePoseLOD = SourcePose.LODLevel;
 		RetargetPoseData.PairsInitForTargetPoseLOD = TargetPose.LODLevel;
@@ -101,7 +101,7 @@ FRigUnit_AnimNextRetargetPose_Execute()
 	FRetargetingTools::RetargetPose(
 		TargetRefPose,
 		SkeletonRemapping,
-		RetargetPoseData.AnimScalePairs, RetargetPoseData.AnimRelativePairs, RetargetPoseData.OrientAndScalePairs,
+		RetargetPoseData.SkeletonPairs, RetargetPoseData.AnimScalePairs, RetargetPoseData.AnimRelativePairs, RetargetPoseData.OrientAndScalePairs,
 		SourceSkeleton->GetRefLocalPoses(),
 		TargetPose);
 }
\ No newline at end of file
diff --git a/Engine/Plugins/Experimental/UAF/UAF/Source/UAF/Private/Graph/RigUnit_AnimNextRetargetPose.h b/Engine/Plugins/Experimental/UAF/UAF/Source/UAF/Private/Graph/RigUnit_AnimNextRetargetPose.h
--- a/Engine/Plugins/Experimental/UAF/UAF/Source/UAF/Private/Graph/RigUnit_AnimNextRetargetPose.h
+++ b/Engine/Plugins/Experimental/UAF/UAF/Source/UAF/Private/Graph/RigUnit_AnimNextRetargetPose.h
@@ -20,6 +20,7 @@ struct FRetargetPoseData
 
 	FRemapPoseData RemapPoseData; // TODO: Store this per remap node or inside the data registry? We don't have ref counting ref poses and don't know when to remove these though.
 
+	BoneTrackArray SkeletonPairs;
 	BoneTrackArray AnimScalePairs;
 	BoneTrackArray AnimRelativePairs;
 	BoneTrackArray OrientAndScalePairs;
diff --git a/Engine/Plugins/Experimental/UAF/UAF/Source/UAF/Private/RetargetingTools.cpp b/Engine/Plugins/Experimental/UAF/UAF/Source/UAF/Private/RetargetingTools.cpp
--- a/Engine/Plugins/Experimental/UAF/UAF/Source/UAF/Private/RetargetingTools.cpp
+++ b/Engine/Plugins/Experimental/UAF/UAF/Source/UAF/Private/RetargetingTools.cpp
@@ -15,10 +15,12 @@ namespace UE::UAF
 	    const USkeleton* TargetSkeleton,
 	    bool bUseSourceRetargetModes,
 	    bool bDisableRetargeting,
+		BoneTrackArray& OutSkeletonPairs,
 	    BoneTrackArray& OutAnimScalePairs,
 	    BoneTrackArray& OutAnimRelativePairs,
 	    BoneTrackArray& OutOrientAndScalePairs)
 	{
+		OutSkeletonPairs.Reset();
 	    OutAnimScalePairs.Reset();
 	    OutAnimRelativePairs.Reset();
 	    OutOrientAndScalePairs.Reset();
@@ -57,6 +59,13 @@ namespace UE::UAF
 
 	        switch (Mode)
 	        {
+				case EBoneTranslationRetargetingMode::Skeleton:
+				{
+					// (AtomIndex = Target LODPose, TrackIndex = Source Skeleton)
+					OutSkeletonPairs.Add(BoneTrackPair(TargetPoseBoneIndex, SourceSkeletonBoneIndex));
+					break;
+				}
+
 				case EBoneTranslationRetargetingMode::AnimationScaled:
 	        	{
 					// (AtomIndex = Target LODPose, TrackIndex = Source Skeleton)
@@ -275,6 +284,7 @@ namespace UE::UAF
 	void FRetargetingTools::RetargetPose(
 		const FReferencePose& TargetReferencePose,
 		const FSkeletonRemapping& SourceToTargetRemapping,
+		BoneTrackArray& SkeletonRetargetingPairs,
 		BoneTrackArray& AnimScaleRetargetingPairs,
 		BoneTrackArray& AnimRelativeRetargetingPairs,
 		BoneTrackArray& OrientAndScaleRetargetingPairs,
@@ -336,6 +346,16 @@ namespace UE::UAF
 			}
 		}
 
+		if (!SkeletonRetargetingPairs.IsEmpty())
+		{
+			for (const BoneTrackPair& BonePair : SkeletonRetargetingPairs)
+			{
+				const int32 LODBoneIndex(BonePair.AtomIndex);
+				const int32 SourceSkeletonBoneIndex = BonePair.TrackIndex;
+				OutPose.LocalTransformsView[LODBoneIndex].SetTranslation(OutPose.GetRefPose().GetRefPoseTransform(LODBoneIndex).GetTranslation());
+			}
+		}
+
 		// TODO : Have to recreate GetRetargetSourceCachedData for optimization
 		const int32 NumBonesToOrientAndScaleRetarget = OrientAndScaleRetargetingPairs.Num();
 		if (NumBonesToOrientAndScaleRetarget > 0)
diff --git a/Engine/Plugins/Experimental/UAF/UAF/Source/UAF/Public/RetargetingTools.h b/Engine/Plugins/Experimental/UAF/UAF/Source/UAF/Public/RetargetingTools.h
--- a/Engine/Plugins/Experimental/UAF/UAF/Source/UAF/Public/RetargetingTools.h
+++ b/Engine/Plugins/Experimental/UAF/UAF/Source/UAF/Public/RetargetingTools.h
@@ -27,6 +27,7 @@ namespace UE::UAF
 		 * @param TargetSkeleton            Retargeting target skeleton.
 		 * @param bUseSourceRetargetModes   Set to true when you wish to extract the retarget mode from the source skeleton instead of the target skeleton.
 		 * @param bDisableRetargeting       Is translational retargeting disabled?
+		 * @param OutSkeletonPairs          Output pairs for EBoneTranslationRetargetingMode::Skeleton retargeting mode where bones use ref pose translation but rotate along with the source.
 		 * @param OutAnimScalePairs         Output pairs for AnimationScaled.
 		 * @param OutAnimRelativePairs      Output pairs for AnimationRelative.
 		 * @param OutOrientAndScalePairs    Output pairs for OrientAndScale.
@@ -39,6 +40,7 @@ namespace UE::UAF
 			const USkeleton* TargetSkeleton,
 			bool bUseSourceRetargetModes,
 			bool bDisableRetargeting,
+			BoneTrackArray& OutSkeletonPairs,
 			BoneTrackArray& OutAnimScalePairs,
 			BoneTrackArray& OutAnimRelativePairs,
 			BoneTrackArray& OutOrientAndScalePairs);
@@ -195,6 +197,7 @@ namespace UE::UAF
 		 * Work scales with number of pairs, not bones.
 		 * @param TargetReferencePose			Target reference pose.
 		 * @param SourceToTargetRemapping		Maps source and target skeleton bone indices.
+		 * @param SkeletonRetargetingPairs		Pairs used for retargeting where bones use ref pose translation but rotate along with the source.
 		 * @param AnimScaleRetargetingPairs		Pairs for AnimationScaled pass.
 		 * @param AnimRelativeRetargetingPairs	Pairs for AnimationRelative pass.
 		 * @param OrientAndScaleRetargetingPairs Pairs for OrientAndScale pass.
@@ -204,6 +207,7 @@ namespace UE::UAF
 		static UAF_API void RetargetPose(
 			const FReferencePose& TargetReferencePose,
 			const FSkeletonRemapping& SourceToTargetRemapping,
+			BoneTrackArray& SkeletonRetargetingPairs,
 			BoneTrackArray& AnimScaleRetargetingPairs,
 			BoneTrackArray& AnimRelativeRetargetingPairs,
 			BoneTrackArray& OrientAndScaleRetargetingPairs,
