Fix unbinding an attribute in also always removing it
* Also add some more editor operations

#Synced-CL 48779969

diff --git a/Engine/Plugins/Experimental/UAF/UAF/Source/UAF/Private/UAF/AbstractSkeleton/AbstractSkeletonSetBinding.cpp b/Engine/Plugins/Experimental/UAF/UAF/Source/UAF/Private/UAF/AbstractSkeleton/AbstractSkeletonSetBinding.cpp
--- a/Engine/Plugins/Experimental/UAF/UAF/Source/UAF/Private/UAF/AbstractSkeleton/AbstractSkeletonSetBinding.cpp
+++ b/Engine/Plugins/Experimental/UAF/UAF/Source/UAF/Private/UAF/AbstractSkeleton/AbstractSkeletonSetBinding.cpp
@@ -174,7 +174,17 @@ bool UAbstractSkeletonSetBinding::RemoveAttributeFromSet(const FAnimationAttribu
 
 	if (BindingIndex != INDEX_NONE)
 	{
-		AttributeBindings.RemoveAtSwap(BindingIndex);
+		FAbstractSkeleton_AttributeBinding& Binding = AttributeBindings[BindingIndex];
+
+		if (Binding.SetName.IsNone())
+		{
+			AttributeBindings.RemoveAtSwap(BindingIndex);
+		}
+		else
+		{
+			Binding.SetName = NAME_None;
+		}
+
 		RebuildBindingData();
 		return true;
 	}
diff --git a/Engine/Plugins/Experimental/UAF/UAF/Source/UAFEditor/Private/UAF/AbstractSkeleton/Sets/SAttributesList.cpp b/Engine/Plugins/Experimental/UAF/UAF/Source/UAFEditor/Private/UAF/AbstractSkeleton/Sets/SAttributesList.cpp
--- a/Engine/Plugins/Experimental/UAF/UAF/Source/UAFEditor/Private/UAF/AbstractSkeleton/Sets/SAttributesList.cpp
+++ b/Engine/Plugins/Experimental/UAF/UAF/Source/UAFEditor/Private/UAF/AbstractSkeleton/Sets/SAttributesList.cpp
@@ -73,6 +73,7 @@ void SAttributesList::Construct(const FArguments& InArgs)
 				SAssignNew(ListView, SListView<FListItemPtr>)
 					.ListItemsSource(&ListItems)
 					.OnGenerateRow(this, &SAttributesList::ListView_OnGenerateRow)
+					.OnContextMenuOpening(this, &SAttributesList::ListView_OnContextMenuOpening)
 					.HeaderRow(
 						SNew(SHeaderRow)
 						+ SHeaderRow::Column(FColumns::SetId)
@@ -265,6 +266,81 @@ TSharedRef<ITableRow> SAttributesList::ListView_OnGenerateRow(FListItemPtr Item,
 			});
 }
 
+TSharedPtr<SWidget> SAttributesList::ListView_OnContextMenuOpening()
+{
+	TArray<FListItemPtr> Selection = ListView->GetSelectedItems();
+	if (Selection.IsEmpty())
+	{
+		return SNullWidget::NullWidget;
+	}
+
+	FMenuBuilder MenuBuilder(true, nullptr);
+	FListItemPtr Selected = Selection[0]; // Single selection only
+
+	if (Selected->SetName.IsNone())
+	{
+		MenuBuilder.AddMenuEntry(
+			LOCTEXT("RemoveUnboundAttribute_Label", "Remove Attribute"),
+			LOCTEXT("RemoveUnboundAttribute_Tooltip", "Removes the selected attribute"),
+			FSlateIcon(FAppStyle::GetAppStyleSetName(), "Icons.Delete"),
+			FUIAction(
+				FExecuteAction::CreateLambda([this, Selected]()
+					{
+						const FScopedTransaction Transaction(LOCTEXT("RemoveUnboundAttribute_Label", "Remove Attribute"));
+						SetBinding->Modify();
+
+						ensure(SetBinding->RemoveAttributeFromSet(Selected->Attribute));
+
+						RepopulateListData();
+					}
+				)
+			));
+	}
+	else
+	{
+		MenuBuilder.AddMenuEntry(
+			LOCTEXT("UnbindBoundAttribute_Label", "Unbind Attribute"),
+			LOCTEXT("UnbindBoundAttribute_Tooltip", "Unbinds the selected attribute"),
+			FSlateIcon(FAppStyle::GetAppStyleSetName(), "Icons.Unlink"),
+			FUIAction(
+				FExecuteAction::CreateLambda([this, Selected]()
+					{
+						const FScopedTransaction Transaction(LOCTEXT("UnbindBoundAttribute_Label", "Unbind Attribute"));
+						SetBinding->Modify();
+
+						ensure(SetBinding->RemoveAttributeFromSet(Selected->Attribute));
+
+						RepopulateListData();
+					}
+				)
+			));
+
+		MenuBuilder.AddMenuEntry(
+			LOCTEXT("UnbindAndRemoveBoundAttribute_Label", "Unbind and Remove Attribute"),
+			LOCTEXT("UnbindAndRemoveBoundAttribute_Tooltip", "Unbinds and removes the selected attribute"),
+			FSlateIcon(FAppStyle::GetAppStyleSetName(), "Icons.Delete"),
+			FUIAction(
+				FExecuteAction::CreateLambda([this, Selected]()
+					{
+						const FScopedTransaction Transaction(LOCTEXT("UnbindAndRemoveBoundAttribute_Label", "Unbind and Remove Attribute"));
+						SetBinding->Modify();
+
+						// First the attribute is removed from the set it is bound to and placed in the 'None' set.
+						// Attributes in the 'None' set are all the unbound attributes that the user has added to the list of attributes.
+						ensure(SetBinding->RemoveAttributeFromSet(Selected->Attribute));
+
+						// Then removing attributes from the 'None' set actually removes them.
+						ensure(SetBinding->RemoveAttributeFromSet(Selected->Attribute));
+
+						RepopulateListData();
+					}
+				)
+			));
+	}
+
+	return MenuBuilder.MakeWidget();
+}
+
 void SAttributesList::RegisterMenus()
 {
 	if (UToolMenus::Get()->IsMenuRegistered(FMenus::AddAttributeId))
diff --git a/Engine/Plugins/Experimental/UAF/UAF/Source/UAFEditor/Private/UAF/AbstractSkeleton/Sets/SAttributesList.h b/Engine/Plugins/Experimental/UAF/UAF/Source/UAFEditor/Private/UAF/AbstractSkeleton/Sets/SAttributesList.h
--- a/Engine/Plugins/Experimental/UAF/UAF/Source/UAFEditor/Private/UAF/AbstractSkeleton/Sets/SAttributesList.h
+++ b/Engine/Plugins/Experimental/UAF/UAF/Source/UAFEditor/Private/UAF/AbstractSkeleton/Sets/SAttributesList.h
@@ -53,6 +53,8 @@ namespace UE::UAF::Sets
 
 	private:
 		TSharedRef<ITableRow> ListView_OnGenerateRow(FListItemPtr Item, const TSharedRef<STableViewBase>& OwnerTable);
+
+		TSharedPtr<SWidget> ListView_OnContextMenuOpening();
 		
 		TSharedRef<SWidget> CreateAddAttributeWidget();
 
diff --git a/Engine/Plugins/Experimental/UAF/UAF/Source/UAFEditor/Private/UAF/AbstractSkeleton/Sets/SSetBinding.cpp b/Engine/Plugins/Experimental/UAF/UAF/Source/UAFEditor/Private/UAF/AbstractSkeleton/Sets/SSetBinding.cpp
--- a/Engine/Plugins/Experimental/UAF/UAF/Source/UAFEditor/Private/UAF/AbstractSkeleton/Sets/SSetBinding.cpp
+++ b/Engine/Plugins/Experimental/UAF/UAF/Source/UAFEditor/Private/UAF/AbstractSkeleton/Sets/SSetBinding.cpp
@@ -604,7 +604,7 @@ void SSetBinding::Construct(const FArguments& InArgs, TWeakObjectPtr<UAbstractSk
 		Menu.AddMenuEntry(
 			LOCTEXT("UnbindAttribute_Label", "Unbind Attribute"),
 			LOCTEXT("UnbindAttribute_Description", "Unbind the selected attribute."),
-			FSlateIcon(),
+			FSlateIcon(FAppStyle::GetAppStyleSetName(), "Icons.Unlink"),
 			UnbindAttribute);
 
 		return true;
