Add per-bone blend support to UAF value runtime

#Synced-CL 48913404

diff --git a/Engine/Plugins/Experimental/UAF/UAF/Source/UAF/Private/UAF/Attributes/Transformers/Accumulate.cpp b/Engine/Plugins/Experimental/UAF/UAF/Source/UAF/Private/UAF/Attributes/Transformers/Accumulate.cpp
--- a/Engine/Plugins/Experimental/UAF/UAF/Source/UAF/Private/UAF/Attributes/Transformers/Accumulate.cpp
+++ b/Engine/Plugins/Experimental/UAF/UAF/Source/UAF/Private/UAF/Attributes/Transformers/Accumulate.cpp
@@ -8,19 +8,31 @@ namespace UE::UAF::Transformers
 {
 	FName FAccumulate::TransformerName = TEXT("UAF::Transformers::Accumulate");
 
-	void FAccumulate::Apply(const FAttributeValueTransformerMapPtr& TransformerMap, const FAttributeCollection& CollectionA, const FAttributeCollection& CollectionB, float Alpha, FAttributeCollection& OutputCollection)
+	void FAccumulate::Apply(const FAttributeValueTransformerMapPtr& TransformerMap, const FAttributeCollection& InputA, const FAttributeCollection& InputB, float WeightB, FAttributeCollection& Output)
 	{
 		// Named sets must match as it ensures that our inputs have the same sizes/shapes
-		checkf(CollectionA.GetNamedSet() == CollectionB.GetNamedSet(), TEXT("Both input collections must use the same named set"));
-		checkf(CollectionA.GetValueSpace() == CollectionB.GetValueSpace(), TEXT("Both input collections must have matching value spaces: [%s] != [%s]"), *CollectionA.GetValueSpace().ToString(), *CollectionB.GetValueSpace().ToString());
+		checkf(InputA.GetNamedSet() == InputB.GetNamedSet(), TEXT("Both input collections must use the same named set"));
+		checkf(InputA.GetValueSpace() == InputB.GetValueSpace(), TEXT("Both input collections must have matching value spaces: [%s] != [%s]"), *InputA.GetValueSpace().ToString(), *InputB.GetValueSpace().ToString());
 
-		const bool bIsOutputA = &CollectionA == &OutputCollection;
-		if (bIsOutputA && !FAnimWeight::IsRelevant(Alpha))
+		const bool bIsOutputA = &InputA == &Output;
+		if (bIsOutputA && !FAnimWeight::IsRelevant(WeightB))
 		{
 			// Our input A is our output and B is not relevant, nothing to do
 			return;
 		}
 
+		// Allocator is irrelevant since it will remain empty
+		FAttributeCollectionHeap EmptyCollection(InputA.GetNamedSet());
+		FAccumulate::Apply(TransformerMap, InputA, InputB, EmptyCollection, WeightB, Output);
+	}
+
+	void FAccumulate::Apply(const FAttributeValueTransformerMapPtr& TransformerMap, const FAttributeCollection& InputA, const FAttributeCollection& InputB, const FAttributeCollection& PerValueWeightsB, float DefaultWeightB, FAttributeCollection& Output)
+	{
+		// Named sets must match as it ensures that our inputs have the same sizes/shapes
+		checkf(InputA.GetNamedSet() == InputB.GetNamedSet(), TEXT("Both input collections must use the same named set"));
+		checkf(InputA.GetNamedSet() == PerValueWeightsB.GetNamedSet(), TEXT("Input and per-value weight collections must use the same named set"));
+		checkf(InputA.GetValueSpace() == InputB.GetValueSpace(), TEXT("Both input collections must have matching value spaces: [%s] != [%s]"), *InputA.GetValueSpace().ToString(), *InputB.GetValueSpace().ToString());
+
 		const FAttributeValueTransformerList* TransformerList = TransformerMap->Find(TransformerName);
 		if (TransformerList == nullptr)
 		{
@@ -28,11 +40,14 @@ namespace UE::UAF::Transformers
 			return;
 		}
 
-		const bool bIsOutputB = &CollectionB == &OutputCollection;
+		UScriptStruct* FloatValueType = FFloatAnimationAttribute::StaticStruct();
+
+		const bool bIsOutputA = &InputA == &Output;
+		const bool bIsOutputB = &InputB == &Output;
 		if (bIsOutputA || bIsOutputB)
 		{
 			// Our output collection matches one of our inputs, we'll modify it in place
-			FReallocFun OutputAllocator = OutputCollection.GetAllocator();
+			FReallocFun OutputAllocator = Output.GetAllocator();
 
 			// Set maps
 			{
@@ -41,7 +56,7 @@ namespace UE::UAF::Transformers
 
 				OuterJoinBy(
 					FAttributeValueTransformerListWithAttributeCollectionSetMapJoinOp(),
-					[Alpha, bIsOutputA, OutputAllocator, &PendingSetMaps](FRawTransformerFunc Transformer, FAttributeSetMap* SetMapA, FAttributeSetMap* SetMapB)
+					[&PerValueWeightsB, DefaultWeightB, bIsOutputA, FloatValueType, OutputAllocator, &PendingSetMaps](FRawTransformerFunc Transformer, FAttributeSetMap* SetMapA, FAttributeSetMap* SetMapB)
 					{
 						if (SetMapA == nullptr)
 						{
@@ -55,8 +70,7 @@ namespace UE::UAF::Transformers
 								if (bIsOutputA)
 								{
 									// A is our output, duplicate it
-									FAttributeSetMap* SetMapCopy = SetMapB->Duplicate(OutputAllocator);
-									PendingSetMaps.Add(SetMapCopy);
+									PendingSetMaps.Add(SetMapB->Duplicate(OutputAllocator));
 								}
 								else
 								{
@@ -76,8 +90,7 @@ namespace UE::UAF::Transformers
 								else
 								{
 									// B is our output, duplicate it
-									FAttributeSetMap* SetMapCopy = SetMapA->Duplicate(OutputAllocator);
-									PendingSetMaps.Add(SetMapCopy);
+									PendingSetMaps.Add(SetMapA->Duplicate(OutputAllocator));
 								}
 							}
 							else
@@ -85,7 +98,7 @@ namespace UE::UAF::Transformers
 								if (Transformer == nullptr)
 								{
 									// No transformer was found for this value type, retain the one with the largest weight
-									if (Alpha >= 0.5f)
+									if (DefaultWeightB >= 0.5f)
 									{
 										// B has more weight
 										if (bIsOutputA)
@@ -117,21 +130,25 @@ namespace UE::UAF::Transformers
 									// Transform our matching pair in place
 									FAttributeSetMap* OutputSetMap = bIsOutputA ? SetMapA : SetMapB;
 
+									// We cannot join our weights because their value type differs
+									const FAttributeMappingKey MappingKey = SetMapA->GetMappingKey().To(FloatValueType);
+									const FAttributeSetMap* PerValueWeightMapB = PerValueWeightsB.FindSetMap(MappingKey);
+
 									FTransformSetMapFunc TypedTransformer = reinterpret_cast<FTransformSetMapFunc>(Transformer);
-									(*TypedTransformer)(SetMapA, SetMapB, Alpha, OutputSetMap, nullptr);
+									(*TypedTransformer)(SetMapA, SetMapB, PerValueWeightMapB, DefaultWeightB, OutputSetMap, nullptr);
 								}
 							}
 						}
 					},
 					TransformerList->CreateSetMapTransformerIterator(),
 					// Remove the const qualifier since we'll mutate one in place as our output
-					const_cast<FAttributeCollection&>(CollectionA).CreateSetMapIterator(),
-					const_cast<FAttributeCollection&>(CollectionB).CreateSetMapIterator());
+					const_cast<FAttributeCollection&>(InputA).CreateSetMapIterator(),
+					const_cast<FAttributeCollection&>(InputB).CreateSetMapIterator());
 
 				// Append our pending entries
 				for (FAttributeSetMap* SetMap : PendingSetMaps)
 				{
-					OutputCollection.AddSetMap(SetMap);
+					Output.AddSetMap(SetMap);
 				}
 			}
 
@@ -142,7 +159,7 @@ namespace UE::UAF::Transformers
 
 				OuterJoinBy(
 					FAttributeValueTransformerListWithAttributeCollectionMapJoinOp(),
-					[Alpha, bIsOutputA, OutputAllocator, &PendingMaps](FRawTransformerFunc Transformer, FAttributeMap* MapA, FAttributeMap* MapB)
+					[DefaultWeightB, bIsOutputA, OutputAllocator, &PendingMaps](FRawTransformerFunc Transformer, FAttributeMap* MapA, FAttributeMap* MapB)
 					{
 						if (MapA == nullptr)
 						{
@@ -156,8 +173,7 @@ namespace UE::UAF::Transformers
 								if (bIsOutputA)
 								{
 									// A is our output, duplicate it
-									FAttributeMap* MapCopy = MapB->Duplicate(OutputAllocator);
-									PendingMaps.Add(MapCopy);
+									PendingMaps.Add(MapB->Duplicate(OutputAllocator));
 								}
 								else
 								{
@@ -177,8 +193,7 @@ namespace UE::UAF::Transformers
 								else
 								{
 									// B is our output, duplicate it
-									FAttributeMap* MapCopy = MapA->Duplicate(OutputAllocator);
-									PendingMaps.Add(MapCopy);
+									PendingMaps.Add(MapA->Duplicate(OutputAllocator));
 								}
 							}
 							else
@@ -186,7 +201,7 @@ namespace UE::UAF::Transformers
 								if (Transformer == nullptr)
 								{
 									// No transformer was found for this value type, retain the one with the largest weight
-									if (Alpha >= 0.5f)
+									if (DefaultWeightB >= 0.5f)
 									{
 										// B has more weight
 										if (bIsOutputA)
@@ -219,7 +234,7 @@ namespace UE::UAF::Transformers
 									FAttributeMap* OutputMap = nullptr;
 
 									FTransformMapFunc TypedTransformer = reinterpret_cast<FTransformMapFunc>(Transformer);
-									(*TypedTransformer)(MapA, MapB, Alpha, OutputMap, OutputAllocator);
+									(*TypedTransformer)(MapA, MapB, DefaultWeightB, OutputMap, OutputAllocator);
 
 									checkf(OutputMap != nullptr, TEXT("Transformer function is responsible for allocating an output buffer"));
 									checkf(OutputMap->GetAllocator() == OutputAllocator, TEXT("Transformer function is responsible for allocating using the output allocator"));
@@ -235,26 +250,26 @@ namespace UE::UAF::Transformers
 					},
 					TransformerList->CreateMapTransformerIterator(),
 					// Remove the const qualifier since we'll mutate one in place as our output
-					const_cast<FAttributeCollection&>(CollectionA).CreateMapIterator(),
-					const_cast<FAttributeCollection&>(CollectionB).CreateMapIterator());
+					const_cast<FAttributeCollection&>(InputA).CreateMapIterator(),
+					const_cast<FAttributeCollection&>(InputB).CreateMapIterator());
 
 				// Append our pending entries
 				for (FAttributeMap* Map : PendingMaps)
 				{
-					OutputCollection.AddMap(Map);
+					Output.AddMap(Map);
 				}
 			}
 		}
 		else
 		{
 			// Our output collection isn't either of our inputs, reset it
-			OutputCollection.Reset(CollectionA.GetNamedSet());
-			OutputCollection.SetValueSpace(CollectionA.GetValueSpace());
+			Output.Reset(InputA.GetNamedSet());
+			Output.SetValueSpace(InputA.GetValueSpace());
 
 			// Set maps
 			OuterJoinBy(
 				FAttributeValueTransformerListWithAttributeCollectionSetMapJoinOp(),
-				[Alpha, &OutputCollection](FRawTransformerFunc Transformer, const FAttributeSetMap* SetMapA, const FAttributeSetMap* SetMapB)
+				[&PerValueWeightsB, DefaultWeightB, FloatValueType, &Output](FRawTransformerFunc Transformer, const FAttributeSetMap* SetMapA, const FAttributeSetMap* SetMapB)
 				{
 					if (SetMapA == nullptr)
 					{
@@ -265,8 +280,7 @@ namespace UE::UAF::Transformers
 						else
 						{
 							// This set map doesn't have a matching pair, duplicate it
-							FAttributeSetMap* SetMapCopy = SetMapB->Duplicate(OutputCollection.GetAllocator());
-							OutputCollection.AppendSetMap(SetMapCopy);
+							Output.AppendSetMap(SetMapB->Duplicate(Output.GetAllocator()));
 						}
 					}
 					else
@@ -274,8 +288,7 @@ namespace UE::UAF::Transformers
 						if (SetMapB == nullptr)
 						{
 							// This set map doesn't have a matching pair, duplicate it
-							FAttributeSetMap* SetMapCopy = SetMapA->Duplicate(OutputCollection.GetAllocator());
-							OutputCollection.AppendSetMap(SetMapCopy);
+							Output.AppendSetMap(SetMapA->Duplicate(Output.GetAllocator()));
 						}
 						else
 						{
@@ -283,40 +296,44 @@ namespace UE::UAF::Transformers
 							{
 								// No transformer was found for this value type, retain the one with the largest weight
 								FAttributeSetMap* SetMapCopy;
-								if (Alpha >= 0.5f)
+								if (DefaultWeightB >= 0.5f)
 								{
-									SetMapCopy = SetMapB->Duplicate(OutputCollection.GetAllocator());
+									SetMapCopy = SetMapB->Duplicate(Output.GetAllocator());
 								}
 								else
 								{
-									SetMapCopy = SetMapA->Duplicate(OutputCollection.GetAllocator());
+									SetMapCopy = SetMapA->Duplicate(Output.GetAllocator());
 								}
 
-								OutputCollection.AppendSetMap(SetMapCopy);
+								Output.AppendSetMap(SetMapCopy);
 							}
 							else
 							{
 								// Transform our matching pair
 								FAttributeSetMap* OutputSetMap = nullptr;
 
+								// We cannot join our weights because their value type differs
+								const FAttributeMappingKey MappingKey = SetMapA->GetMappingKey().To(FloatValueType);
+								const FAttributeSetMap* PerValueWeightMapB = PerValueWeightsB.FindSetMap(MappingKey);
+
 								FTransformSetMapFunc TypedTransformer = reinterpret_cast<FTransformSetMapFunc>(Transformer);
-								(*TypedTransformer)(SetMapA, SetMapB, Alpha, OutputSetMap, OutputCollection.GetAllocator());
+								(*TypedTransformer)(SetMapA, SetMapB, PerValueWeightMapB, DefaultWeightB, OutputSetMap, Output.GetAllocator());
 
 								checkf(OutputSetMap != nullptr, TEXT("Transformer function is responsible for allocating an output buffer"));
-								checkf(OutputSetMap->GetAllocator() == OutputCollection.GetAllocator(), TEXT("Transformer function is responsible for allocating using the output allocator"));
-								OutputCollection.AppendSetMap(OutputSetMap);
+								checkf(OutputSetMap->GetAllocator() == Output.GetAllocator(), TEXT("Transformer function is responsible for allocating using the output allocator"));
+								Output.AppendSetMap(OutputSetMap);
 							}
 						}
 					}
 				},
 				TransformerList->CreateSetMapTransformerIterator(),
-				CollectionA.CreateConstSetMapIterator(),
-				CollectionB.CreateConstSetMapIterator());
+				InputA.CreateConstSetMapIterator(),
+				InputB.CreateConstSetMapIterator());
 
 			// Name/value maps
 			OuterJoinBy(
 				FAttributeValueTransformerListWithAttributeCollectionMapJoinOp(),
-				[Alpha, &OutputCollection](FRawTransformerFunc Transformer, const FAttributeMap* MapA, const FAttributeMap* MapB)
+				[DefaultWeightB, &Output](FRawTransformerFunc Transformer, const FAttributeMap* MapA, const FAttributeMap* MapB)
 				{
 					if (MapA == nullptr)
 					{
@@ -327,8 +344,7 @@ namespace UE::UAF::Transformers
 						else
 						{
 							// This map doesn't have a matching pair, duplicate it
-							FAttributeMap* MapCopy = MapB->Duplicate(OutputCollection.GetAllocator());
-							OutputCollection.AppendMap(MapCopy);
+							Output.AppendMap(MapB->Duplicate(Output.GetAllocator()));
 						}
 					}
 					else
@@ -344,16 +360,16 @@ namespace UE::UAF::Transformers
 							{
 								// No transformer was found for this value type, retain the one with the largest weight
 								FAttributeMap* MapCopy;
-								if (Alpha >= 0.5f)
+								if (DefaultWeightB >= 0.5f)
 								{
-									MapCopy = MapB->Duplicate(OutputCollection.GetAllocator());
+									MapCopy = MapB->Duplicate(Output.GetAllocator());
 								}
 								else
 								{
-									MapCopy = MapA->Duplicate(OutputCollection.GetAllocator());
+									MapCopy = MapA->Duplicate(Output.GetAllocator());
 								}
 
-								OutputCollection.AppendMap(MapCopy);
+								Output.AppendMap(MapCopy);
 							}
 							else
 							{
@@ -361,24 +377,24 @@ namespace UE::UAF::Transformers
 								FAttributeMap* OutputMap = nullptr;
 
 								FTransformMapFunc TypedTransformer = reinterpret_cast<FTransformMapFunc>(Transformer);
-								(*TypedTransformer)(MapA, MapB, Alpha, OutputMap, OutputCollection.GetAllocator());
+								(*TypedTransformer)(MapA, MapB, DefaultWeightB, OutputMap, Output.GetAllocator());
 
 								checkf(OutputMap != nullptr, TEXT("Transformer function is responsible for allocating an output buffer"));
-								checkf(OutputMap->GetAllocator() == OutputCollection.GetAllocator(), TEXT("Transformer function is responsible for allocating using the output allocator"));
-								OutputCollection.AppendMap(OutputMap);
+								checkf(OutputMap->GetAllocator() == Output.GetAllocator(), TEXT("Transformer function is responsible for allocating using the output allocator"));
+								Output.AppendMap(OutputMap);
 							}
 						}
 					}
 				},
 				TransformerList->CreateMapTransformerIterator(),
-				CollectionA.CreateConstMapIterator(),
-				CollectionB.CreateConstMapIterator());
+				InputA.CreateConstMapIterator(),
+				InputB.CreateConstMapIterator());
 		}
 	}
 
 	namespace Private
 	{
-		void FAccumulate_BoneTransformAttribute::TransformSetMap(const FAttributeSetMap* SetMapA, const FAttributeSetMap* SetMapB, float Alpha, FAttributeSetMap*& OutputSetMap, FReallocFun OutputAllocator)
+		void FAccumulate_BoneTransformAttribute::TransformSetMap(const FAttributeSetMap* SetMapA, const FAttributeSetMap* SetMapB, const FAttributeSetMap* PerValueWeightMapB, float DefaultWeightB, FAttributeSetMap*& OutputSetMap, FReallocFun OutputAllocator)
 		{
 			check(SetMapA->GetTypedSet() == SetMapB->GetTypedSet());
 			check(SetMapA->Num() == SetMapB->Num());
@@ -401,20 +417,45 @@ namespace UE::UAF::Transformers
 			FBoneTransformAnimationAttribute* OutputTransformPtr = OutputBoneTransforms->GetData();
 			const FBoneTransformAnimationAttribute* TransformPtrAEnd = TransformPtrA + BoneTransformsA->Num();
 
-			const ScalarRegister VAlpha(Alpha);
+			if (const TAttributeSetMap<FFloatAnimationAttribute>* PerValueWeightsB = Cast<FFloatAnimationAttribute>(PerValueWeightMapB))
+			{
+				check(SetMapA->GetTypedSet() == PerValueWeightsB->GetTypedSet());
+				check(SetMapA->Num() == PerValueWeightsB->Num());
+
+				const FFloatAnimationAttribute* PerValueWeightsBPtr = PerValueWeightsB->GetData();
 
-			while (TransformPtrA < TransformPtrAEnd)
+				while (TransformPtrA < TransformPtrAEnd)
+				{
+					FTransform Result = TransformPtrA->Value;
+					Result.AccumulateWithShortestRotation(TransformPtrB->Value, ScalarRegister(PerValueWeightsBPtr->Value));
+
+					OutputTransformPtr->Value = Result;
+
+					TransformPtrA++;
+					TransformPtrB++;
+					PerValueWeightsBPtr++;
+					OutputTransformPtr++;
+				}
+			}
+			else
 			{
-				OutputTransformPtr->Value = TransformPtrA->Value;
-				OutputTransformPtr->Value.AccumulateWithShortestRotation(TransformPtrB->Value, VAlpha);
+				const ScalarRegister VDefaultWeightB(DefaultWeightB);
+
+				while (TransformPtrA < TransformPtrAEnd)
+				{
+					FTransform Result = TransformPtrA->Value;
+					Result.AccumulateWithShortestRotation(TransformPtrB->Value, VDefaultWeightB);
+
+					OutputTransformPtr->Value = Result;
 
-				TransformPtrA++;
-				TransformPtrB++;
-				OutputTransformPtr++;
+					TransformPtrA++;
+					TransformPtrB++;
+					OutputTransformPtr++;
+				}
 			}
 		}
 
-		void FAccumulate_BoneTransformAttribute::TransformMap(const FAttributeMap* MapA, const FAttributeMap* MapB, float Alpha, FAttributeMap*& OutputMap, FReallocFun OutputAllocator)
+		void FAccumulate_BoneTransformAttribute::TransformMap(const FAttributeMap* MapA, const FAttributeMap* MapB, float WeightB, FAttributeMap*& OutputMap, FReallocFun OutputAllocator)
 		{
 			check(OutputMap == nullptr);
 
@@ -422,20 +463,20 @@ namespace UE::UAF::Transformers
 			const TAttributeMap<FBoneTransformAnimationAttribute>* BoneTransformsB = CastChecked<FBoneTransformAnimationAttribute>(MapB);
 			TAttributeMap<FBoneTransformAnimationAttribute>* OutputBoneTransforms = MakeAttributeMap<FBoneTransformAnimationAttribute>(OutputAllocator);
 
-			const ScalarRegister VAlpha(Alpha);
+			const ScalarRegister VWeightB(WeightB);
 
 			OuterJoinBy(
 				FAttributeMapJoinOp(),
-				[&VAlpha, OutputBoneTransforms](FName Name, const FBoneTransformAnimationAttribute* A, const FBoneTransformAnimationAttribute* B)
+				[VWeightB, OutputBoneTransforms](FName Name, const FBoneTransformAnimationAttribute* InputA, const FBoneTransformAnimationAttribute* InputB)
 				{
 					// Missing values are set to the identity
-					const FTransform& ValueA = A != nullptr ? A->Value : FTransform::Identity;
-					const FTransform& ValueB = B != nullptr ? B->Value : FTransform::Identity;
+					const FTransform& ValueA = InputA != nullptr ? InputA->Value : FTransform::Identity;
+					const FTransform& ValueB = InputB != nullptr ? InputB->Value : FTransform::Identity;
 
-					FBoneTransformAnimationAttribute Output{ ValueA };
-					Output.Value.AccumulateWithShortestRotation(ValueB, VAlpha);
+					FTransform Result = ValueA;
+					Result.AccumulateWithShortestRotation(ValueB, VWeightB);
 
-					OutputBoneTransforms->Append(Name, Output);
+					OutputBoneTransforms->Append(Name, FBoneTransformAnimationAttribute{ Result });
 				},
 				BoneTransformsA->CreateConstIterator(),
 				BoneTransformsB->CreateConstIterator());
@@ -443,7 +484,7 @@ namespace UE::UAF::Transformers
 			OutputMap = OutputBoneTransforms;
 		}
 
-		void FAccumulate_FloatAttribute::TransformSetMap(const FAttributeSetMap* SetMapA, const FAttributeSetMap* SetMapB, float Alpha, FAttributeSetMap*& OutputSetMap, FReallocFun OutputAllocator)
+		void FAccumulate_FloatAttribute::TransformSetMap(const FAttributeSetMap* SetMapA, const FAttributeSetMap* SetMapB, const FAttributeSetMap* PerValueWeightMapB, float DefaultWeightB, FAttributeSetMap*& OutputSetMap, FReallocFun OutputAllocator)
 		{
 			check(SetMapA->GetTypedSet() == SetMapB->GetTypedSet());
 			check(SetMapA->Num() == SetMapB->Num());
@@ -466,17 +507,37 @@ namespace UE::UAF::Transformers
 			FFloatAnimationAttribute* OutputFloatsPtr = OutputFloats->GetData();
 			const FFloatAnimationAttribute* FloatsPtrAEnd = FloatsPtrA + FloatsA->Num();
 
-			while (FloatsPtrA < FloatsPtrAEnd)
+			if (const TAttributeSetMap<FFloatAnimationAttribute>* PerValueWeightsB = Cast<FFloatAnimationAttribute>(PerValueWeightMapB))
+			{
+				check(SetMapA->GetTypedSet() == PerValueWeightsB->GetTypedSet());
+				check(SetMapA->Num() == PerValueWeightsB->Num());
+
+				const FFloatAnimationAttribute* PerValueWeightsBPtr = PerValueWeightsB->GetData();
+
+				while (FloatsPtrA < FloatsPtrAEnd)
+				{
+					OutputFloatsPtr->Value = FloatsPtrA->Value + (FloatsPtrB->Value * PerValueWeightsBPtr->Value);
+
+					FloatsPtrA++;
+					FloatsPtrB++;
+					PerValueWeightsBPtr++;
+					OutputFloatsPtr++;
+				}
+			}
+			else
 			{
-				OutputFloatsPtr->Value = FloatsPtrA->Value + (FloatsPtrB->Value * Alpha);
+				while (FloatsPtrA < FloatsPtrAEnd)
+				{
+					OutputFloatsPtr->Value = FloatsPtrA->Value + (FloatsPtrB->Value * DefaultWeightB);
 
-				FloatsPtrA++;
-				FloatsPtrB++;
-				OutputFloatsPtr++;
+					FloatsPtrA++;
+					FloatsPtrB++;
+					OutputFloatsPtr++;
+				}
 			}
 		}
 
-		void FAccumulate_FloatAttribute::TransformMap(const FAttributeMap* MapA, const FAttributeMap* MapB, float Alpha, FAttributeMap*& OutputMap, FReallocFun OutputAllocator)
+		void FAccumulate_FloatAttribute::TransformMap(const FAttributeMap* MapA, const FAttributeMap* MapB, float WeightB, FAttributeMap*& OutputMap, FReallocFun OutputAllocator)
 		{
 			check(OutputMap == nullptr);
 
@@ -486,14 +547,14 @@ namespace UE::UAF::Transformers
 
 			OuterJoinBy(
 				FAttributeMapJoinOp(),
-				[&Alpha, OutputFloats](FName Name, const FFloatAnimationAttribute* A, const FFloatAnimationAttribute* B)
+				[&WeightB, OutputFloats](FName Name, const FFloatAnimationAttribute* InputA, const FFloatAnimationAttribute* InputB)
 				{
 					// Missing values are set to 0.0
-					const float ValueA = A != nullptr ? A->Value : 0.0f;
-					const float ValueB = B != nullptr ? B->Value : 0.0f;
+					const float ValueA = InputA != nullptr ? InputA->Value : 0.0f;
+					const float ValueB = InputB != nullptr ? InputB->Value : 0.0f;
 
 					FFloatAnimationAttribute Output;
-					Output.Value = ValueA + (ValueB * Alpha);
+					Output.Value = ValueA + (ValueB * WeightB);
 
 					OutputFloats->Append(Name, Output);
 				},
diff --git a/Engine/Plugins/Experimental/UAF/UAF/Source/UAF/Private/UAF/Attributes/Transformers/Overwrite.cpp b/Engine/Plugins/Experimental/UAF/UAF/Source/UAF/Private/UAF/Attributes/Transformers/Overwrite.cpp
--- a/Engine/Plugins/Experimental/UAF/UAF/Source/UAF/Private/UAF/Attributes/Transformers/Overwrite.cpp
+++ b/Engine/Plugins/Experimental/UAF/UAF/Source/UAF/Private/UAF/Attributes/Transformers/Overwrite.cpp
@@ -8,8 +8,24 @@ namespace UE::UAF::Transformers
 {
 	FName FOverwrite::TransformerName = TEXT("UAF::Transformers::Overwrite");
 
-	void FOverwrite::Apply(const FAttributeValueTransformerMapPtr& TransformerMap, const FAttributeCollection& SourceCollection, float Alpha, FAttributeCollection& OutputCollection)
+	void FOverwrite::Apply(const FAttributeValueTransformerMapPtr& TransformerMap, const FAttributeCollection& Input, float Weight, FAttributeCollection& Output)
 	{
+		const bool bIsInputOutput = &Input == &Output;
+		if (bIsInputOutput && FAnimWeight::IsFullWeight(Weight))
+		{
+			// Our input is our output and it has full weight, nothing to do
+			return;
+		}
+		
+		// Allocator is irrelevant since it will remain empty
+		FAttributeCollectionHeap EmptyCollection(Input.GetNamedSet());
+		FOverwrite::Apply(TransformerMap, Input, EmptyCollection, Weight, Output);
+	}
+
+	void FOverwrite::Apply(const FAttributeValueTransformerMapPtr& TransformerMap, const FAttributeCollection& Input, const FAttributeCollection& PerValueWeights, float DefaultWeight, FAttributeCollection& Output)
+	{
+		checkf(Input.GetNamedSet() == PerValueWeights.GetNamedSet(), TEXT("Input and per-value weight collections must use the same named set"));
+
 		const FAttributeValueTransformerList* TransformerList = TransformerMap->Find(TransformerName);
 		if (TransformerList == nullptr)
 		{
@@ -17,24 +33,21 @@ namespace UE::UAF::Transformers
 			return;
 		}
 
-		if (&SourceCollection == &OutputCollection)
-		{
-			if (FAnimWeight::IsFullWeight(Alpha))
-			{
-				// Our input is our output and it has full weight, nothing to do
-				return;
-			}
+		UScriptStruct* FloatValueType = FFloatAnimationAttribute::StaticStruct();
 
+		const bool bIsInputOutput = &Input == &Output;
+		if (bIsInputOutput)
+		{
 			// Our output collection matches one of our inputs, we'll modify it in place
-			FReallocFun OutputAllocator = OutputCollection.GetAllocator();
+			FReallocFun OutputAllocator = Output.GetAllocator();
 
 			// Set maps
 			{
 				OuterJoinBy(
 					FAttributeValueTransformerListWithAttributeCollectionSetMapJoinOp(),
-					[Alpha, OutputAllocator](FRawTransformerFunc Transformer, FAttributeSetMap* SourceSetMap)
+					[&PerValueWeights, DefaultWeight, FloatValueType, OutputAllocator](FRawTransformerFunc Transformer, FAttributeSetMap* InputSetMap)
 					{
-						if (SourceSetMap == nullptr)
+						if (InputSetMap == nullptr)
 						{
 							// No maps found for the transformer type
 						}
@@ -47,25 +60,29 @@ namespace UE::UAF::Transformers
 							else
 							{
 								// Transform our input in place
-								FAttributeSetMap* OutputSetMap = SourceSetMap;
+								FAttributeSetMap* OutputSetMap = InputSetMap;
+
+								// We cannot join our weights because their value type differs
+								const FAttributeMappingKey MappingKey = InputSetMap->GetMappingKey().To(FloatValueType);
+								const FAttributeSetMap* PerValueWeightMap = PerValueWeights.FindSetMap(MappingKey);
 
 								FTransformSetMapFunc TypedTransformer = reinterpret_cast<FTransformSetMapFunc>(Transformer);
-								(*TypedTransformer)(SourceSetMap, Alpha, OutputSetMap, nullptr);
+								(*TypedTransformer)(InputSetMap, PerValueWeightMap, DefaultWeight, OutputSetMap, nullptr);
 							}
 						}
 					},
 					TransformerList->CreateSetMapTransformerIterator(),
 					// Our input is our output
-					OutputCollection.CreateSetMapIterator());
+					Output.CreateSetMapIterator());
 			}
 
 			// Name/value maps
 			{
 				OuterJoinBy(
 					FAttributeValueTransformerListWithAttributeCollectionMapJoinOp(),
-					[Alpha, OutputAllocator](FRawTransformerFunc Transformer, FAttributeMap* SourceMap)
+					[DefaultWeight, OutputAllocator](FRawTransformerFunc Transformer, FAttributeMap* InputMap)
 					{
-						if (SourceMap == nullptr)
+						if (InputMap == nullptr)
 						{
 							// No maps found for the transformer type
 						}
@@ -78,30 +95,30 @@ namespace UE::UAF::Transformers
 							else
 							{
 								// Transform our input in place
-								FAttributeMap* OutputMap = SourceMap;
+								FAttributeMap* OutputMap = InputMap;
 
 								FTransformMapFunc TypedTransformer = reinterpret_cast<FTransformMapFunc>(Transformer);
-								(*TypedTransformer)(SourceMap, Alpha, OutputMap, nullptr);
+								(*TypedTransformer)(InputMap, DefaultWeight, OutputMap, nullptr);
 							}
 						}
 					},
 					TransformerList->CreateMapTransformerIterator(),
 					// Our input is our output
-					OutputCollection.CreateMapIterator());
+					Output.CreateMapIterator());
 			}
 		}
 		else
 		{
 			// Our output collection isn't either of our inputs, reset it
-			OutputCollection.Reset(SourceCollection.GetNamedSet());
-			OutputCollection.SetValueSpace(SourceCollection.GetValueSpace());
+			Output.Reset(Input.GetNamedSet());
+			Output.SetValueSpace(Input.GetValueSpace());
 
 			// Set maps
 			OuterJoinBy(
 				FAttributeValueTransformerListWithAttributeCollectionSetMapJoinOp(),
-				[Alpha, &OutputCollection](FRawTransformerFunc Transformer, const FAttributeSetMap* SetMapA)
+				[&PerValueWeights, DefaultWeight, FloatValueType, &Output](FRawTransformerFunc Transformer, const FAttributeSetMap* InputSetMap)
 				{
-					if (SetMapA == nullptr)
+					if (InputSetMap == nullptr)
 					{
 						// No maps found for the transformer type
 					}
@@ -110,32 +127,35 @@ namespace UE::UAF::Transformers
 						if (Transformer == nullptr)
 						{
 							// No transformer was found for this value type, retain it as-is
-							FAttributeSetMap* SetMapCopy = SetMapA->Duplicate(OutputCollection.GetAllocator());
-							OutputCollection.AppendSetMap(SetMapCopy);
+							Output.AppendSetMap(InputSetMap->Duplicate(Output.GetAllocator()));
 						}
 						else
 						{
 							// Transform our matching pair
 							FAttributeSetMap* OutputSetMap = nullptr;
 
+							// We cannot join our weights because their value type differs
+							const FAttributeMappingKey MappingKey = InputSetMap->GetMappingKey().To(FloatValueType);
+							const FAttributeSetMap* PerValueWeightMap = PerValueWeights.FindSetMap(MappingKey);
+
 							FTransformSetMapFunc TypedTransformer = reinterpret_cast<FTransformSetMapFunc>(Transformer);
-							(*TypedTransformer)(SetMapA, Alpha, OutputSetMap, OutputCollection.GetAllocator());
+							(*TypedTransformer)(InputSetMap, PerValueWeightMap, DefaultWeight, OutputSetMap, Output.GetAllocator());
 
 							checkf(OutputSetMap != nullptr, TEXT("Transformer function is responsible for allocating an output buffer"));
-							checkf(OutputSetMap->GetAllocator() == OutputCollection.GetAllocator(), TEXT("Transformer function is responsible for allocating using the output allocator"));
-							OutputCollection.AppendSetMap(OutputSetMap);
+							checkf(OutputSetMap->GetAllocator() == Output.GetAllocator(), TEXT("Transformer function is responsible for allocating using the output allocator"));
+							Output.AppendSetMap(OutputSetMap);
 						}
 					}
 				},
 				TransformerList->CreateSetMapTransformerIterator(),
-				SourceCollection.CreateConstSetMapIterator());
+				Input.CreateConstSetMapIterator());
 
 			// Name/value maps
 			OuterJoinBy(
 				FAttributeValueTransformerListWithAttributeCollectionMapJoinOp(),
-				[Alpha, &OutputCollection](FRawTransformerFunc Transformer, const FAttributeMap* MapA)
+				[DefaultWeight, &Output](FRawTransformerFunc Transformer, const FAttributeMap* InputMap)
 				{
-					if (MapA == nullptr)
+					if (InputMap == nullptr)
 					{
 						// No maps found for the transformer type
 					}
@@ -144,8 +164,7 @@ namespace UE::UAF::Transformers
 						if (Transformer == nullptr)
 						{
 							// No transformer was found for this value type, retain it as-is
-							FAttributeMap* MapCopy = MapA->Duplicate(OutputCollection.GetAllocator());
-							OutputCollection.AppendMap(MapCopy);
+							Output.AppendMap(InputMap->Duplicate(Output.GetAllocator()));
 						}
 						else
 						{
@@ -153,51 +172,70 @@ namespace UE::UAF::Transformers
 							FAttributeMap* OutputMap = nullptr;
 
 							FTransformMapFunc TypedTransformer = reinterpret_cast<FTransformMapFunc>(Transformer);
-							(*TypedTransformer)(MapA, Alpha, OutputMap, OutputCollection.GetAllocator());
+							(*TypedTransformer)(InputMap, DefaultWeight, OutputMap, Output.GetAllocator());
 
 							checkf(OutputMap != nullptr, TEXT("Transformer function is responsible for allocating an output buffer"));
-							checkf(OutputMap->GetAllocator() == OutputCollection.GetAllocator(), TEXT("Transformer function is responsible for allocating using the output allocator"));
-							OutputCollection.AppendMap(OutputMap);
+							checkf(OutputMap->GetAllocator() == Output.GetAllocator(), TEXT("Transformer function is responsible for allocating using the output allocator"));
+							Output.AppendMap(OutputMap);
 						}
 					}
 				},
 				TransformerList->CreateMapTransformerIterator(),
-				SourceCollection.CreateConstMapIterator());
+				Input.CreateConstMapIterator());
 		}
 	}
 
 	namespace Private
 	{
-		void FOverwrite_BoneTransformAttribute::TransformSetMap(const FAttributeSetMap* SetMap, float Alpha, FAttributeSetMap*& OutputSetMap, FReallocFun OutputAllocator)
+		void FOverwrite_BoneTransformAttribute::TransformSetMap(const FAttributeSetMap* InputSetMap, const FAttributeSetMap* PerValueWeightMap, float DefaultWeight, FAttributeSetMap*& OutputSetMap, FReallocFun OutputAllocator)
 		{
 			if (OutputSetMap == nullptr)
 			{
-				const FAttributeSetMap::FConstructArgs Args(SetMap->GetTypedSet(), SetMap->GetValueType(), OutputAllocator);
+				const FAttributeSetMap::FConstructArgs Args(InputSetMap->GetTypedSet(), InputSetMap->GetValueType(), OutputAllocator);
 				OutputSetMap = MakeAttributeSetMap<FBoneTransformAnimationAttribute>(Args);
 			}
 
-			check(SetMap->GetTypedSet() == OutputSetMap->GetTypedSet());
-			check(SetMap->Num() == OutputSetMap->Num());
+			check(InputSetMap->GetTypedSet() == OutputSetMap->GetTypedSet());
+			check(InputSetMap->Num() == OutputSetMap->Num());
 
-			const TAttributeSetMap<FBoneTransformAnimationAttribute>* InputBoneTransforms = CastChecked<FBoneTransformAnimationAttribute>(SetMap);
+			const TAttributeSetMap<FBoneTransformAnimationAttribute>* InputBoneTransforms = CastChecked<FBoneTransformAnimationAttribute>(InputSetMap);
 			TAttributeSetMap<FBoneTransformAnimationAttribute>* OutputBoneTransforms = CastChecked<FBoneTransformAnimationAttribute>(OutputSetMap);
 
 			const FBoneTransformAnimationAttribute* InputTransformPtr = InputBoneTransforms->GetData();
 			FBoneTransformAnimationAttribute* OutputTransformPtr = OutputBoneTransforms->GetData();
 			const FBoneTransformAnimationAttribute* InputTransformPtrEnd = InputTransformPtr + InputBoneTransforms->Num();
 
-			const ScalarRegister VAlpha(Alpha);
+			if (const TAttributeSetMap<FFloatAnimationAttribute>* PerValueWeights = Cast<FFloatAnimationAttribute>(PerValueWeightMap))
+			{
+				check(InputSetMap->GetTypedSet() == PerValueWeights->GetTypedSet());
+				check(InputSetMap->Num() == PerValueWeights->Num());
 
-			while (InputTransformPtr < InputTransformPtrEnd)
+				const FFloatAnimationAttribute* PerValueWeightsPtr = PerValueWeights->GetData();
+
+				while (InputTransformPtr < InputTransformPtrEnd)
+				{
+					OutputTransformPtr->Value = InputTransformPtr->Value * ScalarRegister(PerValueWeightsPtr->Value);
+
+					InputTransformPtr++;
+					PerValueWeightsPtr++;
+					OutputTransformPtr++;
+				}
+			}
+			else
 			{
-				OutputTransformPtr->Value = InputTransformPtr->Value * VAlpha;
+				const ScalarRegister VDefaultWeight(DefaultWeight);
 
-				InputTransformPtr++;
-				OutputTransformPtr++;
+				while (InputTransformPtr < InputTransformPtrEnd)
+				{
+					OutputTransformPtr->Value = InputTransformPtr->Value * VDefaultWeight;
+
+					InputTransformPtr++;
+					OutputTransformPtr++;
+				}
 			}
 		}
 
-		void FOverwrite_BoneTransformAttribute::TransformMap(const FAttributeMap* Map, float Alpha, FAttributeMap*& OutputMap, FReallocFun OutputAllocator)
+		void FOverwrite_BoneTransformAttribute::TransformMap(const FAttributeMap* InputMap, float Weight, FAttributeMap*& OutputMap, FReallocFun OutputAllocator)
 		{
 			const bool bAllocateOutput = OutputMap == nullptr;
 			if (bAllocateOutput)
@@ -205,17 +243,17 @@ namespace UE::UAF::Transformers
 				OutputMap = MakeAttributeMap<FBoneTransformAnimationAttribute>(OutputAllocator);
 			}
 
-			const TAttributeMap<FBoneTransformAnimationAttribute>* InputBoneTransforms = CastChecked<FBoneTransformAnimationAttribute>(Map);
+			const TAttributeMap<FBoneTransformAnimationAttribute>* InputBoneTransforms = CastChecked<FBoneTransformAnimationAttribute>(InputMap);
 			TAttributeMap<FBoneTransformAnimationAttribute>* OutputBoneTransforms = CastChecked<FBoneTransformAnimationAttribute>(OutputMap);
 
-			const ScalarRegister VAlpha(Alpha);
+			const ScalarRegister VWeight(Weight);
 
 			if (bAllocateOutput)
 			{
 				// We allocated a new output, append to it
 				for (auto It = InputBoneTransforms->CreateConstIterator(); It; ++It)
 				{
-					FBoneTransformAnimationAttribute Output{ It.GetValue()->Value * VAlpha };
+					FBoneTransformAnimationAttribute Output{ It.GetValue()->Value * VWeight };
 					OutputBoneTransforms->Append(It.GetName(), Output);
 				}
 			}
@@ -223,41 +261,61 @@ namespace UE::UAF::Transformers
 			{
 				// We overwrite in place
 				check(InputBoneTransforms == OutputBoneTransforms);
+
 				for (auto It = OutputBoneTransforms->CreateIterator(); It; ++It)
 				{
-					It.GetValue()->Value *= VAlpha;
+					It.GetValue()->Value *= VWeight;
 				}
 			}
 		}
 
-		void FOverwrite_FloatAttribute::TransformSetMap(const FAttributeSetMap* SetMap, float Alpha, FAttributeSetMap*& OutputSetMap, FReallocFun OutputAllocator)
+		void FOverwrite_FloatAttribute::TransformSetMap(const FAttributeSetMap* InputSetMap, const FAttributeSetMap* PerValueWeightMap, float DefaultWeight, FAttributeSetMap*& OutputSetMap, FReallocFun OutputAllocator)
 		{
 			if (OutputSetMap == nullptr)
 			{
-				const FAttributeSetMap::FConstructArgs Args(SetMap->GetTypedSet(), SetMap->GetValueType(), OutputAllocator);
+				const FAttributeSetMap::FConstructArgs Args(InputSetMap->GetTypedSet(), InputSetMap->GetValueType(), OutputAllocator);
 				OutputSetMap = MakeAttributeSetMap<FFloatAnimationAttribute>(Args);
 			}
 
-			check(SetMap->GetTypedSet() == OutputSetMap->GetTypedSet());
-			check(SetMap->Num() == OutputSetMap->Num());
+			check(InputSetMap->GetTypedSet() == OutputSetMap->GetTypedSet());
+			check(InputSetMap->Num() == OutputSetMap->Num());
 
-			const TAttributeSetMap<FFloatAnimationAttribute>* InputFloats = CastChecked<FFloatAnimationAttribute>(SetMap);
+			const TAttributeSetMap<FFloatAnimationAttribute>* InputFloats = CastChecked<FFloatAnimationAttribute>(InputSetMap);
 			TAttributeSetMap<FFloatAnimationAttribute>* OutputFloats = CastChecked<FFloatAnimationAttribute>(OutputSetMap);
 
 			const FFloatAnimationAttribute* FloatsPtr = InputFloats->GetData();
 			FFloatAnimationAttribute* OutputFloatsPtr = OutputFloats->GetData();
 			const FFloatAnimationAttribute* FloatsPtrEnd = FloatsPtr + InputFloats->Num();
 
-			while (FloatsPtr < FloatsPtrEnd)
+			if (const TAttributeSetMap<FFloatAnimationAttribute>* PerValueWeights = Cast<FFloatAnimationAttribute>(PerValueWeightMap))
 			{
-				OutputFloatsPtr->Value = FloatsPtr->Value * Alpha;
+				check(InputSetMap->GetTypedSet() == PerValueWeights->GetTypedSet());
+				check(InputSetMap->Num() == PerValueWeights->Num());
+
+				const FFloatAnimationAttribute* PerValueWeightsPtr = PerValueWeights->GetData();
+
+				while (FloatsPtr < FloatsPtrEnd)
+				{
+					OutputFloatsPtr->Value = FloatsPtr->Value * PerValueWeightsPtr->Value;
 
-				FloatsPtr++;
-				OutputFloatsPtr++;
+					FloatsPtr++;
+					PerValueWeightsPtr++;
+					OutputFloatsPtr++;
+				}
+			}
+			else
+			{
+				while (FloatsPtr < FloatsPtrEnd)
+				{
+					OutputFloatsPtr->Value = FloatsPtr->Value * DefaultWeight;
+
+					FloatsPtr++;
+					OutputFloatsPtr++;
+				}
 			}
 		}
 
-		void FOverwrite_FloatAttribute::TransformMap(const FAttributeMap* Map, float Alpha, FAttributeMap*& OutputMap, FReallocFun OutputAllocator)
+		void FOverwrite_FloatAttribute::TransformMap(const FAttributeMap* InputMap, float Weight, FAttributeMap*& OutputMap, FReallocFun OutputAllocator)
 		{
 			const bool bAllocateOutput = OutputMap == nullptr;
 			if (bAllocateOutput)
@@ -265,7 +323,7 @@ namespace UE::UAF::Transformers
 				OutputMap = MakeAttributeMap<FFloatAnimationAttribute>(OutputAllocator);
 			}
 
-			const TAttributeMap<FFloatAnimationAttribute>* InputFloats = CastChecked<FFloatAnimationAttribute>(Map);
+			const TAttributeMap<FFloatAnimationAttribute>* InputFloats = CastChecked<FFloatAnimationAttribute>(InputMap);
 			TAttributeMap<FFloatAnimationAttribute>* OutputFloats = CastChecked<FFloatAnimationAttribute>(OutputMap);
 
 			if (bAllocateOutput)
@@ -273,7 +331,7 @@ namespace UE::UAF::Transformers
 				// We allocated a new output, append to it
 				for (auto It = InputFloats->CreateConstIterator(); It; ++It)
 				{
-					FFloatAnimationAttribute Output{ It.GetValue()->Value * Alpha };
+					FFloatAnimationAttribute Output{ It.GetValue()->Value * Weight };
 					OutputFloats->Append(It.GetName(), Output);
 				}
 			}
@@ -281,9 +339,10 @@ namespace UE::UAF::Transformers
 			{
 				// We overwrite in place
 				check(InputFloats == OutputFloats);
+
 				for (auto It = OutputFloats->CreateIterator(); It; ++It)
 				{
-					It.GetValue()->Value *= Alpha;
+					It.GetValue()->Value *= Weight;
 				}
 			}
 		}
diff --git a/Engine/Plugins/Experimental/UAF/UAF/Source/UAF/Public/UAF/Attributes/AttributeBindingIndex.h b/Engine/Plugins/Experimental/UAF/UAF/Source/UAF/Public/UAF/Attributes/AttributeBindingIndex.h
--- a/Engine/Plugins/Experimental/UAF/UAF/Source/UAF/Public/UAF/Attributes/AttributeBindingIndex.h
+++ b/Engine/Plugins/Experimental/UAF/UAF/Source/UAF/Public/UAF/Attributes/AttributeBindingIndex.h
@@ -31,6 +31,13 @@ namespace UE::UAF
 			return FSkeletonPoseBoneIndex(BoneIndex);
 		}
 
+		// Returns the skeleton bone index that maps to this binding index (only valid for bones)
+		operator FSkeletonPoseBoneIndex() const
+		{
+			// For now, the skeleton bone indices map directly to the attribute binding indices
+			return FSkeletonPoseBoneIndex(BoneIndex);
+		}
+
 		UE_BONE_INDEX_OPERATORS(FAttributeBindingIndex)
 	};
 }
diff --git a/Engine/Plugins/Experimental/UAF/UAF/Source/UAF/Public/UAF/Attributes/AttributeSetMap.h b/Engine/Plugins/Experimental/UAF/UAF/Source/UAF/Public/UAF/Attributes/AttributeSetMap.h
--- a/Engine/Plugins/Experimental/UAF/UAF/Source/UAF/Public/UAF/Attributes/AttributeSetMap.h
+++ b/Engine/Plugins/Experimental/UAF/UAF/Source/UAF/Public/UAF/Attributes/AttributeSetMap.h
@@ -140,6 +140,9 @@ namespace UE::UAF
 		// Returns a copy of this set map using the specified allocator
 		[[nodiscard]] virtual FAttributeSetMap* Duplicate(FReallocFun InReallocFun) const override;
 
+		// Initializes every value to the specified one
+		void FillWithValue(InValueType Value);
+
 		// Initializes every value to its identity, optionally additive
 		virtual void FillWithIdentity(bool bIsAdditive) override;
 
@@ -281,6 +284,24 @@ namespace UE::UAF
 		return Copy;
 	}
 
+	template<class InValueType>
+	inline void TAttributeSetMap<InValueType>::FillWithValue(InValueType Value)
+	{
+		if (NumValues == 0)
+		{
+			return;
+		}
+
+		InValueType* ValuePtr = GetData();
+		InValueType* ValuePtrEnd = ValuePtr + NumValues;
+		while (ValuePtr < ValuePtrEnd)
+		{
+			*ValuePtr = Value;
+
+			ValuePtr++;
+		}
+	}
+
 	template<class InValueType>
 	inline void TAttributeSetMap<InValueType>::FillWithIdentity(bool bIsAdditive)
 	{
@@ -303,14 +324,7 @@ namespace UE::UAF
 			// Everything else uses the default value
 		}
 
-		InValueType* ValuePtr = GetData();
-		InValueType* ValuePtrEnd = ValuePtr + NumValues;
-		while (ValuePtr < ValuePtrEnd)
-		{
-			*ValuePtr = Identity;
-
-			ValuePtr++;
-		}
+		FillWithValue(Identity);
 	}
 
 	template<class InValueType>
diff --git a/Engine/Plugins/Experimental/UAF/UAF/Source/UAF/Public/UAF/Attributes/Transformers/Accumulate.h b/Engine/Plugins/Experimental/UAF/UAF/Source/UAF/Public/UAF/Attributes/Transformers/Accumulate.h
--- a/Engine/Plugins/Experimental/UAF/UAF/Source/UAF/Public/UAF/Attributes/Transformers/Accumulate.h
+++ b/Engine/Plugins/Experimental/UAF/UAF/Source/UAF/Public/UAF/Attributes/Transformers/Accumulate.h
@@ -19,13 +19,21 @@ namespace UE::UAF::Transformers
 		static UE_API FName TransformerName;
 
 		// If the output buffer is nullptr, an allocator is provided and the implementation is responsible for allocating the output buffer and returning it
-		using FTransformSetMapFunc = void (*)(const FAttributeSetMap* SetMapA, const FAttributeSetMap* SetMapB, float Alpha, FAttributeSetMap*& OutputSetMap, FReallocFun OutputAllocator);
-		using FTransformMapFunc = void (*)(const FAttributeMap* MapA, const FAttributeMap* MapB, float Alpha, FAttributeMap*& OutputMap, FReallocFun OutputAllocator);
+		// Per-value map is optional, when not present, the default weight should be used
+		// Unbound values do not support per-value weights
+		using FTransformSetMapFunc = void (*)(const FAttributeSetMap* SetMapA, const FAttributeSetMap* SetMapB, const FAttributeSetMap* PerValueWeightMapB, float DefaultWeightB, FAttributeSetMap*& OutputSetMap, FReallocFun OutputAllocator);
+		using FTransformMapFunc = void (*)(const FAttributeMap* MapA, const FAttributeMap* MapB, float WeightB, FAttributeMap*& OutputMap, FReallocFun OutputAllocator);
 
 		// Accumulate attribute values
-		// Output = accumulate(A, B, Alpha) -> Output = A + (B * Alpha)
+		// Output = A + (B * WeightB)
 		// This function supports outputting to either input
-		static UE_API void Apply(const FAttributeValueTransformerMapPtr& TransformerMap, const FAttributeCollection& CollectionA, const FAttributeCollection& CollectionB, float Alpha, FAttributeCollection& OutputCollection);
+		static UE_API void Apply(const FAttributeValueTransformerMapPtr& TransformerMap, const FAttributeCollection& InputA, const FAttributeCollection& InputB, float WeightB, FAttributeCollection& Output);
+
+		// Accumulate attribute values
+		// If a weight isn't specified, the default weight will be used
+		// Output = A + (B * (HasValueWeightB ? ValueWeightB : DefaultWeightB))
+		// This function supports outputting to either input
+		static UE_API void Apply(const FAttributeValueTransformerMapPtr& TransformerMap, const FAttributeCollection& InputA, const FAttributeCollection& InputB, const FAttributeCollection& PerValueWeightsB, float DefaultWeightB, FAttributeCollection& Output);
 	};
 
 	//////////////////////////////////////////////////////////////////////////
@@ -38,8 +46,8 @@ namespace UE::UAF::Transformers
 			using TransformerType = FAccumulate;
 			using ValueType = FBoneTransformAnimationAttribute;
 
-			static void TransformSetMap(const FAttributeSetMap* SetMapA, const FAttributeSetMap* SetMapB, float Alpha, FAttributeSetMap*& OutputSetMap, FReallocFun OutputAllocator);
-			static void TransformMap(const FAttributeMap* MapA, const FAttributeMap* MapB, float Alpha, FAttributeMap*& OutputMap, FReallocFun OutputAllocator);
+			static void TransformSetMap(const FAttributeSetMap* SetMapA, const FAttributeSetMap* SetMapB, const FAttributeSetMap* PerValueWeightMapB, float DefaultWeightB, FAttributeSetMap*& OutputSetMap, FReallocFun OutputAllocator);
+			static void TransformMap(const FAttributeMap* MapA, const FAttributeMap* MapB, float WeightB, FAttributeMap*& OutputMap, FReallocFun OutputAllocator);
 		};
 
 		struct FAccumulate_FloatAttribute
@@ -47,8 +55,8 @@ namespace UE::UAF::Transformers
 			using TransformerType = FAccumulate;
 			using ValueType = FFloatAnimationAttribute;
 
-			static void TransformSetMap(const FAttributeSetMap* SetMapA, const FAttributeSetMap* SetMapB, float Alpha, FAttributeSetMap*& OutputSetMap, FReallocFun OutputAllocator);
-			static void TransformMap(const FAttributeMap* MapA, const FAttributeMap* MapB, float Alpha, FAttributeMap*& OutputMap, FReallocFun OutputAllocator);
+			static void TransformSetMap(const FAttributeSetMap* SetMapA, const FAttributeSetMap* SetMapB, const FAttributeSetMap* PerValueWeightMapB, float DefaultWeightB, FAttributeSetMap*& OutputSetMap, FReallocFun OutputAllocator);
+			static void TransformMap(const FAttributeMap* MapA, const FAttributeMap* MapB, float WeightB, FAttributeMap*& OutputMap, FReallocFun OutputAllocator);
 		};
 	}
 }
diff --git a/Engine/Plugins/Experimental/UAF/UAF/Source/UAF/Public/UAF/Attributes/Transformers/Overwrite.h b/Engine/Plugins/Experimental/UAF/UAF/Source/UAF/Public/UAF/Attributes/Transformers/Overwrite.h
--- a/Engine/Plugins/Experimental/UAF/UAF/Source/UAF/Public/UAF/Attributes/Transformers/Overwrite.h
+++ b/Engine/Plugins/Experimental/UAF/UAF/Source/UAF/Public/UAF/Attributes/Transformers/Overwrite.h
@@ -20,13 +20,21 @@ namespace UE::UAF::Transformers
 
 		// If the output buffer is nullptr, an allocator is provided and the implementation is responsible for allocating the output buffer and returning it
 		// Note that the input/output buffers can be the same buffer
-		using FTransformSetMapFunc = void (*)(const FAttributeSetMap* SetMap, float Alpha, FAttributeSetMap*& OutputSetMap, FReallocFun OutputAllocator);
-		using FTransformMapFunc = void (*)(const FAttributeMap* Map, float Alpha, FAttributeMap*& OutputMap, FReallocFun OutputAllocator);
+		// Per-value map is optional, when not present, the default weight should be used
+		// Unbound values do not support per-value weights
+		using FTransformSetMapFunc = void (*)(const FAttributeSetMap* InputSetMap, const FAttributeSetMap* PerValueWeightMap, float DefaultWeight, FAttributeSetMap*& OutputSetMap, FReallocFun OutputAllocator);
+		using FTransformMapFunc = void (*)(const FAttributeMap* InputMap, float Weight, FAttributeMap*& OutputMap, FReallocFun OutputAllocator);
 
 		// Overwrite attribute values
-		// Output = overwrite(A, Alpha) -> Output = (A * Alpha)
+		// Output = (A * Weight)
 		// This function supports outputting to its input
-		static UE_API void Apply(const FAttributeValueTransformerMapPtr& TransformerMap, const FAttributeCollection& SourceCollection, float Alpha, FAttributeCollection& OutputCollection);
+		static UE_API void Apply(const FAttributeValueTransformerMapPtr& TransformerMap, const FAttributeCollection& Input, float Weight, FAttributeCollection& Output);
+
+		// Overwrite attribute values using per value weights (must be of type FFloatAnimationAttribute)
+		// If a weight isn't specified, the default weight will be used
+		// Output = A * (HasValueWeight ? ValueWeight : DefaultWeight)
+		// This function supports outputting to its input
+		static UE_API void Apply(const FAttributeValueTransformerMapPtr& TransformerMap, const FAttributeCollection& Input, const FAttributeCollection& PerValueWeights, float DefaultWeight, FAttributeCollection& Output);
 	};
 
 	//////////////////////////////////////////////////////////////////////////
@@ -39,8 +47,8 @@ namespace UE::UAF::Transformers
 			using TransformerType = FOverwrite;
 			using ValueType = FBoneTransformAnimationAttribute;
 
-			static void TransformSetMap(const FAttributeSetMap* SetMap, float Alpha, FAttributeSetMap*& OutputSetMap, FReallocFun OutputAllocator);
-			static void TransformMap(const FAttributeMap* Map, float Alpha, FAttributeMap*& OutputMap, FReallocFun OutputAllocator);
+			static void TransformSetMap(const FAttributeSetMap* InputSetMap, const FAttributeSetMap* PerValueWeightMap, float DefaultWeight, FAttributeSetMap*& OutputSetMap, FReallocFun OutputAllocator);
+			static void TransformMap(const FAttributeMap* InputMap, float Weight, FAttributeMap*& OutputMap, FReallocFun OutputAllocator);
 		};
 
 		struct FOverwrite_FloatAttribute
@@ -48,8 +56,8 @@ namespace UE::UAF::Transformers
 			using TransformerType = FOverwrite;
 			using ValueType = FFloatAnimationAttribute;
 
-			static void TransformSetMap(const FAttributeSetMap* SetMap, float Alpha, FAttributeSetMap*& OutputSetMap, FReallocFun OutputAllocator);
-			static void TransformMap(const FAttributeMap* Map, float Alpha, FAttributeMap*& OutputMap, FReallocFun OutputAllocator);
+			static void TransformSetMap(const FAttributeSetMap* InputSetMap, const FAttributeSetMap* PerValueWeightMap, float DefaultWeight, FAttributeSetMap*& OutputSetMap, FReallocFun OutputAllocator);
+			static void TransformMap(const FAttributeMap* InputMap, float Weight, FAttributeMap*& OutputMap, FReallocFun OutputAllocator);
 		};
 	}
 }
diff --git a/Engine/Plugins/Experimental/UAF/UAFAnimGraph/Source/UAFAnimGraph/Private/EvaluationVM/Tasks/BlendKeyframesPerBone.cpp b/Engine/Plugins/Experimental/UAF/UAFAnimGraph/Source/UAFAnimGraph/Private/EvaluationVM/Tasks/BlendKeyframesPerBone.cpp
--- a/Engine/Plugins/Experimental/UAF/UAFAnimGraph/Source/UAFAnimGraph/Private/EvaluationVM/Tasks/BlendKeyframesPerBone.cpp
+++ b/Engine/Plugins/Experimental/UAF/UAFAnimGraph/Source/UAFAnimGraph/Private/EvaluationVM/Tasks/BlendKeyframesPerBone.cpp
@@ -11,9 +11,107 @@
 #include "EvaluationVM/EvaluationVM.h"
 #include "EvaluationVM/KeyframeState.h"
 #include "TransformArrayOperations.h"
+#include "UAF/Attributes/Transformers/Accumulate.h"
+#include "UAF/Attributes/Transformers/Overwrite.h"
 
 #include UE_INLINE_GENERATED_CPP_BY_NAME(BlendKeyframesPerBone)
 
+namespace UE::UAF::Private
+{
+	static FAttributeCollectionStack BuildPerValueWeights(
+		const FAttributeCollection& Input,
+		UE::UAF::FEvaluationVM& VM,
+		const FBlendSampleData* BlendData,
+		const IBlendProfileInterface* BlendProfile,
+		bool bPerBoneDataUsesSkeletonIndex)
+	{
+		const float BlendWeight = BlendData->GetClampedWeight();
+
+		UScriptStruct* FloatAttributeType = FFloatAnimationAttribute::StaticStruct();
+
+		// Build our per-value weights (only for bound values)
+		FAttributeCollectionStack PerValueWeights(VM.GetActiveNamedSet());
+
+		for (auto It = Input.CreateConstSetMapIterator(); It; ++It)
+		{
+			UScriptStruct* AttributeType = It.GetAttributeType();
+
+			// Every attribute type maps to a float, our weight
+			FAttributeMappingKey MappingKey = FAttributeMappingKey::MakeFromTo(AttributeType, FloatAttributeType);
+
+			TAttributeSetMap<FFloatAnimationAttribute>* SetWeights = PerValueWeights.AddSetMap<FFloatAnimationAttribute>(MappingKey);
+
+			if (const TAttributeSetMap<FBoneTransformAnimationAttribute>* BoneMap = Cast<FBoneTransformAnimationAttribute>(It.GetMap()))
+			{
+				const FAttributeTypedSetPtr& BoneSet = BoneMap->GetTypedSet();
+				const int32 NumBonesInSet = BoneMap->Num();
+
+				if (BlendProfile != nullptr)
+				{
+					const USkeleton* TargetSkeleton = VM.GetActiveEvaluationContext().GetSkeleton();
+					TSharedPtr<IInterpolationIndexProvider::FPerBoneInterpolationData> Data = BlendProfile->GetPerBoneInterpolationData(TargetSkeleton);
+
+					const FSkeletonRemapping& SkeletonRemapping = UE::Anim::FSkeletonRemappingRegistry::Get().GetRemapping(BlendProfile->GetSkeleton(), TargetSkeleton);
+
+					if (SkeletonRemapping.IsValid())
+					{
+						for (FAttributeSetIndex SetIndex(0); SetIndex < NumBonesInSet; ++SetIndex)
+						{
+							const FAttributeBindingIndex BindingIndex = BoneSet->GetBindingIndex(SetIndex);
+
+							const FSkeletonPoseBoneIndex TargetSkeletonBoneIndex = BindingIndex;
+							const FSkeletonPoseBoneIndex SourceSkeletonBoneIndex(SkeletonRemapping.GetSourceSkeletonBoneIndex(TargetSkeletonBoneIndex.GetInt()));
+
+							const int32 WeightIndex = SourceSkeletonBoneIndex.IsValid() ? BlendProfile->GetPerBoneInterpolationIndex(SourceSkeletonBoneIndex, TargetSkeleton, Data.Get()) : INDEX_NONE;
+							(*SetWeights)[SetIndex].Value = WeightIndex != INDEX_NONE ? BlendData->PerBoneBlendData[WeightIndex] : BlendWeight;
+						}
+					}
+					else
+					{
+						for (FAttributeSetIndex SetIndex(0); SetIndex < NumBonesInSet; ++SetIndex)
+						{
+							const FAttributeBindingIndex BindingIndex = BoneSet->GetBindingIndex(SetIndex);
+
+							const FSkeletonPoseBoneIndex TargetSkeletonBoneIndex = BindingIndex;
+
+							const int32 WeightIndex = BlendProfile->GetPerBoneInterpolationIndex(TargetSkeletonBoneIndex, TargetSkeleton, Data.Get());
+							(*SetWeights)[SetIndex].Value = WeightIndex != INDEX_NONE ? BlendData->PerBoneBlendData[WeightIndex] : BlendWeight;
+						}
+					}
+				}
+				else
+				{
+					// If no blend profile is provided, we use the provided blend per bone data under the assumption it covers a full pose
+					// either using LODBone indexing or full skeleton indexing
+					check(BlendData->PerBoneBlendData.Num() >= NumBonesInSet);
+
+					for (FAttributeSetIndex SetIndex(0); SetIndex < NumBonesInSet; ++SetIndex)
+					{
+						if (bPerBoneDataUsesSkeletonIndex)
+						{
+							const FAttributeBindingIndex BindingIndex = BoneSet->GetBindingIndex(SetIndex);
+
+							const FSkeletonPoseBoneIndex TargetSkeletonBoneIndex = BindingIndex;
+							(*SetWeights)[SetIndex].Value = TargetSkeletonBoneIndex.IsValid() ? BlendData->PerBoneBlendData[TargetSkeletonBoneIndex.GetInt()] : BlendWeight;
+						}
+						else
+						{
+							(*SetWeights)[SetIndex].Value = BlendData->PerBoneBlendData[SetIndex.GetInt()];
+						}
+					}
+				}
+			}
+			else
+			{
+				// Everything else defaults to full weight
+				SetWeights->FillWithValue(FFloatAnimationAttribute{ 1.0f });
+			}
+		}
+
+		return PerValueWeights;
+	}
+}
+
 FAnimNextBlendOverwriteKeyframePerBoneWithScaleTask FAnimNextBlendOverwriteKeyframePerBoneWithScaleTask::Make(const IBlendProfileInterface* BlendProfile, const FBlendSampleData& BlendData, float ScaleFactor)
 {
 	FAnimNextBlendOverwriteKeyframePerBoneWithScaleTask Task;
@@ -39,13 +137,6 @@ void FAnimNextBlendOverwriteKeyframePerBoneWithScaleTask::Execute(UE::UAF::FEval
 {
 	using namespace UE::UAF;
 
-	if (VM.GetActiveNamedSet())
-	{
-		// TODO: Implement with new attribute runtime
-		Super::Execute(VM);
-		return;
-	}
-
 	if (BlendProfile == nullptr && BlendData->PerBoneBlendData.Num() == 0)
 	{
 		// If we don't have a blend profile and per bone blend data, we blend the whole pose
@@ -53,89 +144,110 @@ void FAnimNextBlendOverwriteKeyframePerBoneWithScaleTask::Execute(UE::UAF::FEval
 		return;
 	}
 
-	TUniquePtr<FKeyframeState> Keyframe;
-	if (!VM.PopValue(KEYFRAME_STACK_NAME, Keyframe))
+	if (VM.GetActiveNamedSet())
 	{
-		// We have no inputs, nothing to do
-		return;
-	}
+		TUniquePtr<FAttributeCollection> Input;
+		if (!VM.PopValue(ATTRIBUTE_STACK_NAME, Input))
+		{
+			// We have no inputs, nothing to do
+			return;
+		}
 
-	const TArrayView<const FBoneIndexType> LODBoneIndexToSkeletonBoneIndexMap = Keyframe->Pose.GetLODBoneIndexToSkeletonBoneIndexMap();
-	const int32 NumLODBones = LODBoneIndexToSkeletonBoneIndexMap.Num();
+		const float BlendWeight = BlendData->GetClampedWeight();
 
-	TArray<int32> LODBoneIndexToWeightIndexMap;
-	LODBoneIndexToWeightIndexMap.AddUninitialized(NumLODBones);
+		const FAttributeCollectionStack PerValueWeights = Private::BuildPerValueWeights(*Input, VM, BlendData, BlendProfile, bPerBoneDataUsesSkeletonIndex);
+
+		// Modify our input in place
+		Transformers::FOverwrite::Apply(VM.GetTransformerMap(), *Input, PerValueWeights, BlendWeight, *Input);
 
-	if (BlendProfile)
+		VM.PushValue(ATTRIBUTE_STACK_NAME, MoveTemp(Input));
+	}
+	else
 	{
-		const USkeleton* TargetSkeleton = Keyframe->Pose.GetSkeletonAsset();
-		TSharedPtr<IInterpolationIndexProvider::FPerBoneInterpolationData> Data = BlendProfile->GetPerBoneInterpolationData(TargetSkeleton);
+		TUniquePtr<FKeyframeState> Keyframe;
+		if (!VM.PopValue(KEYFRAME_STACK_NAME, Keyframe))
+		{
+			// We have no inputs, nothing to do
+			return;
+		}
+
+		const TArrayView<const FBoneIndexType> LODBoneIndexToSkeletonBoneIndexMap = Keyframe->Pose.GetLODBoneIndexToSkeletonBoneIndexMap();
+		const int32 NumLODBones = LODBoneIndexToSkeletonBoneIndexMap.Num();
 
-		const FSkeletonRemapping& SkeletonRemapping = UE::Anim::FSkeletonRemappingRegistry::Get().GetRemapping(BlendProfile->GetSkeleton(), TargetSkeleton);
+		TArray<int32> LODBoneIndexToWeightIndexMap;
+		LODBoneIndexToWeightIndexMap.AddUninitialized(NumLODBones);
 
-		if (SkeletonRemapping.IsValid())
+		if (BlendProfile)
 		{
-			for (int32 LODBoneIndex = 0; LODBoneIndex < NumLODBones; ++LODBoneIndex)
+			const USkeleton* TargetSkeleton = Keyframe->Pose.GetSkeletonAsset();
+			TSharedPtr<IInterpolationIndexProvider::FPerBoneInterpolationData> Data = BlendProfile->GetPerBoneInterpolationData(TargetSkeleton);
+
+			const FSkeletonRemapping& SkeletonRemapping = UE::Anim::FSkeletonRemappingRegistry::Get().GetRemapping(BlendProfile->GetSkeleton(), TargetSkeleton);
+
+			if (SkeletonRemapping.IsValid())
+			{
+				for (int32 LODBoneIndex = 0; LODBoneIndex < NumLODBones; ++LODBoneIndex)
+				{
+					const FSkeletonPoseBoneIndex TargetSkeletonBoneIndex(LODBoneIndexToSkeletonBoneIndexMap[LODBoneIndex]);
+					const FSkeletonPoseBoneIndex SourceSkeletonBoneIndex(SkeletonRemapping.GetSourceSkeletonBoneIndex(TargetSkeletonBoneIndex.GetInt()));
+					LODBoneIndexToWeightIndexMap[LODBoneIndex] = SourceSkeletonBoneIndex.IsValid() ? BlendProfile->GetPerBoneInterpolationIndex(SourceSkeletonBoneIndex, TargetSkeleton, Data.Get()) : INDEX_NONE;
+				}
+			}
+			else
 			{
-				const FSkeletonPoseBoneIndex TargetSkeletonBoneIndex(LODBoneIndexToSkeletonBoneIndexMap[LODBoneIndex]);
-				const FSkeletonPoseBoneIndex SourceSkeletonBoneIndex(SkeletonRemapping.GetSourceSkeletonBoneIndex(TargetSkeletonBoneIndex.GetInt()));
-				LODBoneIndexToWeightIndexMap[LODBoneIndex] = SourceSkeletonBoneIndex.IsValid() ? BlendProfile->GetPerBoneInterpolationIndex(SourceSkeletonBoneIndex, TargetSkeleton, Data.Get()) : INDEX_NONE;
+				for (int32 LODBoneIndex = 0; LODBoneIndex < NumLODBones; ++LODBoneIndex)
+				{
+					const FSkeletonPoseBoneIndex SkeletonBoneIndex(LODBoneIndexToSkeletonBoneIndexMap[LODBoneIndex]);
+					LODBoneIndexToWeightIndexMap[LODBoneIndex] = BlendProfile->GetPerBoneInterpolationIndex(SkeletonBoneIndex, TargetSkeleton, Data.Get());
+				}
 			}
 		}
 		else
 		{
+			// If no blend profile is provided, we use the provided blend per bone data under the assumption it covers a full pose
+			// either using LODBone indexing or full skeleton indexing
+			check(BlendData->PerBoneBlendData.Num() >= NumLODBones);
 			for (int32 LODBoneIndex = 0; LODBoneIndex < NumLODBones; ++LODBoneIndex)
 			{
-				const FSkeletonPoseBoneIndex SkeletonBoneIndex(LODBoneIndexToSkeletonBoneIndexMap[LODBoneIndex]);
-				LODBoneIndexToWeightIndexMap[LODBoneIndex] = BlendProfile->GetPerBoneInterpolationIndex(SkeletonBoneIndex, TargetSkeleton, Data.Get());
+				if (bPerBoneDataUsesSkeletonIndex)
+				{
+					const FSkeletonPoseBoneIndex TargetSkeletonBoneIndex(LODBoneIndexToSkeletonBoneIndexMap[LODBoneIndex]);
+					LODBoneIndexToWeightIndexMap[LODBoneIndex] = TargetSkeletonBoneIndex.GetInt();
+				}
+				else
+				{
+					LODBoneIndexToWeightIndexMap[LODBoneIndex] = LODBoneIndex;
+				}
 			}
 		}
-	}
-	else
-	{
-		// If no blend profile is provided, we use the provided blend per bone data under the assumption it covers a full pose
-		// either using LODBone indexing or full skeleton indexing
-		check(BlendData->PerBoneBlendData.Num() >= NumLODBones);
-		for (int32 LODBoneIndex = 0; LODBoneIndex < NumLODBones; ++LODBoneIndex)
+
+		const float BlendWeight = BlendData->GetClampedWeight();
+
+		if (EnumHasAnyFlags(VM.GetFlags(), EEvaluationFlags::Bones))
 		{
-			if (bPerBoneDataUsesSkeletonIndex)
-			{
-				const FSkeletonPoseBoneIndex TargetSkeletonBoneIndex(LODBoneIndexToSkeletonBoneIndexMap[LODBoneIndex]);
-				LODBoneIndexToWeightIndexMap[LODBoneIndex] = TargetSkeletonBoneIndex.GetInt();
-			}
-			else
-			{
-				LODBoneIndexToWeightIndexMap[LODBoneIndex] = LODBoneIndex;
-			}
+			BlendOverwritePerBoneWithScale(
+				Keyframe->Pose.LocalTransforms.GetView(), Keyframe->Pose.LocalTransforms.GetConstView(),
+				LODBoneIndexToWeightIndexMap, BlendData->PerBoneBlendData, BlendWeight);
 		}
-	}
-	
-	const float BlendWeight = BlendData->GetClampedWeight();
 
-	if (EnumHasAnyFlags(VM.GetFlags(), EEvaluationFlags::Bones))
-	{
-		BlendOverwritePerBoneWithScale(
-			Keyframe->Pose.LocalTransforms.GetView(), Keyframe->Pose.LocalTransforms.GetConstView(),
-			LODBoneIndexToWeightIndexMap, BlendData->PerBoneBlendData, BlendWeight);
-	}
+		if (EnumHasAnyFlags(VM.GetFlags(), EEvaluationFlags::Curves))
+		{
+			// Curves cannot override in place
+			FBlendedCurve Result;
+			Result.Override(Keyframe->Curves, BlendWeight);
 
-	if (EnumHasAnyFlags(VM.GetFlags(), EEvaluationFlags::Curves))
-	{
-		// Curves cannot override in place
-		FBlendedCurve Result;
-		Result.Override(Keyframe->Curves, BlendWeight);
+			Keyframe->Curves = MoveTemp(Result);
+		}
 
-		Keyframe->Curves = MoveTemp(Result);
-	}
+		if (EnumHasAnyFlags(VM.GetFlags(), EEvaluationFlags::Attributes))
+		{
+			UE::Anim::FStackAttributeContainer OutputAttributes;
+			UE::Anim::Attributes::BlendAttributesPerBone({ Keyframe->Attributes }, { LODBoneIndexToWeightIndexMap }, { *BlendData }, { 0 }, { OutputAttributes });
+			Keyframe->Attributes.MoveFrom(OutputAttributes);
+		}
 
-	if (EnumHasAnyFlags(VM.GetFlags(), EEvaluationFlags::Attributes))
-	{
-		UE::Anim::FStackAttributeContainer OutputAttributes;
-		UE::Anim::Attributes::BlendAttributesPerBone({ Keyframe->Attributes }, { LODBoneIndexToWeightIndexMap }, { *BlendData }, { 0 }, { OutputAttributes });
-		Keyframe->Attributes.MoveFrom(OutputAttributes);
+		VM.PushValue(KEYFRAME_STACK_NAME, MoveTemp(Keyframe));
 	}
-
-	VM.PushValue(KEYFRAME_STACK_NAME, MoveTemp(Keyframe));
 }
 
 FAnimNextBlendAddKeyframePerBoneWithScaleTask FAnimNextBlendAddKeyframePerBoneWithScaleTask::Make(const IBlendProfileInterface* BlendProfile, const FBlendSampleData& BlendDataA, const FBlendSampleData& BlendDataB, float ScaleFactor)
@@ -165,14 +277,7 @@ void FAnimNextBlendAddKeyframePerBoneWithScaleTask::Execute(UE::UAF::FEvaluation
 {
 	using namespace UE::UAF;
 
-	if (VM.GetActiveNamedSet())
-	{
-		// TODO: Implement with new attribute runtime
-		Super::Execute(VM);
-		return;
-	}
-
-	if (BlendProfile == nullptr 
+	if (BlendProfile == nullptr
 		&& ((BlendDataA->PerBoneBlendData.Num() == 0 || BlendDataB->PerBoneBlendData.Num() == 0)
 			|| (BlendDataA->PerBoneBlendData.Num() != BlendDataB->PerBoneBlendData.Num())))
 	{
@@ -181,104 +286,134 @@ void FAnimNextBlendAddKeyframePerBoneWithScaleTask::Execute(UE::UAF::FEvaluation
 		return;
 	}
 
-	// Pop our top two poses, we'll re-use the top keyframe for our result
-
-	TUniquePtr<FKeyframeState> KeyframeB;
-	if (!VM.PopValue(KEYFRAME_STACK_NAME, KeyframeB))
+	if (VM.GetActiveNamedSet())
 	{
-		// We have no inputs, nothing to do
-		return;
-	}
+		TUniquePtr<FAttributeCollection> InputB;
+		if (!VM.PopValue(ATTRIBUTE_STACK_NAME, InputB))
+		{
+			// We have no inputs, nothing to do
+			return;
+		}
 
-	TUniquePtr<FKeyframeState> KeyframeA;
-	if (!VM.PopValue(KEYFRAME_STACK_NAME, KeyframeA))
-	{
-		// We have a single input, leave it on top of the stack
-		VM.PushValue(KEYFRAME_STACK_NAME, MoveTemp(KeyframeB));
-		return;
-	}
+		TUniquePtr<FAttributeCollection> InputA;
+		if (!VM.PopValue(ATTRIBUTE_STACK_NAME, InputA))
+		{
+			// We have a single input, leave it on top of the stack
+			VM.PushValue(ATTRIBUTE_STACK_NAME, MoveTemp(InputB));
+			return;
+		}
 
-	const TArrayView<const FBoneIndexType> LODBoneIndexToSkeletonBoneIndexMap = KeyframeB->Pose.GetLODBoneIndexToSkeletonBoneIndexMap();
-	const int32 NumLODBones = LODBoneIndexToSkeletonBoneIndexMap.Num();
+		const float BlendWeight = BlendDataA->GetClampedWeight();
 
-	TArray<int32> LODBoneIndexToWeightIndexMap;
-	LODBoneIndexToWeightIndexMap.AddUninitialized(NumLODBones);
+		const FAttributeCollectionStack PerValueWeightsA = Private::BuildPerValueWeights(*InputA, VM, BlendDataA, BlendProfile, bPerBoneDataUsesSkeletonIndex);
 
-	if (BlendProfile)
-	{
-		const USkeleton* TargetSkeleton = KeyframeB->Pose.GetSkeletonAsset();
-		TSharedPtr<IInterpolationIndexProvider::FPerBoneInterpolationData> Data = BlendProfile->GetPerBoneInterpolationData(TargetSkeleton);
+		// Modify input B in place
+		// We flip A and B here because state is pushed in reversed order
+		Transformers::FAccumulate::Apply(VM.GetTransformerMap(), *InputB, *InputA, PerValueWeightsA, BlendWeight, *InputB);
 
-		const FSkeletonRemapping& SkeletonRemapping = UE::Anim::FSkeletonRemappingRegistry::Get().GetRemapping(BlendProfile->GetSkeleton(), TargetSkeleton);
+		VM.PushValue(ATTRIBUTE_STACK_NAME, MoveTemp(InputB));
+	}
+	else
+	{
+		// Pop our top two poses, we'll re-use the top keyframe for our result
 
-		if (SkeletonRemapping.IsValid())
+		TUniquePtr<FKeyframeState> KeyframeB;
+		if (!VM.PopValue(KEYFRAME_STACK_NAME, KeyframeB))
 		{
-			for (int32 LODBoneIndex = 0; LODBoneIndex < NumLODBones; ++LODBoneIndex)
-			{
-				const FSkeletonPoseBoneIndex TargetSkeletonBoneIndex(LODBoneIndexToSkeletonBoneIndexMap[LODBoneIndex]);
-				const FSkeletonPoseBoneIndex SourceSkeletonBoneIndex(SkeletonRemapping.GetSourceSkeletonBoneIndex(TargetSkeletonBoneIndex.GetInt()));
-				LODBoneIndexToWeightIndexMap[LODBoneIndex] = SourceSkeletonBoneIndex.IsValid() ? BlendProfile->GetPerBoneInterpolationIndex(SourceSkeletonBoneIndex, TargetSkeleton, Data.Get()) : INDEX_NONE;
-			}
+			// We have no inputs, nothing to do
+			return;
 		}
-		else
+
+		TUniquePtr<FKeyframeState> KeyframeA;
+		if (!VM.PopValue(KEYFRAME_STACK_NAME, KeyframeA))
 		{
-			for (int32 LODBoneIndex = 0; LODBoneIndex < NumLODBones; ++LODBoneIndex)
-			{
-				const FSkeletonPoseBoneIndex SkeletonBoneIndex(LODBoneIndexToSkeletonBoneIndexMap[LODBoneIndex]);
-				LODBoneIndexToWeightIndexMap[LODBoneIndex] = BlendProfile->GetPerBoneInterpolationIndex(SkeletonBoneIndex, TargetSkeleton, Data.Get());
-			}
+			// We have a single input, leave it on top of the stack
+			VM.PushValue(KEYFRAME_STACK_NAME, MoveTemp(KeyframeB));
+			return;
 		}
-	}
-	else
-	{
-		// If no blend profile is provided, we assume the provided bone data encompasses a full pose to be blended together 
-		// either using LOD based indexing or full skeleton indexing
-		check(BlendDataA->PerBoneBlendData.Num() >= NumLODBones);
-		for (int32 LODBoneIndex = 0; LODBoneIndex < NumLODBones; ++LODBoneIndex)
+
+		const TArrayView<const FBoneIndexType> LODBoneIndexToSkeletonBoneIndexMap = KeyframeB->Pose.GetLODBoneIndexToSkeletonBoneIndexMap();
+		const int32 NumLODBones = LODBoneIndexToSkeletonBoneIndexMap.Num();
+
+		TArray<int32> LODBoneIndexToWeightIndexMap;
+		LODBoneIndexToWeightIndexMap.AddUninitialized(NumLODBones);
+
+		if (BlendProfile)
 		{
-			if (bPerBoneDataUsesSkeletonIndex)
+			const USkeleton* TargetSkeleton = KeyframeB->Pose.GetSkeletonAsset();
+			TSharedPtr<IInterpolationIndexProvider::FPerBoneInterpolationData> Data = BlendProfile->GetPerBoneInterpolationData(TargetSkeleton);
+
+			const FSkeletonRemapping& SkeletonRemapping = UE::Anim::FSkeletonRemappingRegistry::Get().GetRemapping(BlendProfile->GetSkeleton(), TargetSkeleton);
+
+			if (SkeletonRemapping.IsValid())
 			{
-				const FSkeletonPoseBoneIndex TargetSkeletonBoneIndex(LODBoneIndexToSkeletonBoneIndexMap[LODBoneIndex]);
-				LODBoneIndexToWeightIndexMap[LODBoneIndex] = TargetSkeletonBoneIndex.GetInt();
+				for (int32 LODBoneIndex = 0; LODBoneIndex < NumLODBones; ++LODBoneIndex)
+				{
+					const FSkeletonPoseBoneIndex TargetSkeletonBoneIndex(LODBoneIndexToSkeletonBoneIndexMap[LODBoneIndex]);
+					const FSkeletonPoseBoneIndex SourceSkeletonBoneIndex(SkeletonRemapping.GetSourceSkeletonBoneIndex(TargetSkeletonBoneIndex.GetInt()));
+					LODBoneIndexToWeightIndexMap[LODBoneIndex] = SourceSkeletonBoneIndex.IsValid() ? BlendProfile->GetPerBoneInterpolationIndex(SourceSkeletonBoneIndex, TargetSkeleton, Data.Get()) : INDEX_NONE;
+				}
 			}
 			else
 			{
-				LODBoneIndexToWeightIndexMap[LODBoneIndex] = LODBoneIndex;
+				for (int32 LODBoneIndex = 0; LODBoneIndex < NumLODBones; ++LODBoneIndex)
+				{
+					const FSkeletonPoseBoneIndex SkeletonBoneIndex(LODBoneIndexToSkeletonBoneIndexMap[LODBoneIndex]);
+					LODBoneIndexToWeightIndexMap[LODBoneIndex] = BlendProfile->GetPerBoneInterpolationIndex(SkeletonBoneIndex, TargetSkeleton, Data.Get());
+				}
+			}
+		}
+		else
+		{
+			// If no blend profile is provided, we assume the provided bone data encompasses a full pose to be blended together 
+			// either using LOD based indexing or full skeleton indexing
+			check(BlendDataA->PerBoneBlendData.Num() >= NumLODBones);
+			for (int32 LODBoneIndex = 0; LODBoneIndex < NumLODBones; ++LODBoneIndex)
+			{
+				if (bPerBoneDataUsesSkeletonIndex)
+				{
+					const FSkeletonPoseBoneIndex TargetSkeletonBoneIndex(LODBoneIndexToSkeletonBoneIndexMap[LODBoneIndex]);
+					LODBoneIndexToWeightIndexMap[LODBoneIndex] = TargetSkeletonBoneIndex.GetInt();
+				}
+				else
+				{
+					LODBoneIndexToWeightIndexMap[LODBoneIndex] = LODBoneIndex;
+				}
 			}
 		}
-	}
 
-	const float BlendWeightA = BlendDataA->GetClampedWeight();
+		const float BlendWeightA = BlendDataA->GetClampedWeight();
 
-	if (EnumHasAnyFlags(VM.GetFlags(), EEvaluationFlags::Bones))
-	{
-		check(KeyframeA->Pose.GetNumBones() == KeyframeB->Pose.GetNumBones());
+		if (EnumHasAnyFlags(VM.GetFlags(), EEvaluationFlags::Bones))
+		{
+			check(KeyframeA->Pose.GetNumBones() == KeyframeB->Pose.GetNumBones());
 
-		BlendAddPerBoneWithScale(
-			KeyframeB->Pose.LocalTransforms.GetView(), KeyframeA->Pose.LocalTransforms.GetConstView(),
-			LODBoneIndexToWeightIndexMap, BlendDataA->PerBoneBlendData, BlendWeightA);
-	}
+			BlendAddPerBoneWithScale(
+				KeyframeB->Pose.LocalTransforms.GetView(), KeyframeA->Pose.LocalTransforms.GetConstView(),
+				LODBoneIndexToWeightIndexMap, BlendDataA->PerBoneBlendData, BlendWeightA);
+		}
 
-	if (EnumHasAnyFlags(VM.GetFlags(), EEvaluationFlags::Curves))
-	{
-		KeyframeB->Curves.Accumulate(KeyframeA->Curves, BlendWeightA);
-	}
+		if (EnumHasAnyFlags(VM.GetFlags(), EEvaluationFlags::Curves))
+		{
+			KeyframeB->Curves.Accumulate(KeyframeA->Curves, BlendWeightA);
+		}
 
-	if (EnumHasAnyFlags(VM.GetFlags(), EEvaluationFlags::Attributes))
-	{
-		// TODO: Might need to be revisited, might not work as intended
-		UE::Anim::FStackAttributeContainer OutputAttributes;
-		UE::Anim::Attributes::BlendAttributesPerBone(
-			{ KeyframeA->Attributes, KeyframeB->Attributes },
-			{ LODBoneIndexToWeightIndexMap },
-			{ *BlendDataA, *BlendDataB },
-			{ 0, 1 },
-			{ OutputAttributes });
-
-		KeyframeB->Attributes.MoveFrom(OutputAttributes);
-	}
+		if (EnumHasAnyFlags(VM.GetFlags(), EEvaluationFlags::Attributes))
+		{
+			// TODO: Might need to be revisited, might not work as intended
+			UE::Anim::FStackAttributeContainer OutputAttributes;
+			UE::Anim::Attributes::BlendAttributesPerBone(
+				{ KeyframeA->Attributes, KeyframeB->Attributes },
+				{ LODBoneIndexToWeightIndexMap },
+				{ *BlendDataA, *BlendDataB },
+				{ 0, 1 },
+				{ OutputAttributes });
+
+			KeyframeB->Attributes.MoveFrom(OutputAttributes);
+		}
 
-	VM.PushValue(KEYFRAME_STACK_NAME, MoveTemp(KeyframeB));
+		VM.PushValue(KEYFRAME_STACK_NAME, MoveTemp(KeyframeB));
+	}
 }
 
 FAnimNextBlendKeyframePerBoneWithScaleTask FAnimNextBlendKeyframePerBoneWithScaleTask::Make(TSharedPtr<IInterpolationIndexProvider> BlendProfile, const USkeleton* Skeleton, const TArray<float>& BoneMaskWeights, const UE::Anim::TNamedValueArray<FDefaultAllocator, UE::Anim::FCurveElement>& CurveMaskWeights, const TArray<FMaskedAttributeWeight>& AttributeMaskWeights, float ScaleFactor)
