Add context menu option to unbind all attributes in a set

#Synced-CL 48947880

diff --git a/Engine/Plugins/Experimental/UAF/UAF/Source/UAF/Private/UAF/AbstractSkeleton/AbstractSkeletonSetBinding.cpp b/Engine/Plugins/Experimental/UAF/UAF/Source/UAF/Private/UAF/AbstractSkeleton/AbstractSkeletonSetBinding.cpp
--- a/Engine/Plugins/Experimental/UAF/UAF/Source/UAF/Private/UAF/AbstractSkeleton/AbstractSkeletonSetBinding.cpp
+++ b/Engine/Plugins/Experimental/UAF/UAF/Source/UAF/Private/UAF/AbstractSkeleton/AbstractSkeletonSetBinding.cpp
@@ -194,6 +194,23 @@ bool UAbstractSkeletonSetBinding::RemoveAttributeFromSet(const FAnimationAttribu
 	}
 }
 
+bool UAbstractSkeletonSetBinding::RemoveAllFromSet(const FName SetName)
+{
+	int32 bCountRemoved = 0;
+
+	bCountRemoved += BoneBindings.RemoveAll([&](const FAbstractSkeleton_BoneBinding& Binding)
+		{
+			return Binding.SetName == SetName;	
+		});
+
+	bCountRemoved += AttributeBindings.RemoveAll([&](const FAbstractSkeleton_AttributeBinding& Binding)
+		{
+			return Binding.SetName == SetName;
+		});
+
+	return bCountRemoved > 0;
+}
+
 bool UAbstractSkeletonSetBinding::IsAttributeInSet(const FAnimationAttributeIdentifier Attribute) const
 {
 	return AttributeBindings.ContainsByPredicate([&](const FAbstractSkeleton_AttributeBinding& Binding)
diff --git a/Engine/Plugins/Experimental/UAF/UAF/Source/UAF/Public/UAF/AbstractSkeleton/AbstractSkeletonSetBinding.h b/Engine/Plugins/Experimental/UAF/UAF/Source/UAF/Public/UAF/AbstractSkeleton/AbstractSkeletonSetBinding.h
--- a/Engine/Plugins/Experimental/UAF/UAF/Source/UAF/Public/UAF/AbstractSkeleton/AbstractSkeletonSetBinding.h
+++ b/Engine/Plugins/Experimental/UAF/UAF/Source/UAF/Public/UAF/AbstractSkeleton/AbstractSkeletonSetBinding.h
@@ -79,6 +79,9 @@ public:
 	UE_API bool AddAttributeToSet(const FAnimationAttributeIdentifier Attribute, const FName SetName);
 
 	UE_API bool RemoveAttributeFromSet(const FAnimationAttributeIdentifier Attribute);
+
+	// Unbinds all attributes bound to a set. Returns whether or not any attributes were removed.
+	UE_API bool RemoveAllFromSet(const FName SetName);
 	
 	// Returns true if the provided attribute is in ANY set
 	UE_API bool IsAttributeInSet(const FAnimationAttributeIdentifier Attribute) const;
diff --git a/Engine/Plugins/Experimental/UAF/UAF/Source/UAFEditor/Private/UAF/AbstractSkeleton/Sets/SSetBinding.cpp b/Engine/Plugins/Experimental/UAF/UAF/Source/UAFEditor/Private/UAF/AbstractSkeleton/Sets/SSetBinding.cpp
--- a/Engine/Plugins/Experimental/UAF/UAF/Source/UAFEditor/Private/UAF/AbstractSkeleton/Sets/SSetBinding.cpp
+++ b/Engine/Plugins/Experimental/UAF/UAF/Source/UAFEditor/Private/UAF/AbstractSkeleton/Sets/SSetBinding.cpp
@@ -563,7 +563,25 @@ void SSetBinding::Construct(const FArguments& InArgs, TWeakObjectPtr<UAbstractSk
 
 	bool SSetBinding::FTreeItem_Set::GenerateContextMenu(FMenuBuilder& Menu)
 	{
-		return false;
+		FUIAction UnbindAll(
+			FExecuteAction::CreateLambda([&]()
+				{
+					FScopedTransaction Transaction(LOCTEXT("UnbindAllInSet", "Unbind All In Set"));
+					SetBindingWidget.SetBinding->Modify();
+
+					SetBindingWidget.SetBinding->RemoveAllFromSet(Set.SetName);
+					SetBindingWidget.RepopulateTreeData(); // TODO: Changes to the set binding should broadcast which 
+					// the widget listens to and dirties the tree instead of manually repopulating.
+				})
+		);
+
+		Menu.AddMenuEntry(
+			LOCTEXT("UnbindAllInSet_Label", "Unbind All In Set"),
+			LOCTEXT("UnbindAllInSet_Description", "Unbinds all attributes currently bound to this set."),
+			FSlateIcon(FAppStyle::GetAppStyleSetName(), "Icons.Unlink"),
+			UnbindAll);
+
+		return true;
 	}
 
 	bool SSetBinding::FTreeItem_Bone::GenerateContextMenu(FMenuBuilder& Menu)
@@ -582,7 +600,7 @@ void SSetBinding::Construct(const FArguments& InArgs, TWeakObjectPtr<UAbstractSk
 		Menu.AddMenuEntry(
 			LOCTEXT("UnbindBone_Label", "Unbind Bone"),
 			LOCTEXT("UnbindBone_Description", "Unbind the selected bone."),
-			FSlateIcon(),
+			FSlateIcon(FAppStyle::GetAppStyleSetName(), "Icons.Unlink"),
 			UnbindBone);
 
 		return true;
