Add additive per-bone blend support to UAF value runtime

#Synced-CL 48953711

diff --git a/Engine/Plugins/Experimental/UAF/UAF/Source/UAF/Private/UAF/Attributes/Transformers/AdditiveSpace.cpp b/Engine/Plugins/Experimental/UAF/UAF/Source/UAF/Private/UAF/Attributes/Transformers/AdditiveSpace.cpp
--- a/Engine/Plugins/Experimental/UAF/UAF/Source/UAF/Private/UAF/Attributes/Transformers/AdditiveSpace.cpp
+++ b/Engine/Plugins/Experimental/UAF/UAF/Source/UAF/Private/UAF/Attributes/Transformers/AdditiveSpace.cpp
@@ -114,6 +114,26 @@ namespace UE::UAF::Transformers
 		const FAttributeCollection& Additive,
 		float AdditiveWeight,
 		FAttributeCollection& Output)
+	{
+		const bool bIsOutputBase = &Base == &Output;
+		if (bIsOutputBase && !FAnimWeight::IsRelevant(AdditiveWeight))
+		{
+			// Our input base is our output and the additive input is not relevant, nothing to do
+			return;
+		}
+
+		// Allocator is irrelevant since it will remain empty
+		FAttributeCollectionHeap EmptyCollection(Additive.GetNamedSet());
+		FApplyAdditiveSpace::Apply(TransformerMap, Base, Additive, EmptyCollection, AdditiveWeight, Output);
+	}
+
+	void FApplyAdditiveSpace::Apply(
+		const FAttributeValueTransformerMapPtr& TransformerMap,
+		const FAttributeCollection& Base,
+		const FAttributeCollection& Additive,
+		const FAttributeCollection& PerValueWeights,
+		float DefaultWeight,
+		FAttributeCollection& Output)
 	{
 		// Named sets must match as it ensures that our inputs have the same sizes / shapes
 		checkf(Base.GetNamedSet() == Additive.GetNamedSet(), TEXT("Base and Additive collections must use the same named set"));
@@ -126,13 +146,6 @@ namespace UE::UAF::Transformers
 		checkf(AdditiveSpace.IsAdditive(), TEXT("Additive input must be in additive space"));
 		checkf(BaseSpace.GetType() == AdditiveSpace.GetType(), TEXT("Additive value space must match the base (local or mesh space additive)"));
 
-		const bool bIsOutputBase = &Base == &Output;
-		if (bIsOutputBase && !FAnimWeight::IsRelevant(AdditiveWeight))
-		{
-			// Our input base is our output and the additive input is not relevant, nothing to do
-			return;
-		}
-
 		const FAttributeValueTransformerList* TransformerList = TransformerMap->Find(TransformerName);
 		if (TransformerList == nullptr)
 		{
@@ -140,6 +153,9 @@ namespace UE::UAF::Transformers
 			return;
 		}
 
+		UScriptStruct* FloatValueType = FFloatAnimationAttribute::StaticStruct();
+
+		const bool bIsOutputBase = &Base == &Output;
 		const bool bIsOutputAdditive = &Additive == &Output;
 		if (bIsOutputBase || bIsOutputAdditive)
 		{
@@ -153,7 +169,7 @@ namespace UE::UAF::Transformers
 
 				OuterJoinBy(
 					FAttributeValueTransformerListWithAttributeCollectionSetMapJoinOp(),
-					[AdditiveWeight, bIsOutputBase, OutputAllocator, &PendingSetMaps](FRawTransformerFunc Transformer, FAttributeSetMap* BaseMap, FAttributeSetMap* AdditiveMap)
+					[&PerValueWeights, DefaultWeight, bIsOutputBase, FloatValueType, OutputAllocator, &PendingSetMaps](FRawTransformerFunc Transformer, FAttributeSetMap* BaseMap, FAttributeSetMap* AdditiveMap)
 					{
 						if (AdditiveMap == nullptr)
 						{
@@ -167,10 +183,9 @@ namespace UE::UAF::Transformers
 							if (BaseMap == nullptr)
 							{
 								// We have no base, duplicate the additive
-								FAttributeSetMap* SetMapCopy = AdditiveMap->Duplicate(OutputAllocator);
-								PendingSetMaps.Add(SetMapCopy);
+								PendingSetMaps.Add(AdditiveMap->Duplicate(OutputAllocator));
 							}
-							else if (AdditiveWeight >= 0.5f)
+							else if (DefaultWeight >= 0.5f)
 							{
 								// Additive has more weight
 								if (bIsOutputBase)
@@ -204,8 +219,12 @@ namespace UE::UAF::Transformers
 						FAttributeSetMap* OutputSetMap = bIsOutputBase ? BaseMap : AdditiveMap;
 						const bool bIsOutputMissing = OutputSetMap == nullptr;
 
+						// We cannot join our weights because their value type differs
+						const FAttributeMappingKey MappingKey = AdditiveMap->GetMappingKey().To(FloatValueType);
+						const FAttributeSetMap* PerValueWeightMap = PerValueWeights.FindSetMap(MappingKey);
+
 						FTransformSetMapFunc TypedTransformer = reinterpret_cast<FTransformSetMapFunc>(Transformer);
-						(*TypedTransformer)(BaseMap, AdditiveMap, AdditiveWeight, OutputSetMap, OutputAllocator);
+						(*TypedTransformer)(BaseMap, AdditiveMap, PerValueWeightMap, DefaultWeight, OutputSetMap, OutputAllocator);
 
 						if (bIsOutputMissing)
 						{
@@ -234,7 +253,7 @@ namespace UE::UAF::Transformers
 
 				OuterJoinBy(
 					FAttributeValueTransformerListWithAttributeCollectionMapJoinOp(),
-					[AdditiveWeight, bIsOutputBase, OutputAllocator, &PendingMaps](FRawTransformerFunc Transformer, FAttributeMap* BaseMap, FAttributeMap* AdditiveMap)
+					[DefaultWeight, bIsOutputBase, OutputAllocator, &PendingMaps](FRawTransformerFunc Transformer, FAttributeMap* BaseMap, FAttributeMap* AdditiveMap)
 					{
 						if (AdditiveMap == nullptr)
 						{
@@ -251,15 +270,14 @@ namespace UE::UAF::Transformers
 								if (bIsOutputBase)
 								{
 									// Duplicate the additive
-									FAttributeMap* MapCopy = AdditiveMap->Duplicate(OutputAllocator);
-									PendingMaps.Add(MapCopy);
+									PendingMaps.Add(AdditiveMap->Duplicate(OutputAllocator));
 								}
 								else
 								{
 									// Our output is the additive, it already contains it
 								}
 							}
-							else if (AdditiveWeight >= 0.5f)
+							else if (DefaultWeight >= 0.5f)
 							{
 								// Additive has more weight or we have no base
 								if (bIsOutputBase)
@@ -293,7 +311,7 @@ namespace UE::UAF::Transformers
 						FAttributeMap* OutputMap = nullptr;
 
 						FTransformMapFunc TypedTransformer = reinterpret_cast<FTransformMapFunc>(Transformer);
-						(*TypedTransformer)(BaseMap, AdditiveMap, AdditiveWeight, OutputMap, OutputAllocator);
+						(*TypedTransformer)(BaseMap, AdditiveMap, DefaultWeight, OutputMap, OutputAllocator);
 
 						checkf(OutputMap != nullptr, TEXT("Transformer function is responsible for allocating an output buffer"));
 						checkf(OutputMap->GetAllocator() == OutputAllocator, TEXT("Transformer function is responsible for allocating using the output allocator"));
@@ -344,7 +362,7 @@ namespace UE::UAF::Transformers
 			// Set maps
 			OuterJoinBy(
 				FAttributeValueTransformerListWithAttributeCollectionSetMapJoinOp(),
-				[AdditiveWeight, &Output](FRawTransformerFunc Transformer, const FAttributeSetMap* BaseMap, const FAttributeSetMap* AdditiveMap)
+				[&PerValueWeights, DefaultWeight, FloatValueType, &Output](FRawTransformerFunc Transformer, const FAttributeSetMap* BaseMap, const FAttributeSetMap* AdditiveMap)
 				{
 					if (AdditiveMap == nullptr)
 					{
@@ -356,7 +374,7 @@ namespace UE::UAF::Transformers
 					{
 						// No transformer was found for this value type, retain the one with the largest weight
 						FAttributeSetMap* MapCopy;
-						if (BaseMap == nullptr || AdditiveWeight >= 0.5f)
+						if (BaseMap == nullptr || DefaultWeight >= 0.5f)
 						{
 							// Additive map has more weight or we have no base
 							MapCopy = AdditiveMap->Duplicate(Output.GetAllocator());
@@ -375,8 +393,12 @@ namespace UE::UAF::Transformers
 					// If it is missing, transformer specializations should use a default value
 					FAttributeSetMap* OutputSetMap = nullptr;
 
+					// We cannot join our weights because their value type differs
+					const FAttributeMappingKey MappingKey = AdditiveMap->GetMappingKey().To(FloatValueType);
+					const FAttributeSetMap* PerValueWeightMap = PerValueWeights.FindSetMap(MappingKey);
+
 					FTransformSetMapFunc TypedTransformer = reinterpret_cast<FTransformSetMapFunc>(Transformer);
-					(*TypedTransformer)(BaseMap, AdditiveMap, AdditiveWeight, OutputSetMap, Output.GetAllocator());
+					(*TypedTransformer)(BaseMap, AdditiveMap, PerValueWeightMap, DefaultWeight, OutputSetMap, Output.GetAllocator());
 
 					checkf(OutputSetMap != nullptr, TEXT("Transformer function is responsible for allocating an output buffer"));
 					checkf(OutputSetMap->GetAllocator() == Output.GetAllocator(), TEXT("Transformer function is responsible for allocating using the output allocator"));
@@ -389,7 +411,7 @@ namespace UE::UAF::Transformers
 			// Name/value maps
 			OuterJoinBy(
 				FAttributeValueTransformerListWithAttributeCollectionMapJoinOp(),
-				[AdditiveWeight, &Output](FRawTransformerFunc Transformer, const FAttributeMap* BaseMap, const FAttributeMap* AdditiveMap)
+				[DefaultWeight, &Output](FRawTransformerFunc Transformer, const FAttributeMap* BaseMap, const FAttributeMap* AdditiveMap)
 				{
 					if (AdditiveMap == nullptr)
 					{
@@ -401,7 +423,7 @@ namespace UE::UAF::Transformers
 					{
 						// No transformer was found for this value type, retain the one with the largest weight
 						FAttributeMap* MapCopy;
-						if (BaseMap == nullptr || AdditiveWeight >= 0.5f)
+						if (BaseMap == nullptr || DefaultWeight >= 0.5f)
 						{
 							// Additive map has more weight or we have no base
 							MapCopy = AdditiveMap->Duplicate(Output.GetAllocator());
@@ -421,7 +443,7 @@ namespace UE::UAF::Transformers
 					FAttributeMap* OutputMap = nullptr;
 
 					FTransformMapFunc TypedTransformer = reinterpret_cast<FTransformMapFunc>(Transformer);
-					(*TypedTransformer)(BaseMap, AdditiveMap, AdditiveWeight, OutputMap, Output.GetAllocator());
+					(*TypedTransformer)(BaseMap, AdditiveMap, DefaultWeight, OutputMap, Output.GetAllocator());
 
 					checkf(OutputMap != nullptr, TEXT("Transformer function is responsible for allocating an output buffer"));
 					checkf(OutputMap->GetAllocator() == Output.GetAllocator(), TEXT("Transformer function is responsible for allocating using the output allocator"));
@@ -619,7 +641,7 @@ namespace UE::UAF::Transformers
 		//////////////////////////////////////////////////////////////////////////
 		// ApplyAdditiveSpace transformers
 
-		void FApplyAdditiveSpace_BoneTransformAttribute::TransformSetMap(const FAttributeSetMap* Base, const FAttributeSetMap* Additive, float AdditiveWeight, FAttributeSetMap*& Output, FReallocFun OutputAllocator)
+		void FApplyAdditiveSpace_BoneTransformAttribute::TransformSetMap(const FAttributeSetMap* Base, const FAttributeSetMap* Additive, const FAttributeSetMap* PerValueWeightMap, float DefaultWeight, FAttributeSetMap*& Output, FReallocFun OutputAllocator)
 		{
 			if (Base == nullptr)
 			{
@@ -659,37 +681,60 @@ namespace UE::UAF::Transformers
 			FBoneTransformAnimationAttribute* OutputTransformPtr = OutputBoneTransforms->GetData();
 			const FBoneTransformAnimationAttribute* BaseTransformPtrEnd = BaseTransformPtr + BaseBoneTransforms->Num();
 
-			const ScalarRegister VAdditiveWeight(AdditiveWeight);
-
-			if (FAnimWeight::IsFullWeight(AdditiveWeight))
+			if (const TAttributeSetMap<FFloatAnimationAttribute>* PerValueWeights = Cast<FFloatAnimationAttribute>(PerValueWeightMap))
 			{
-				// Our additive has full weight, no need to interpolate with identity
+				check(Base->GetTypedSet() == PerValueWeights->GetTypedSet());
+				check(Base->Num() == PerValueWeights->Num());
+
+				const FFloatAnimationAttribute* PerValueWeightsPtr = PerValueWeights->GetData();
+
 				while (BaseTransformPtr < BaseTransformPtrEnd)
 				{
-					// TODO: The VAdditiveWeight is practically 1.0, can assume that it is?
 					FTransform Result = BaseTransformPtr->Value;
-					Result.AccumulateWithAdditiveScale(AdditiveTransformPtr->Value, VAdditiveWeight);
+					FTransform::BlendFromIdentityAndAccumulate(Result, AdditiveTransformPtr->Value, ScalarRegister(PerValueWeightsPtr->Value));
 
 					OutputTransformPtr->Value = Result;
 
 					BaseTransformPtr++;
 					AdditiveTransformPtr++;
+					PerValueWeightsPtr++;
 					OutputTransformPtr++;
 				}
 			}
 			else
 			{
-				// Interpolate with identity before adding
-				while (BaseTransformPtr < BaseTransformPtrEnd)
+				const ScalarRegister VDefaultWeight(DefaultWeight);
+
+				if (FAnimWeight::IsFullWeight(DefaultWeight))
 				{
-					FTransform Result = BaseTransformPtr->Value;
-					FTransform::BlendFromIdentityAndAccumulate(Result, AdditiveTransformPtr->Value, VAdditiveWeight);
+					// Our additive has full weight, no need to interpolate with identity
+					while (BaseTransformPtr < BaseTransformPtrEnd)
+					{
+						// TODO: The VDefaultWeight is practically 1.0, can assume that it is?
+						FTransform Result = BaseTransformPtr->Value;
+						Result.AccumulateWithAdditiveScale(AdditiveTransformPtr->Value, VDefaultWeight);
 
-					OutputTransformPtr->Value = Result;
+						OutputTransformPtr->Value = Result;
 
-					BaseTransformPtr++;
-					AdditiveTransformPtr++;
-					OutputTransformPtr++;
+						BaseTransformPtr++;
+						AdditiveTransformPtr++;
+						OutputTransformPtr++;
+					}
+				}
+				else
+				{
+					// Interpolate with identity before adding
+					while (BaseTransformPtr < BaseTransformPtrEnd)
+					{
+						FTransform Result = BaseTransformPtr->Value;
+						FTransform::BlendFromIdentityAndAccumulate(Result, AdditiveTransformPtr->Value, VDefaultWeight);
+
+						OutputTransformPtr->Value = Result;
+
+						BaseTransformPtr++;
+						AdditiveTransformPtr++;
+						OutputTransformPtr++;
+					}
 				}
 			}
 		}
@@ -754,7 +799,7 @@ namespace UE::UAF::Transformers
 			Output = OutputBoneTransforms;
 		}
 
-		void FApplyAdditiveSpace_FloatAttribute::TransformSetMap(const FAttributeSetMap* Base, const FAttributeSetMap* Additive, float AdditiveWeight, FAttributeSetMap*& Output, FReallocFun OutputAllocator)
+		void FApplyAdditiveSpace_FloatAttribute::TransformSetMap(const FAttributeSetMap* Base, const FAttributeSetMap* Additive, const FAttributeSetMap* PerValueWeightMap, float DefaultWeight, FAttributeSetMap*& Output, FReallocFun OutputAllocator)
 		{
 			if (Base == nullptr)
 			{
@@ -794,13 +839,33 @@ namespace UE::UAF::Transformers
 			FFloatAnimationAttribute* OutputFloatsPtr = OutputFloats->GetData();
 			const FFloatAnimationAttribute* BaseFloatsPtrEnd = BaseFloatsPtr + BaseFloats->Num();
 
-			while (BaseFloatsPtr < BaseFloatsPtrEnd)
+			if (const TAttributeSetMap<FFloatAnimationAttribute>* PerValueWeights = Cast<FFloatAnimationAttribute>(PerValueWeightMap))
 			{
-				OutputFloatsPtr->Value = BaseFloatsPtr->Value + (AdditiveWeight * AdditiveFloatsPtr->Value);
+				check(Base->GetTypedSet() == PerValueWeights->GetTypedSet());
+				check(Base->Num() == PerValueWeights->Num());
 
-				BaseFloatsPtr++;
-				AdditiveFloatsPtr++;
-				OutputFloatsPtr++;
+				const FFloatAnimationAttribute* PerValueWeightsPtr = PerValueWeights->GetData();
+
+				while (BaseFloatsPtr < BaseFloatsPtrEnd)
+				{
+					OutputFloatsPtr->Value = BaseFloatsPtr->Value + (PerValueWeightsPtr->Value * AdditiveFloatsPtr->Value);
+
+					BaseFloatsPtr++;
+					AdditiveFloatsPtr++;
+					PerValueWeightsPtr++;
+					OutputFloatsPtr++;
+				}
+			}
+			else
+			{
+				while (BaseFloatsPtr < BaseFloatsPtrEnd)
+				{
+					OutputFloatsPtr->Value = BaseFloatsPtr->Value + (DefaultWeight * AdditiveFloatsPtr->Value);
+
+					BaseFloatsPtr++;
+					AdditiveFloatsPtr++;
+					OutputFloatsPtr++;
+				}
 			}
 		}
 
diff --git a/Engine/Plugins/Experimental/UAF/UAF/Source/UAF/Public/UAF/Attributes/Transformers/AdditiveSpace.h b/Engine/Plugins/Experimental/UAF/UAF/Source/UAF/Public/UAF/Attributes/Transformers/AdditiveSpace.h
--- a/Engine/Plugins/Experimental/UAF/UAF/Source/UAF/Public/UAF/Attributes/Transformers/AdditiveSpace.h
+++ b/Engine/Plugins/Experimental/UAF/UAF/Source/UAF/Public/UAF/Attributes/Transformers/AdditiveSpace.h
@@ -36,10 +36,13 @@ namespace UE::UAF::Transformers
 
 		// If the output buffer is nullptr, an allocator is provided and the implementation is responsible for allocating the output buffer and returning it
 		// If the base map is missing, transformers should use a default value
-		using FTransformSetMapFunc = void (*)(const FAttributeSetMap* Base, const FAttributeSetMap* Additive, float AdditiveWeight, FAttributeSetMap*& Output, FReallocFun OutputAllocator);
+		// Per-value map is optional, when not present, the default weight should be used
+		// Unbound values do not support per-value weights
+		using FTransformSetMapFunc = void (*)(const FAttributeSetMap* Base, const FAttributeSetMap* Additive, const FAttributeSetMap* PerValueWeightMap, float DefaultWeight, FAttributeSetMap*& Output, FReallocFun OutputAllocator);
 		using FTransformMapFunc = void (*)(const FAttributeMap* Base, const FAttributeMap* Additive, float AdditiveWeight, FAttributeMap*& Output, FReallocFun OutputAllocator);
 
 		// Applies an additive collection onto a base by adding it
+		// Output = Base + Lerp(Identity, Additive, Weight)
 		// This function supports outputting to either input
 		static UE_API void Apply(
 			const FAttributeValueTransformerMapPtr& TransformerMap,
@@ -47,6 +50,18 @@ namespace UE::UAF::Transformers
 			const FAttributeCollection& Additive,
 			float AdditiveWeight,
 			FAttributeCollection& Output);
+
+		// Applies an additive collection onto a base by adding it using per value weights (must be of type FFloatAnimationAttribute)
+		// If a weight isn't specified, the default weight will be used
+		// Output = Base + Lerp(Identity, Additive, HasValueWeight ? ValueWeight : DefaultWeight)
+		// This function supports outputting to either input
+		static UE_API void Apply(
+			const FAttributeValueTransformerMapPtr& TransformerMap,
+			const FAttributeCollection& Base,
+			const FAttributeCollection& Additive,
+			const FAttributeCollection& PerValueWeights,
+			float DefaultWeight,
+			FAttributeCollection& Output);
 	};
 
 	//////////////////////////////////////////////////////////////////////////
@@ -77,7 +92,7 @@ namespace UE::UAF::Transformers
 			using TransformerType = FApplyAdditiveSpace;
 			using ValueType = FBoneTransformAnimationAttribute;
 
-			static void TransformSetMap(const FAttributeSetMap* Base, const FAttributeSetMap* Additive, float AdditiveWeight, FAttributeSetMap*& Output, FReallocFun OutputAllocator);
+			static void TransformSetMap(const FAttributeSetMap* Base, const FAttributeSetMap* Additive, const FAttributeSetMap* PerValueWeightMap, float DefaultWeight, FAttributeSetMap*& Output, FReallocFun OutputAllocator);
 			static void TransformMap(const FAttributeMap* Base, const FAttributeMap* Additive, float AdditiveWeight, FAttributeMap*& Output, FReallocFun OutputAllocator);
 		};
 
@@ -86,7 +101,7 @@ namespace UE::UAF::Transformers
 			using TransformerType = FApplyAdditiveSpace;
 			using ValueType = FFloatAnimationAttribute;
 
-			static void TransformSetMap(const FAttributeSetMap* Base, const FAttributeSetMap* Additive, float AdditiveWeight, FAttributeSetMap*& Output, FReallocFun OutputAllocator);
+			static void TransformSetMap(const FAttributeSetMap* Base, const FAttributeSetMap* Additive, const FAttributeSetMap* PerValueWeightMap, float DefaultWeight, FAttributeSetMap*& Output, FReallocFun OutputAllocator);
 			static void TransformMap(const FAttributeMap* Base, const FAttributeMap* Additive, float AdditiveWeight, FAttributeMap*& Output, FReallocFun OutputAllocator);
 		};
 	}
diff --git a/Engine/Plugins/Experimental/UAF/UAFAnimGraph/Source/UAFAnimGraph/Private/EvaluationVM/Tasks/ApplyAdditiveKeyframePerBone.cpp b/Engine/Plugins/Experimental/UAF/UAFAnimGraph/Source/UAFAnimGraph/Private/EvaluationVM/Tasks/ApplyAdditiveKeyframePerBone.cpp
--- a/Engine/Plugins/Experimental/UAF/UAFAnimGraph/Source/UAFAnimGraph/Private/EvaluationVM/Tasks/ApplyAdditiveKeyframePerBone.cpp
+++ b/Engine/Plugins/Experimental/UAF/UAFAnimGraph/Source/UAFAnimGraph/Private/EvaluationVM/Tasks/ApplyAdditiveKeyframePerBone.cpp
@@ -6,6 +6,9 @@
 #include "EvaluationVM/KeyframeState.h"
 #include "TransformArrayOperations.h"
 #include "Animation/AnimationAsset.h"
+#include "UAF/Attributes/AttributeCollection.h"
+#include "UAF/Attributes/Transformers/AdditiveSpace.h"
+#include "UAF/Attributes/Transformers/BoneSpace.h"
 
 #include UE_INLINE_GENERATED_CPP_BY_NAME(ApplyAdditiveKeyframePerBone)
 
@@ -42,88 +45,185 @@ void FUAFApplyAdditiveKeyframePerBoneTask::Execute(UE::UAF::FEvaluationVM& VM) c
 
 	if (VM.GetActiveNamedSet())
 	{
-		// TODO: Implement with new attribute runtime
-		TUniquePtr<FAttributeCollection> DiscardedCollection;
-		(void)VM.PopValue(ATTRIBUTE_STACK_NAME, DiscardedCollection);
-		return;
-	}
+		// Pop our top two poses, the top one being the additive pose and the second one the base
+		TUniquePtr<FAttributeCollection> AdditiveInput;
+		if (!VM.PopValue(ATTRIBUTE_STACK_NAME, AdditiveInput))
+		{
+			// We have no inputs, nothing to do
+			return;
+		}
 
-	// Pop our top two poses, the top one being the additive pose and the second one the base
-	TUniquePtr<FKeyframeState> AdditiveKeyframe;
-	if (!VM.PopValue(KEYFRAME_STACK_NAME, AdditiveKeyframe))
-	{
-		// We have no inputs, nothing to do
-		return;
-	}
+		TUniquePtr<FAttributeCollection> BaseInput;
+		if (!VM.PopValue(ATTRIBUTE_STACK_NAME, BaseInput))
+		{
+			// We have a single input, discard it since it must be the additive pose, either way something went wrong
+			// Push the reference pose since we'll expect a non-additive pose
+			FAttributeCollectionStack Collection(VM.GetActiveNamedSet());
+			Collection.InitWithValueSpace(FValueSpace(EValueSpaceType::Local));
+
+			VM.PushValue(ATTRIBUTE_STACK_NAME, MakeUnique<FAttributeCollection>(MoveTemp(Collection)));
+			UE_LOG(LogAnimation, Warning, TEXT("FUAFApplyAdditiveKeyframePerBoneTask::Execute: Could not apply additive keyframe - only a single input was provided. Pushing a ref pose instead."));
+			return;
+		}
 
-	TUniquePtr<FKeyframeState> BaseKeyframe;
-	if (!VM.PopValue(KEYFRAME_STACK_NAME, BaseKeyframe))
-	{
-		// We have a single input, discard it since it must be the additive pose, either way something went wrong
-		// Push the reference pose since we'll expect a non-additive pose
-		VM.PushValue(KEYFRAME_STACK_NAME, MakeUnique<FKeyframeState>(VM.MakeReferenceKeyframe(false)));
-		UE_LOG(LogAnimation, Warning, TEXT("FUAFApplyAdditiveKeyframePerBoneTask::Execute: Could not apply additive keyframe - only a single input was provided. Pushing a ref pose instead."));
-		return;
-	}
+		const FValueSpace AdditiveSpace = AdditiveInput->GetValueSpace();
+		if (!AdditiveSpace.IsAdditive())
+		{
+			// Additive input must be additive, push reference pose if not the case
+			FAttributeCollectionStack Collection(VM.GetActiveNamedSet());
+			Collection.InitWithValueSpace(FValueSpace(EValueSpaceType::Local));
 
-	if (!AdditiveKeyframe->Pose.IsAdditive())
-	{
-		// Additive must be additive type, push reference pose if not the case
-		VM.PushValue(KEYFRAME_STACK_NAME, MakeUnique<FKeyframeState>(VM.MakeReferenceKeyframe(false)));
-		UE_LOG(LogAnimation, Warning, TEXT("FUAFApplyAdditiveKeyframePerBoneTask::Execute: Could not apply additive keyframe - The expected additive pose was not of an additive type. Pushing a ref pose instead."));
-		return;
-	}
+			VM.PushValue(ATTRIBUTE_STACK_NAME, MakeUnique<FAttributeCollection>(MoveTemp(Collection)));
+			UE_LOG(LogAnimation, Warning, TEXT("FUAFApplyAdditiveKeyframePerBoneTask::Execute: Could not apply additive keyframe - The expected additive pose was not of an additive type. Pushing a ref pose instead."));
+			return;
+		}
 
-	const float AdditiveWeight = GetInterpolationAlpha(BaseKeyframe.Get(), AdditiveKeyframe.Get());
+		const float AdditiveWeight = GetInterpolationAlpha(*BaseInput, *AdditiveInput);
 
-	const TArrayView<const FBoneIndexType> LODBoneIndexToSkeletonBoneIndexMap = BaseKeyframe->Pose.GetLODBoneIndexToSkeletonBoneIndexMap();
-	const int32 NumLODBones = LODBoneIndexToSkeletonBoneIndexMap.Num();
+		UScriptStruct* FloatAttributeType = FFloatAnimationAttribute::StaticStruct();
 
-	TArray<int32> LODBoneIndexToWeightIndexMap;
-	LODBoneIndexToWeightIndexMap.AddUninitialized(NumLODBones);
+		// Build our per-value weights (only for bound values)
+		FAttributeCollectionStack PerValueWeights(VM.GetActiveNamedSet());
 
-	for (int32 LODBoneIndex = 0; LODBoneIndex < NumLODBones; ++LODBoneIndex)
-	{
-		if (bPerBoneDataUsesSkeletonIndex)
+		for (auto It = BaseInput->CreateConstSetMapIterator(); It; ++It)
 		{
-			const FSkeletonPoseBoneIndex TargetSkeletonBoneIndex(LODBoneIndexToSkeletonBoneIndexMap[LODBoneIndex]);
-			LODBoneIndexToWeightIndexMap[LODBoneIndex] = TargetSkeletonBoneIndex.GetInt();
+			UScriptStruct* AttributeType = It.GetAttributeType();
+
+			// Every attribute type maps to a float, our weight
+			FAttributeMappingKey MappingKey = FAttributeMappingKey::MakeFromTo(AttributeType, FloatAttributeType);
+
+			TAttributeSetMap<FFloatAnimationAttribute>* SetWeights = PerValueWeights.AddSetMap<FFloatAnimationAttribute>(MappingKey);
+
+			if (const TAttributeSetMap<FBoneTransformAnimationAttribute>* BoneMap = Cast<FBoneTransformAnimationAttribute>(It.GetMap()))
+			{
+				const FAttributeTypedSetPtr& BoneSet = BoneMap->GetTypedSet();
+				const int32 NumBonesInSet = BoneMap->Num();
+
+				check(BlendData->PerBoneBlendData.Num() >= NumBonesInSet);
+
+				for (FAttributeSetIndex SetIndex(0); SetIndex < NumBonesInSet; ++SetIndex)
+				{
+					if (bPerBoneDataUsesSkeletonIndex)
+					{
+						const FAttributeBindingIndex BindingIndex = BoneSet->GetBindingIndex(SetIndex);
+
+						const FSkeletonPoseBoneIndex TargetSkeletonBoneIndex = BindingIndex;
+						(*SetWeights)[SetIndex].Value = TargetSkeletonBoneIndex.IsValid() ? BlendData->PerBoneBlendData[TargetSkeletonBoneIndex.GetInt()] : AdditiveWeight;
+					}
+					else
+					{
+						(*SetWeights)[SetIndex].Value = BlendData->PerBoneBlendData[SetIndex.GetInt()];
+					}
+				}
+			}
+			else
+			{
+				// Everything else defaults to full weight
+				SetWeights->FillWithValue(FFloatAnimationAttribute{ 1.0f });
+			}
 		}
-		else
+
+		bool bConvertBaseToMeshSpace = false;
+		if (AdditiveSpace.GetMixedSpaceFlags() == EMixedSpaceFlags::MeshRotation)
 		{
-			LODBoneIndexToWeightIndexMap[LODBoneIndex] = LODBoneIndex;
+			// If our additive input is in mesh rotation space, we ensure our base is in that space as well
+			const FValueSpace BaseSpace = BaseInput->GetValueSpace();
+			bConvertBaseToMeshSpace = BaseSpace.GetMixedSpaceFlags() != EMixedSpaceFlags::MeshRotation;
 		}
-	}
 
-	if (EnumHasAnyFlags(VM.GetFlags(), EEvaluationFlags::Bones))
+		if (bConvertBaseToMeshSpace)
+		{
+			Transformers::FBoneSpace::LocalToMeshRotation(*BaseInput, *BaseInput);
+		}
+
+		Transformers::FApplyAdditiveSpace::Apply(VM.GetTransformerMap(), *BaseInput, *AdditiveInput, PerValueWeights, AdditiveWeight, *BaseInput);
+
+		if (bConvertBaseToMeshSpace)
+		{
+			Transformers::FBoneSpace::MeshRotationToLocal(*BaseInput, *BaseInput);
+		}
+
+		VM.PushValue(ATTRIBUTE_STACK_NAME, MoveTemp(BaseInput));
+	}
+	else
 	{
-		check(BaseKeyframe->Pose.GetNumBones() == AdditiveKeyframe->Pose.GetNumBones());
+		// Pop our top two poses, the top one being the additive pose and the second one the base
+		TUniquePtr<FKeyframeState> AdditiveKeyframe;
+		if (!VM.PopValue(KEYFRAME_STACK_NAME, AdditiveKeyframe))
+		{
+			// We have no inputs, nothing to do
+			return;
+		}
+
+		TUniquePtr<FKeyframeState> BaseKeyframe;
+		if (!VM.PopValue(KEYFRAME_STACK_NAME, BaseKeyframe))
+		{
+			// We have a single input, discard it since it must be the additive pose, either way something went wrong
+			// Push the reference pose since we'll expect a non-additive pose
+			VM.PushValue(KEYFRAME_STACK_NAME, MakeUnique<FKeyframeState>(VM.MakeReferenceKeyframe(false)));
+			UE_LOG(LogAnimation, Warning, TEXT("FUAFApplyAdditiveKeyframePerBoneTask::Execute: Could not apply additive keyframe - only a single input was provided. Pushing a ref pose instead."));
+			return;
+		}
+
+		if (!AdditiveKeyframe->Pose.IsAdditive())
+		{
+			// Additive must be additive type, push reference pose if not the case
+			VM.PushValue(KEYFRAME_STACK_NAME, MakeUnique<FKeyframeState>(VM.MakeReferenceKeyframe(false)));
+			UE_LOG(LogAnimation, Warning, TEXT("FUAFApplyAdditiveKeyframePerBoneTask::Execute: Could not apply additive keyframe - The expected additive pose was not of an additive type. Pushing a ref pose instead."));
+			return;
+		}
 
-		const FTransformArrayView BaseTransformsView = BaseKeyframe->Pose.LocalTransforms.GetView();
+		const float AdditiveWeight = GetInterpolationAlpha(BaseKeyframe.Get(), AdditiveKeyframe.Get());
 
-		if (AdditiveKeyframe->Pose.IsMeshSpaceAdditive())
+		const TArrayView<const FBoneIndexType> LODBoneIndexToSkeletonBoneIndexMap = BaseKeyframe->Pose.GetLODBoneIndexToSkeletonBoneIndexMap();
+		const int32 NumLODBones = LODBoneIndexToSkeletonBoneIndexMap.Num();
+
+		TArray<int32> LODBoneIndexToWeightIndexMap;
+		LODBoneIndexToWeightIndexMap.AddUninitialized(NumLODBones);
+
+		for (int32 LODBoneIndex = 0; LODBoneIndex < NumLODBones; ++LODBoneIndex)
 		{
-			BlendWithIdentityAndAccumulatePerBoneWeightMesh(
-				BaseTransformsView, AdditiveKeyframe->Pose.LocalTransforms.GetConstView(),
-				AdditiveKeyframe->Pose.GetLODBoneIndexToParentLODBoneIndexMap(), LODBoneIndexToWeightIndexMap, BlendData->PerBoneBlendData, AdditiveWeight);
+			if (bPerBoneDataUsesSkeletonIndex)
+			{
+				const FSkeletonPoseBoneIndex TargetSkeletonBoneIndex(LODBoneIndexToSkeletonBoneIndexMap[LODBoneIndex]);
+				LODBoneIndexToWeightIndexMap[LODBoneIndex] = TargetSkeletonBoneIndex.GetInt();
+			}
+			else
+			{
+				LODBoneIndexToWeightIndexMap[LODBoneIndex] = LODBoneIndex;
+			}
 		}
-		else
+
+		if (EnumHasAnyFlags(VM.GetFlags(), EEvaluationFlags::Bones))
 		{
-			BlendWithIdentityAndAccumulatePerBoneWeight(BaseTransformsView, AdditiveKeyframe->Pose.LocalTransforms.GetConstView(), LODBoneIndexToWeightIndexMap, BlendData->PerBoneBlendData, AdditiveWeight);
+			check(BaseKeyframe->Pose.GetNumBones() == AdditiveKeyframe->Pose.GetNumBones());
+
+			const FTransformArrayView BaseTransformsView = BaseKeyframe->Pose.LocalTransforms.GetView();
+
+			if (AdditiveKeyframe->Pose.IsMeshSpaceAdditive())
+			{
+				BlendWithIdentityAndAccumulatePerBoneWeightMesh(
+					BaseTransformsView, AdditiveKeyframe->Pose.LocalTransforms.GetConstView(),
+					AdditiveKeyframe->Pose.GetLODBoneIndexToParentLODBoneIndexMap(), LODBoneIndexToWeightIndexMap, BlendData->PerBoneBlendData, AdditiveWeight);
+			}
+			else
+			{
+				BlendWithIdentityAndAccumulatePerBoneWeight(BaseTransformsView, AdditiveKeyframe->Pose.LocalTransforms.GetConstView(), LODBoneIndexToWeightIndexMap, BlendData->PerBoneBlendData, AdditiveWeight);
+			}
+
+			NormalizeRotations(BaseTransformsView);
 		}
 
-		NormalizeRotations(BaseTransformsView);
-	}
+		if (EnumHasAnyFlags(VM.GetFlags(), EEvaluationFlags::Curves))
+		{
+			BaseKeyframe->Curves.Accumulate(AdditiveKeyframe->Curves, AdditiveWeight);
+		}
 
-	if (EnumHasAnyFlags(VM.GetFlags(), EEvaluationFlags::Curves))
-	{
-		BaseKeyframe->Curves.Accumulate(AdditiveKeyframe->Curves, AdditiveWeight);
-	}
+		if (EnumHasAnyFlags(VM.GetFlags(), EEvaluationFlags::Attributes))
+		{
+			UE::Anim::Attributes::AccumulateAttributes(AdditiveKeyframe->Attributes, BaseKeyframe->Attributes, AdditiveWeight, AAT_None);
+		}
 
-	if (EnumHasAnyFlags(VM.GetFlags(), EEvaluationFlags::Attributes))
-	{
-		UE::Anim::Attributes::AccumulateAttributes(AdditiveKeyframe->Attributes, BaseKeyframe->Attributes, AdditiveWeight, AAT_None);
+		VM.PushValue(KEYFRAME_STACK_NAME, MoveTemp(BaseKeyframe));
 	}
-
-	VM.PushValue(KEYFRAME_STACK_NAME, MoveTemp(BaseKeyframe));
 }
