Add delegate to set binding asset when the bindings change which widgets can listen to instead of manually reconstructing data on edit

#Synced-CL 48947880

diff --git a/Engine/Plugins/Experimental/UAF/UAF/Source/UAF/Private/UAF/AbstractSkeleton/AbstractSkeletonSetBinding.cpp b/Engine/Plugins/Experimental/UAF/UAF/Source/UAF/Private/UAF/AbstractSkeleton/AbstractSkeletonSetBinding.cpp
--- a/Engine/Plugins/Experimental/UAF/UAF/Source/UAF/Private/UAF/AbstractSkeleton/AbstractSkeletonSetBinding.cpp
+++ b/Engine/Plugins/Experimental/UAF/UAF/Source/UAF/Private/UAF/AbstractSkeleton/AbstractSkeletonSetBinding.cpp
@@ -64,7 +64,7 @@ bool UAbstractSkeletonSetBinding::AddBoneToSet(const FName BoneName, const FName
 	FAbstractSkeleton_BoneBinding& Binding = BoneBindings.AddDefaulted_GetRef();
 	Binding.BoneName = BoneName;
 	Binding.SetName = SetName;
-	RebuildBindingData();
+	OnBindingsChanged();
 
 	return true;
 }
@@ -79,7 +79,7 @@ bool UAbstractSkeletonSetBinding::RemoveBoneFromSet(const FName BoneName)
 	if (BindingIndex != INDEX_NONE)
 	{
 		BoneBindings.RemoveAtSwap(BindingIndex);
-		RebuildBindingData();
+		OnBindingsChanged();
 		return true;
 	}
 	else
@@ -148,7 +148,7 @@ bool UAbstractSkeletonSetBinding::AddAttributeToSet(const FAnimationAttributeIde
 		if (ensure(UnboundAttribute))
 		{
 			UnboundAttribute->SetName = SetName;
-			RebuildBindingData();
+			OnBindingsChanged();
 			return true;
 		}
 		else
@@ -160,7 +160,7 @@ bool UAbstractSkeletonSetBinding::AddAttributeToSet(const FAnimationAttributeIde
 	FAbstractSkeleton_AttributeBinding& Binding = AttributeBindings.AddDefaulted_GetRef();
 	Binding.Attribute = Attribute;
 	Binding.SetName = SetName;
-	RebuildBindingData();
+	OnBindingsChanged();
 
 	return true;
 }
@@ -185,7 +185,7 @@ bool UAbstractSkeletonSetBinding::RemoveAttributeFromSet(const FAnimationAttribu
 			Binding.SetName = NAME_None;
 		}
 
-		RebuildBindingData();
+		OnBindingsChanged();
 		return true;
 	}
 	else
@@ -208,7 +208,13 @@ bool UAbstractSkeletonSetBinding::RemoveAllFromSet(const FName SetName)
 			return Binding.SetName == SetName;
 		});
 
-	return bCountRemoved > 0;
+	if (bCountRemoved > 0)
+	{
+		OnBindingsChanged();
+		return true;
+	}
+
+	return false;
 }
 
 bool UAbstractSkeletonSetBinding::IsAttributeInSet(const FAnimationAttributeIdentifier Attribute) const
@@ -237,6 +243,16 @@ const UE::UAF::FAttributeBindingDataPtr& UAbstractSkeletonSetBinding::GetAttribu
 	return AttributeBindingData;
 }
 
+void UAbstractSkeletonSetBinding::OnBindingsChanged()
+{
+	RebuildBindingData();
+
+#if WITH_EDITOR
+	// Notify external listeners a change has been made
+	OnBindingsChangedDelegate.Broadcast();
+#endif
+}
+
 void UAbstractSkeletonSetBinding::RebuildBindingData()
 {
 	// Clear our old data
@@ -271,7 +287,7 @@ bool UAbstractSkeletonSetBinding::SetSkeleton(const TObjectPtr<USkeleton> InSkel
 	BoneBindings.Empty();
 	AttributeBindings.Empty();
 
-	RebuildBindingData();
+	OnBindingsChanged();
 
 #if WITH_EDITOR
 	if (InSkeleton)
@@ -308,7 +324,7 @@ bool UAbstractSkeletonSetBinding::SetSetCollection(const TObjectPtr<UAbstractSke
 	BoneBindings.Empty();
 	AttributeBindings.Empty();
 
-	RebuildBindingData();
+	OnBindingsChanged();
 
 #if WITH_EDITOR
 	if (InSetCollection)
@@ -326,6 +342,7 @@ TObjectPtr<const UAbstractSkeletonSetCollection> UAbstractSkeletonSetBinding::Ge
 }
 
 #if WITH_EDITOR
+
 void UAbstractSkeletonSetBinding::OnSkeletonChanged()
 {
 	RebuildBindingData();
@@ -335,4 +352,15 @@ void UAbstractSkeletonSetBinding::OnSetCollectionChanged()
 {
 	RebuildBindingData();
 }
-#endif
+
+FDelegateHandle UAbstractSkeletonSetBinding::RegisterOnBindingsChanged(FSimpleMulticastDelegate::FDelegate&& InDelegate)
+{
+	return OnBindingsChangedDelegate.Add(InDelegate);
+}
+
+bool UAbstractSkeletonSetBinding::UnregisterOnBindingsChanged(const FDelegateHandle InHandle)
+{
+	return OnBindingsChangedDelegate.Remove(InHandle);
+}
+
+#endif // WITH_EDITOR
diff --git a/Engine/Plugins/Experimental/UAF/UAF/Source/UAF/Public/UAF/AbstractSkeleton/AbstractSkeletonSetBinding.h b/Engine/Plugins/Experimental/UAF/UAF/Source/UAF/Public/UAF/AbstractSkeleton/AbstractSkeletonSetBinding.h
--- a/Engine/Plugins/Experimental/UAF/UAF/Source/UAF/Public/UAF/AbstractSkeleton/AbstractSkeletonSetBinding.h
+++ b/Engine/Plugins/Experimental/UAF/UAF/Source/UAF/Public/UAF/AbstractSkeleton/AbstractSkeletonSetBinding.h
@@ -94,12 +94,20 @@ public:
 	// Returns the baked binding data
 	UE_API const UE::UAF::FAttributeBindingDataPtr& GetAttributeBindingData() const;
 
+#if WITH_EDITOR
+	UE_API FDelegateHandle RegisterOnBindingsChanged(FSimpleMulticastDelegate::FDelegate&& InDelegate);
+	UE_API bool UnregisterOnBindingsChanged(const FDelegateHandle InHandle);
+#endif
+
 private:
 #if WITH_EDITOR
 	void OnSkeletonChanged();
 	void OnSetCollectionChanged();
 #endif
 
+	// Called internally whenever a change is made to the attribute bindings (bones + attributes)
+	void OnBindingsChanged();
+
 	void RebuildBindingData();
 
 	UPROPERTY()
@@ -122,6 +130,8 @@ private:
 
 #if WITH_EDITOR
 	FDelegateHandle SetCollectionChangedHandle;
+
+	FSimpleMulticastDelegate OnBindingsChangedDelegate;
 #endif
 };
 
diff --git a/Engine/Plugins/Experimental/UAF/UAF/Source/UAFEditor/Private/UAF/AbstractSkeleton/Sets/SAttributesList.cpp b/Engine/Plugins/Experimental/UAF/UAF/Source/UAFEditor/Private/UAF/AbstractSkeleton/Sets/SAttributesList.cpp
--- a/Engine/Plugins/Experimental/UAF/UAF/Source/UAFEditor/Private/UAF/AbstractSkeleton/Sets/SAttributesList.cpp
+++ b/Engine/Plugins/Experimental/UAF/UAF/Source/UAFEditor/Private/UAF/AbstractSkeleton/Sets/SAttributesList.cpp
@@ -19,21 +19,36 @@
 // TODO: Move to UE::UAF::Editor
 using namespace UE::UAF::Sets;
 
-FName SAttributesList::FColumns::SetId = FName("Set");
-FName SAttributesList::FColumns::AttributeId = FName("Attribute");
+static const FName ColumnId_Set("Set");
+static const FName ColumnId_Attribute("Attribute");
 
-FName SAttributesList::FMenus::AddAttributeId = FName("SAttributesList.AddAttribute");
+static const FName MenuId_AddAttribute("SAttributesList.AddAttribute");
 
 void SAttributesList::SetSetBinding(TWeakObjectPtr<UAbstractSkeletonSetBinding> InSetBinding)
 {
+	if (SetBinding.IsValid())
+	{
+		SetBinding->UnregisterOnBindingsChanged(OnBindingsChangedHandle);
+	}
+
 	SetBinding = InSetBinding;
+
+	if (SetBinding.IsValid())
+	{
+		OnBindingsChangedHandle = SetBinding->RegisterOnBindingsChanged(FSimpleMulticastDelegate::FDelegate::CreateSP(this, &SAttributesList::HandleBindingsChanged));
+	}
+
 	RepopulateListData();
 }
 
-void SAttributesList::Construct(const FArguments& InArgs)
+void SAttributesList::Construct(const FArguments& InArgs, TWeakObjectPtr<UAbstractSkeletonSetBinding> InSetBinding)
 {
-	SetBinding = InArgs._SetBinding;
-	OnListRefreshed = InArgs._OnListRefreshed;
+	SetBinding = InSetBinding;
+
+	if (SetBinding.IsValid())
+	{
+		OnBindingsChangedHandle = SetBinding->RegisterOnBindingsChanged(FSimpleMulticastDelegate::FDelegate::CreateSP(this, &SAttributesList::HandleBindingsChanged));
+	}
 
 	RegisterMenus();
 
@@ -76,10 +91,10 @@ void SAttributesList::Construct(const FArguments& InArgs)
 					.OnContextMenuOpening(this, &SAttributesList::ListView_OnContextMenuOpening)
 					.HeaderRow(
 						SNew(SHeaderRow)
-						+ SHeaderRow::Column(FColumns::SetId)
+						+ SHeaderRow::Column(ColumnId_Set)
 						.DefaultLabel(LOCTEXT("SetColumnLabel", "Set"))
 						.FillWidth(0.3f)
-						+ SHeaderRow::Column(FColumns::AttributeId)
+						+ SHeaderRow::Column(ColumnId_Attribute)
 						.DefaultLabel(LOCTEXT("AttributeColumnLabel", "Attribute"))
 						.FillWidth(0.7f)
 					)
@@ -89,27 +104,39 @@ void SAttributesList::Construct(const FArguments& InArgs)
 	RepopulateListData();
 }
 
-void SAttributesList::RepopulateListData()
+void SAttributesList::Tick(const FGeometry& AllocatedGeometry, const double InCurrentTime, const float InDeltaTime)
 {
-	if (!bRepopulating)
+	if (bListDirty)
 	{
-		TGuardValue<bool> RecursionGuard(bRepopulating, true);
+		RepopulateListData();
+		bListDirty = false;
+	}
+}
 
-		ListItems.Empty();
+void SAttributesList::PostUndo(bool bSuccess)
+{
+	RepopulateListData();
+}
 
-		if (SetBinding.IsValid())
-		{
-			for (const FAbstractSkeleton_AttributeBinding& Binding : SetBinding->GetAttributeBindings())
-			{
-				FListItemPtr BindingItem = MakeShared<FListItem>(Binding.SetName, Binding.Attribute);
-				ListItems.Add(BindingItem);
-			}
-		}
+void SAttributesList::PostRedo(bool bSuccess)
+{
+	RepopulateListData();
+}
 
-		ListView->RequestListRefresh();
+void SAttributesList::RepopulateListData()
+{
+	ListItems.Empty();
 
-		OnListRefreshed.ExecuteIfBound();
+	if (SetBinding.IsValid())
+	{
+		for (const FAbstractSkeleton_AttributeBinding& Binding : SetBinding->GetAttributeBindings())
+		{
+			FListItemPtr BindingItem = MakeShared<FListItem>(Binding.SetName, Binding.Attribute);
+			ListItems.Add(BindingItem);
+		}
 	}
+
+	ListView->RequestListRefresh();
 }
 
 class SAttributesListRow : public SMultiColumnTableRow<SAttributesList::FListItemPtr>
@@ -137,7 +164,7 @@ public:
 
 	virtual TSharedRef<SWidget> GenerateWidgetForColumn(const FName& ColumnName) override
 	{
-		if (ColumnName == SAttributesList::FColumns::SetId)
+		if (ColumnName == ColumnId_Set)
 		{
 			return SNew(SHorizontalBox)
 				+ SHorizontalBox::Slot()
@@ -171,7 +198,7 @@ public:
 						.Text(FText::FromName(Item->SetName))
 				];
 		}
-		else if (ColumnName == SAttributesList::FColumns::AttributeId)
+		else if (ColumnName == ColumnId_Attribute)
 		{
 			return SNew(SHorizontalBox)
 				+ SHorizontalBox::Slot()
@@ -239,8 +266,6 @@ TSharedRef<ITableRow> SAttributesList::ListView_OnGenerateRow(FListItemPtr Item,
 						SetBinding->AddAttributeToSet(TargetTreeItem->Attribute, DragDropOp->GetDraggedSets()[0].SetName);
 					}
 					
-					RepopulateListData();
-
 					return FReply::Handled();
 				}
 
@@ -290,8 +315,6 @@ TSharedPtr<SWidget> SAttributesList::ListView_OnContextMenuOpening()
 						SetBinding->Modify();
 
 						ensure(SetBinding->RemoveAttributeFromSet(Selected->Attribute));
-
-						RepopulateListData();
 					}
 				)
 			));
@@ -309,8 +332,6 @@ TSharedPtr<SWidget> SAttributesList::ListView_OnContextMenuOpening()
 						SetBinding->Modify();
 
 						ensure(SetBinding->RemoveAttributeFromSet(Selected->Attribute));
-
-						RepopulateListData();
 					}
 				)
 			));
@@ -331,8 +352,6 @@ TSharedPtr<SWidget> SAttributesList::ListView_OnContextMenuOpening()
 
 						// Then removing attributes from the 'None' set actually removes them.
 						ensure(SetBinding->RemoveAttributeFromSet(Selected->Attribute));
-
-						RepopulateListData();
 					}
 				)
 			));
@@ -343,13 +362,13 @@ TSharedPtr<SWidget> SAttributesList::ListView_OnContextMenuOpening()
 
 void SAttributesList::RegisterMenus()
 {
-	if (UToolMenus::Get()->IsMenuRegistered(FMenus::AddAttributeId))
+	if (UToolMenus::Get()->IsMenuRegistered(MenuId_AddAttribute))
 	{
 		return;
 	}
 
 	FToolMenuOwnerScoped OwnerScoped(this);
-	UToolMenu* Menu = UToolMenus::Get()->RegisterMenu(FMenus::AddAttributeId);
+	UToolMenu* Menu = UToolMenus::Get()->RegisterMenu(MenuId_AddAttribute);
 
 	UToolMenu* AddCurveMenu = Menu->AddSubMenu(FToolMenuOwner(), "AddCurve", "AddCurve", LOCTEXT("AddCurve", "Add Curve"));
 	UToolMenu* AddAttributeMenu = Menu->AddSubMenu(FToolMenuOwner(), "AddAttribute", "AddAttribute", LOCTEXT("AddAttribute", "Add Attribute"));
@@ -387,7 +406,6 @@ void SAttributesList::RegisterMenus()
 										}
 
 										FSlateApplication::Get().DismissAllMenus();
-										WeakSetBindingWidget.Pin()->RepopulateListData();
 									}),
 								true,
 								false),
@@ -434,7 +452,6 @@ void SAttributesList::RegisterMenus()
 										SetBinding->AddAttributeToSet(CurveAttrId, NAME_None);
 
 										FSlateApplication::Get().DismissAllMenus();
-										SetBindingWidget.Pin()->RepopulateListData();
 									})),
 							FText(),
 							true,
@@ -449,6 +466,11 @@ void SAttributesList::RegisterMenus()
 	}
 }
 
+void SAttributesList::HandleBindingsChanged()
+{
+	bListDirty = true;
+}
+
 TSharedRef<SWidget> SAttributesList::CreateAddAttributeWidget()
 {
 	UAttributesListMenuContext* ListContext = NewObject<UAttributesListMenuContext>();
@@ -457,7 +479,7 @@ TSharedRef<SWidget> SAttributesList::CreateAddAttributeWidget()
 	FToolMenuContext MenuContext;
 	MenuContext.AddObject(ListContext);
 
-	return UToolMenus::Get()->GenerateWidget(FMenus::AddAttributeId, MenuContext);
+	return UToolMenus::Get()->GenerateWidget(MenuId_AddAttribute, MenuContext);
 }
 
 #undef LOCTEXT_NAMESPACE
\ No newline at end of file
diff --git a/Engine/Plugins/Experimental/UAF/UAF/Source/UAFEditor/Private/UAF/AbstractSkeleton/Sets/SAttributesList.h b/Engine/Plugins/Experimental/UAF/UAF/Source/UAFEditor/Private/UAF/AbstractSkeleton/Sets/SAttributesList.h
--- a/Engine/Plugins/Experimental/UAF/UAF/Source/UAFEditor/Private/UAF/AbstractSkeleton/Sets/SAttributesList.h
+++ b/Engine/Plugins/Experimental/UAF/UAF/Source/UAFEditor/Private/UAF/AbstractSkeleton/Sets/SAttributesList.h
@@ -4,6 +4,7 @@
 
 #include "Animation/AnimData/AttributeIdentifier.h"
 #include "CoreMinimal.h"
+#include "EditorUndoClient.h"
 #include "Widgets/SCompoundWidget.h"
 #include "Widgets/Views/SListView.h"
 
@@ -15,32 +16,17 @@ class UAbstractSkeletonSetBinding;
 namespace UE::UAF::Sets
 {
 
-	class SAttributesList : public SCompoundWidget
+	class SAttributesList : public SCompoundWidget, public FSelfRegisteringEditorUndoClient
 	{
 	public:
 		SLATE_BEGIN_ARGS(SAttributesList) {}
-			SLATE_ARGUMENT(TWeakObjectPtr<UAbstractSkeletonSetBinding>, SetBinding)
-			SLATE_EVENT(FSimpleDelegate, OnListRefreshed)
 		SLATE_END_ARGS()
 
-		virtual void Construct(const FArguments& InArgs);
-
-		void RepopulateListData();
+		void Construct(const FArguments& InArgs, TWeakObjectPtr<UAbstractSkeletonSetBinding> SetBinding);
 
 		void SetSetBinding(TWeakObjectPtr<UAbstractSkeletonSetBinding> InSetBinding);
 
 	public:
-		struct FColumns
-		{
-			static FName SetId;
-			static FName AttributeId;
-		};
-
-		struct FMenus
-		{
-			static FName AddAttributeId;
-		};
-
 		struct FListItem
 		{
 			FListItem(const FName InSetName, const FAnimationAttributeIdentifier& InAttribute);
@@ -52,14 +38,26 @@ namespace UE::UAF::Sets
 		using FListItemPtr = TSharedPtr<FListItem>;
 
 	private:
+		/** SCompoundWidget */
+		virtual void Tick(const FGeometry& AllocatedGeometry, const double InCurrentTime, const float InDeltaTime) override;
+
+		/** FSelfRegisteringEditorUndoClient interface */
+		virtual void PostUndo(bool bSuccess) override;
+		virtual void PostRedo(bool bSuccess) override;
+
+		void RepopulateListData();
+
+		void RegisterMenus();
+
+		// Called when the bindings inside the set binding asset we're editing have been modified
+		void HandleBindingsChanged();
+
 		TSharedRef<ITableRow> ListView_OnGenerateRow(FListItemPtr Item, const TSharedRef<STableViewBase>& OwnerTable);
 
 		TSharedPtr<SWidget> ListView_OnContextMenuOpening();
 		
 		TSharedRef<SWidget> CreateAddAttributeWidget();
 
-		void RegisterMenus();
-
 	private:
 		TArray<FListItemPtr> ListItems;
 
@@ -67,9 +65,9 @@ namespace UE::UAF::Sets
 
 		TWeakObjectPtr<UAbstractSkeletonSetBinding> SetBinding;
 
-		FSimpleDelegate OnListRefreshed;
+		FDelegateHandle OnBindingsChangedHandle;
 
-		bool bRepopulating = false;
+		bool bListDirty = false;
 	};
 
 }
diff --git a/Engine/Plugins/Experimental/UAF/UAF/Source/UAFEditor/Private/UAF/AbstractSkeleton/Sets/SSetBinding.cpp b/Engine/Plugins/Experimental/UAF/UAF/Source/UAFEditor/Private/UAF/AbstractSkeleton/Sets/SSetBinding.cpp
--- a/Engine/Plugins/Experimental/UAF/UAF/Source/UAFEditor/Private/UAF/AbstractSkeleton/Sets/SSetBinding.cpp
+++ b/Engine/Plugins/Experimental/UAF/UAF/Source/UAFEditor/Private/UAF/AbstractSkeleton/Sets/SSetBinding.cpp
@@ -19,10 +19,14 @@
 // TODO: Move to UE::UAF::Editor
 namespace UE::UAF::Sets
 {
-void SSetBinding::Construct(const FArguments& InArgs, TWeakObjectPtr<UAbstractSkeletonSetBinding> InSetBinding)
+	void SSetBinding::Construct(const FArguments& InArgs, TWeakObjectPtr<UAbstractSkeletonSetBinding> InSetBinding)
 	{
 		SetBinding = InSetBinding;
-		OnTreeRefreshed = InArgs._OnTreeRefreshed;
+
+		if (SetBinding.IsValid())
+		{
+			OnBindingsChangedHandle = SetBinding->RegisterOnBindingsChanged(FSimpleMulticastDelegate::FDelegate::CreateSP(this, &SSetBinding::HandleBindingsChanged));
+		}
 
 		ChildSlot
 			[
@@ -92,6 +96,15 @@ void SSetBinding::Construct(const FArguments& InArgs, TWeakObjectPtr<UAbstractSk
 		BindCommands();
 	}
 
+	void SSetBinding::Tick(const FGeometry& AllocatedGeometry, const double InCurrentTime, const float InDeltaTime)
+	{
+		if (bTreeDirty)
+		{
+			RepopulateTreeData();
+			bTreeDirty = false;
+		}
+	}
+
 	FReply SSetBinding::OnKeyDown(const FGeometry& MyGeometry, const FKeyEvent& InKeyEvent)
 	{
 		if (CommandList->ProcessCommandBindings(InKeyEvent))
@@ -144,114 +157,123 @@ void SSetBinding::Construct(const FArguments& InArgs, TWeakObjectPtr<UAbstractSk
 
 	void SSetBinding::SetSetBinding(TWeakObjectPtr<UAbstractSkeletonSetBinding> InSetBinding)
 	{
+		if (SetBinding.IsValid())
+		{
+			SetBinding->UnregisterOnBindingsChanged(OnBindingsChangedHandle);
+		}
+
 		SetBinding = InSetBinding;
+
+		if (SetBinding.IsValid())
+		{
+			OnBindingsChangedHandle = SetBinding->RegisterOnBindingsChanged(FSimpleMulticastDelegate::FDelegate::CreateSP(this, &SSetBinding::HandleBindingsChanged));
+		}
+
 		RepopulateTreeData();
 	}
 
 	void SSetBinding::RepopulateTreeData()
 	{
-		if (!bRepopulating)
+		// Save the current expanded state of each tree item
+		// - Only Set items can contain children so only Set items can be expanded/collapsed
+		TSet<FName> ClosedSetCollections;
 		{
-			TGuardValue<bool> RecursionGuard(bRepopulating, true);
-
-			// Save the current expanded state of each tree item
-			// - Only Set items can contain children so only Set items can be expanded/collapsed
-			TSet<FName> ClosedSetCollections;
+			for (TSharedPtr<ITreeItem> TreeItem : GetAllTreeItems())
 			{
-				for (TSharedPtr<ITreeItem> TreeItem : GetAllTreeItems())
+				if (TreeItem->GetType() == FTreeItem_Set::StaticGetType() && !TreeView->IsItemExpanded(TreeItem))
 				{
-					if (TreeItem->GetType() == FTreeItem_Set::StaticGetType() && !TreeView->IsItemExpanded(TreeItem))
-					{
-						TSharedPtr<FTreeItem_Set> SetTreeItem = StaticCastSharedPtr<FTreeItem_Set>(TreeItem);
-						ClosedSetCollections.Add(SetTreeItem->Set.SetName);
-					}
+					TSharedPtr<FTreeItem_Set> SetTreeItem = StaticCastSharedPtr<FTreeItem_Set>(TreeItem);
+					ClosedSetCollections.Add(SetTreeItem->Set.SetName);
 				}
 			}
+		}
 			
-			// Rebuild tree data
-			RootItems.Reset();
+		// Rebuild tree data
+		RootItems.Reset();
 
-			if (SetBinding.IsValid())
-			{
-				TMap<FName, TSharedPtr<FTreeItem_Set>> TreeItemMap;
+		if (SetBinding.IsValid())
+		{
+			TMap<FName, TSharedPtr<FTreeItem_Set>> TreeItemMap;
 				
-				if (SetBinding->GetSetCollection())
+			if (SetBinding->GetSetCollection())
+			{
+				for (const FAbstractSkeletonSet& Set : SetBinding->GetSetCollection()->GetSetHierarchy())
 				{
-					for (const FAbstractSkeletonSet& Set : SetBinding->GetSetCollection()->GetSetHierarchy())
-					{
-						TSharedPtr<FTreeItem_Set> SetItem = MakeShared<FTreeItem_Set>(*this, Set, false);
+					TSharedPtr<FTreeItem_Set> SetItem = MakeShared<FTreeItem_Set>(*this, Set, false);
 
-						if (Set.ParentSetName == NAME_None)
-						{
-							RootItems.Add(SetItem);
-						}
-						else
-						{
-							ensure(TreeItemMap.Contains(Set.ParentSetName));
+					if (Set.ParentSetName == NAME_None)
+					{
+						RootItems.Add(SetItem);
+					}
+					else
+					{
+						ensure(TreeItemMap.Contains(Set.ParentSetName));
 							
-							TSharedPtr<FTreeItem_Set> ParentItem = TreeItemMap[Set.ParentSetName];
-							ParentItem->Children.Add(SetItem);
-						}
-
-						TreeItemMap.Add(Set.SetName, SetItem);
+						TSharedPtr<FTreeItem_Set> ParentItem = TreeItemMap[Set.ParentSetName];
+						ParentItem->Children.Add(SetItem);
 					}
+
+					TreeItemMap.Add(Set.SetName, SetItem);
 				}
+			}
 
-				for (const FAbstractSkeleton_BoneBinding& Binding : SetBinding->GetBoneBindings())
-				{
-					TSharedPtr<FTreeItem_Bone> BindingItem = MakeShared<FTreeItem_Bone>(*this, Binding);
+			for (const FAbstractSkeleton_BoneBinding& Binding : SetBinding->GetBoneBindings())
+			{
+				TSharedPtr<FTreeItem_Bone> BindingItem = MakeShared<FTreeItem_Bone>(*this, Binding);
 
-					if (TreeItemMap.Contains(Binding.SetName))
-					{
-						TSharedPtr<FTreeItem_Set> SetItem = TreeItemMap[Binding.SetName];
-						SetItem->Children.Insert(BindingItem, 0);
-					}
-					else
-					{
-						// This binding binds to a set that no longer exists in the set collection
-						// These bindings are not auto-deleted to avoid implicit destructive edits
-						// This binding can be fixed by re-adding the set in the set collection or
-						// by removing this binding. Show this binding as an orphan.
+				if (TreeItemMap.Contains(Binding.SetName))
+				{
+					TSharedPtr<FTreeItem_Set> SetItem = TreeItemMap[Binding.SetName];
+					SetItem->Children.Insert(BindingItem, 0);
+				}
+				else
+				{
+					// This binding binds to a set that no longer exists in the set collection
+					// These bindings are not auto-deleted to avoid implicit destructive edits
+					// This binding can be fixed by re-adding the set in the set collection or
+					// by removing this binding. Show this binding as an orphan.
 
-						TSharedPtr<FTreeItem_Set> DeletedSet = MakeShared<FTreeItem_Set>(*this, FAbstractSkeletonSet(Binding.SetName, NAME_None), true);
-						DeletedSet->Children.Insert(BindingItem, 0);
+					TSharedPtr<FTreeItem_Set> DeletedSet = MakeShared<FTreeItem_Set>(*this, FAbstractSkeletonSet(Binding.SetName, NAME_None), true);
+					DeletedSet->Children.Insert(BindingItem, 0);
 
-						TreeItemMap.Add(Binding.SetName, DeletedSet);
-						RootItems.Add(DeletedSet);
-					}
+					TreeItemMap.Add(Binding.SetName, DeletedSet);
+					RootItems.Add(DeletedSet);
 				}
+			}
 
-				for (const FAbstractSkeleton_AttributeBinding& Binding : SetBinding->GetAttributeBindings())
-				{
-					TSharedPtr<FTreeItem_Attribute> BindingItem = MakeShared<FTreeItem_Attribute>(*this, Binding);
+			for (const FAbstractSkeleton_AttributeBinding& Binding : SetBinding->GetAttributeBindings())
+			{
+				TSharedPtr<FTreeItem_Attribute> BindingItem = MakeShared<FTreeItem_Attribute>(*this, Binding);
 
-					if (TreeItemMap.Contains(Binding.SetName))
-					{
-						TSharedPtr<FTreeItem_Set> SetItem = TreeItemMap[Binding.SetName];
-						SetItem->Children.Insert(BindingItem, 0);
-					}
+				if (TreeItemMap.Contains(Binding.SetName))
+				{
+					TSharedPtr<FTreeItem_Set> SetItem = TreeItemMap[Binding.SetName];
+					SetItem->Children.Insert(BindingItem, 0);
 				}
 			}
+		}
 
-			TreeView->RebuildList();
+		TreeView->RebuildList();
 
-			// Restore expanded/collapsed state
+		// Restore expanded/collapsed state
+		{
+			for (TSharedPtr<ITreeItem> TreeItem : GetAllTreeItems())
 			{
-				for (TSharedPtr<ITreeItem> TreeItem : GetAllTreeItems())
+				if (TreeItem->GetType() == FTreeItem_Set::StaticGetType())
 				{
-					if (TreeItem->GetType() == FTreeItem_Set::StaticGetType())
-					{
-						const TSharedPtr<FTreeItem_Set> SetTreeItem = StaticCastSharedPtr<FTreeItem_Set>(TreeItem);
-						const bool bShouldExpand = !ClosedSetCollections.Contains(SetTreeItem->Set.SetName);
-						TreeView->SetItemExpansion(TreeItem, bShouldExpand);
-					}
+					const TSharedPtr<FTreeItem_Set> SetTreeItem = StaticCastSharedPtr<FTreeItem_Set>(TreeItem);
+					const bool bShouldExpand = !ClosedSetCollections.Contains(SetTreeItem->Set.SetName);
+					TreeView->SetItemExpansion(TreeItem, bShouldExpand);
 				}
 			}
-
-			OnTreeRefreshed.ExecuteIfBound();
 		}
 	}
 
+	void SSetBinding::HandleBindingsChanged()
+	{
+		bTreeDirty = true;
+	}
+
 	void SSetBinding::ExpandAllTreeItems()
 	{
 		for (TSharedPtr<ITreeItem> TreeItem : GetAllTreeItems())
@@ -349,20 +371,13 @@ void SSetBinding::Construct(const FArguments& InArgs, TWeakObjectPtr<UAbstractSk
 					{
 						TSharedPtr<FTreeItem_Set> TargetSetTreeItem = StaticCastSharedPtr<FTreeItem_Set>(TargetTreeItem);
 
-						const FScopedTransaction Transaction(LOCTEXT("AddAttributeToSet", "Add Attribute to Set"));
+						const FScopedTransaction Transaction(LOCTEXT("BindAttributeToSet", "Bind Attribute to Set"));
 						SetBindingWidget.SetBinding->Modify();
 
-						bool bTreeDirty = false;
-									
 						for (const TSharedPtr<SSetsSkeletonTree::FTreeItem>& Item : DragDropOp->Items)
 						{
-							bTreeDirty |= SetBindingWidget.SetBinding->RemoveBoneFromSet(Item->BoneName);
-							bTreeDirty |= SetBindingWidget.SetBinding->AddBoneToSet(Item->BoneName, TargetSetTreeItem->Set.SetName);
-						}
-							
-						if (bTreeDirty)
-						{
-							SetBindingWidget.RepopulateTreeData();
+							SetBindingWidget.SetBinding->RemoveBoneFromSet(Item->BoneName);
+							SetBindingWidget.SetBinding->AddBoneToSet(Item->BoneName, TargetSetTreeItem->Set.SetName);
 						}
 
 						return FReply::Handled();
@@ -372,20 +387,13 @@ void SSetBinding::Construct(const FArguments& InArgs, TWeakObjectPtr<UAbstractSk
 					{
 						TSharedPtr<FTreeItem_Set> TargetSetTreeItem = StaticCastSharedPtr<FTreeItem_Set>(TargetTreeItem);
 						
-						const FScopedTransaction Transaction(LOCTEXT("AddAttributeToSet", "Add Attribute to Set"));
+						const FScopedTransaction Transaction(LOCTEXT("BindAttributeToSet", "Bind Attribute to Set"));
 						SetBindingWidget.SetBinding->Modify();
 
-						bool bTreeDirty = false;
-									
 						for (const FAnimationAttributeIdentifier& Attribute : DragDropOp->Attributes)
 						{
-							bTreeDirty |= SetBindingWidget.SetBinding->RemoveAttributeFromSet(Attribute);
-							bTreeDirty |= SetBindingWidget.SetBinding->AddAttributeToSet(Attribute, TargetSetTreeItem->Set.SetName);
-						}
-							
-						if (bTreeDirty)
-						{
-							SetBindingWidget.RepopulateTreeData();
+							SetBindingWidget.SetBinding->RemoveAttributeFromSet(Attribute);
+							SetBindingWidget.SetBinding->AddAttributeToSet(Attribute, TargetSetTreeItem->Set.SetName);
 						}
 
 						return FReply::Handled();
@@ -411,7 +419,6 @@ void SSetBinding::Construct(const FArguments& InArgs, TWeakObjectPtr<UAbstractSk
 						SNew(STextBlock)
 							.Text(FText::FromName(Set.SetName))
 					]
-
 			];
 	}
 
@@ -453,17 +460,10 @@ void SSetBinding::Construct(const FArguments& InArgs, TWeakObjectPtr<UAbstractSk
 						const FScopedTransaction Transaction(LOCTEXT("AddAttributeToSet", "Add Attribute to Set"));
 						SetBindingWidget.SetBinding->Modify();
 
-						bool bTreeDirty = false;
-									
 						for (const TSharedPtr<SSetsSkeletonTree::FTreeItem>& Item : DragDropOp->Items)
 						{
-							bTreeDirty |= SetBindingWidget.SetBinding->RemoveBoneFromSet(Item->BoneName);
-							bTreeDirty |= SetBindingWidget.SetBinding->AddBoneToSet(Item->BoneName, Binding.SetName);
-						}
-							
-						if (bTreeDirty)
-						{
-							SetBindingWidget.RepopulateTreeData();
+							SetBindingWidget.SetBinding->RemoveBoneFromSet(Item->BoneName);
+							SetBindingWidget.SetBinding->AddBoneToSet(Item->BoneName, Binding.SetName);
 						}
 
 						return FReply::Handled();
@@ -521,17 +521,10 @@ void SSetBinding::Construct(const FArguments& InArgs, TWeakObjectPtr<UAbstractSk
 						const FScopedTransaction Transaction(LOCTEXT("AddAttributeToSet", "Add Attribute to Set"));
 						SetBindingWidget.SetBinding->Modify();
 
-						bool bTreeDirty = false;
-									
 						for (const TSharedPtr<SSetsSkeletonTree::FTreeItem>& Item : DragDropOp->Items)
 						{
-							bTreeDirty |= SetBindingWidget.SetBinding->RemoveBoneFromSet(Item->BoneName);
-							bTreeDirty |= SetBindingWidget.SetBinding->AddBoneToSet(Item->BoneName, Binding.SetName);
-						}
-							
-						if (bTreeDirty)
-						{
-							SetBindingWidget.RepopulateTreeData();
+							SetBindingWidget.SetBinding->RemoveBoneFromSet(Item->BoneName);
+							SetBindingWidget.SetBinding->AddBoneToSet(Item->BoneName, Binding.SetName);
 						}
 
 						return FReply::Handled();
@@ -570,8 +563,6 @@ void SSetBinding::Construct(const FArguments& InArgs, TWeakObjectPtr<UAbstractSk
 					SetBindingWidget.SetBinding->Modify();
 
 					SetBindingWidget.SetBinding->RemoveAllFromSet(Set.SetName);
-					SetBindingWidget.RepopulateTreeData(); // TODO: Changes to the set binding should broadcast which 
-					// the widget listens to and dirties the tree instead of manually repopulating.
 				})
 		);
 
@@ -593,7 +584,6 @@ void SSetBinding::Construct(const FArguments& InArgs, TWeakObjectPtr<UAbstractSk
 				SetBindingWidget.SetBinding->Modify();
 
 				SetBindingWidget.SetBinding->RemoveBoneFromSet(Binding.BoneName);
-				SetBindingWidget.RepopulateTreeData();
 			})
 		);
 
@@ -615,7 +605,6 @@ void SSetBinding::Construct(const FArguments& InArgs, TWeakObjectPtr<UAbstractSk
 				SetBindingWidget.SetBinding->Modify();
 
 				SetBindingWidget.SetBinding->RemoveAttributeFromSet(Binding.Attribute);
-				SetBindingWidget.RepopulateTreeData();
 			})
 		);
 
diff --git a/Engine/Plugins/Experimental/UAF/UAF/Source/UAFEditor/Private/UAF/AbstractSkeleton/Sets/SSetBinding.h b/Engine/Plugins/Experimental/UAF/UAF/Source/UAFEditor/Private/UAF/AbstractSkeleton/Sets/SSetBinding.h
--- a/Engine/Plugins/Experimental/UAF/UAF/Source/UAFEditor/Private/UAF/AbstractSkeleton/Sets/SSetBinding.h
+++ b/Engine/Plugins/Experimental/UAF/UAF/Source/UAFEditor/Private/UAF/AbstractSkeleton/Sets/SSetBinding.h
@@ -20,7 +20,6 @@ namespace UE::UAF::Sets
 	class SSetBinding : public SCompoundWidget, public FSelfRegisteringEditorUndoClient
 	{
 		SLATE_BEGIN_ARGS(SSetBinding) {}
-			SLATE_EVENT(FSimpleDelegate, OnTreeRefreshed)
 		SLATE_END_ARGS()
 
 	public:
@@ -105,16 +104,27 @@ namespace UE::UAF::Sets
 
 		void Construct(const FArguments& InArgs, TWeakObjectPtr<UAbstractSkeletonSetBinding> InSetBinding);
 
+		void SetSetBinding(TWeakObjectPtr<UAbstractSkeletonSetBinding> InSetBinding);
+
+	private:
+		/** SCompoundWidget */
+		virtual void Tick(const FGeometry& AllocatedGeometry, const double InCurrentTime, const float InDeltaTime) override;
+		virtual FReply OnKeyDown(const FGeometry& MyGeometry, const FKeyEvent& InKeyEvent) override;
+
 		/** FSelfRegisteringEditorUndoClient interface */
 		virtual void PostUndo(bool bSuccess) override;
 		virtual void PostRedo(bool bSuccess) override;
 
 		void RepopulateTreeData();
 
-		void SetSetBinding(TWeakObjectPtr<UAbstractSkeletonSetBinding> InSetBinding);
+		void BindCommands();
 
-	private:
-		virtual FReply OnKeyDown(const FGeometry& MyGeometry, const FKeyEvent& InKeyEvent) override;
+		void ExpandAllTreeItems();
+
+		TArray<TSharedPtr<ITreeItem>> GetAllTreeItems();
+
+		// Called when the bindings inside the set binding asset we're editing have been modified
+		void HandleBindingsChanged();
 
 		TSharedRef<ITableRow> TreeView_OnGenerateRow(TSharedPtr<ITreeItem> InItem, const TSharedRef<STableViewBase>& OwnerTable);
 
@@ -122,12 +132,6 @@ namespace UE::UAF::Sets
 
 		TSharedPtr<SWidget> TreeView_OnContextMenuOpening();
 
-		void ExpandAllTreeItems();
-
-		void BindCommands();
-
-		TArray<TSharedPtr<ITreeItem>> GetAllTreeItems();
-
 	public:
 		TWeakObjectPtr<UAbstractSkeletonSetBinding> SetBinding;
 
@@ -140,9 +144,9 @@ namespace UE::UAF::Sets
 
 		TSharedPtr<FUICommandList> CommandList;
 
-		FSimpleDelegate OnTreeRefreshed;
+		FDelegateHandle OnBindingsChangedHandle;
 
-		bool bRepopulating = false;
+		bool bTreeDirty = false;
 	};
 
 }
diff --git a/Engine/Plugins/Experimental/UAF/UAF/Source/UAFEditor/Private/UAF/AbstractSkeleton/Sets/SSetSkeletonTree.cpp b/Engine/Plugins/Experimental/UAF/UAF/Source/UAFEditor/Private/UAF/AbstractSkeleton/Sets/SSetSkeletonTree.cpp
--- a/Engine/Plugins/Experimental/UAF/UAF/Source/UAFEditor/Private/UAF/AbstractSkeleton/Sets/SSetSkeletonTree.cpp
+++ b/Engine/Plugins/Experimental/UAF/UAF/Source/UAFEditor/Private/UAF/AbstractSkeleton/Sets/SSetSkeletonTree.cpp
@@ -16,10 +16,14 @@
 // TODO: Move to UE::UAF::Editor
 namespace UE::UAF::Sets
 {
-	void SSetsSkeletonTree::Construct(const FArguments& InArgs)
+	void SSetsSkeletonTree::Construct(const FArguments& InArgs, TWeakObjectPtr<UAbstractSkeletonSetBinding> InSetBinding)
 	{
-		SetBinding = InArgs._SetBinding;
-		OnTreeRefreshed = InArgs._OnTreeRefreshed;
+		SetBinding = InSetBinding;
+
+		if (SetBinding.IsValid())
+		{
+			OnBindingsChangedHandle = SetBinding->RegisterOnBindingsChanged(FSimpleMulticastDelegate::FDelegate::CreateSP(this, &SSetsSkeletonTree::HandleBindingsChanged));
+		}
 
 		ChildSlot
 			[
@@ -60,6 +64,25 @@ namespace UE::UAF::Sets
 		RepopulateTreeData();
 	}
 
+	void SSetsSkeletonTree::Tick(const FGeometry& AllocatedGeometry, const double InCurrentTime, const float InDeltaTime)
+	{
+		if (bTreeDirty)
+		{
+			RepopulateTreeData();
+			bTreeDirty = false;
+		}
+	}
+
+	void SSetsSkeletonTree::PostUndo(bool bSuccess)
+	{
+		RepopulateTreeData();
+	}
+
+	void SSetsSkeletonTree::PostRedo(bool bSuccess)
+	{
+		RepopulateTreeData();
+	}
+
 	class SSetsSkeletonTreeRow : public SMultiColumnTableRow<TSharedPtr<SSetsSkeletonTree::FTreeItem>>
 	{
 	public:
@@ -188,10 +211,8 @@ namespace UE::UAF::Sets
 						// TODO: Use all dragged sets
 						check(!DragDropOp->GetDraggedSets().IsEmpty());
 						const bool bSuccess = SetBinding->AddBoneToSet(TargetSetTreeItem->BoneName, DragDropOp->GetDraggedSets()[0].SetName);
-						check(bSuccess);
+						ensure(bSuccess);
 						
-						RepopulateTreeData();
-
 						return FReply::Handled();
 					}
 					return FReply::Unhandled();
@@ -229,69 +250,78 @@ namespace UE::UAF::Sets
 
 	void SSetsSkeletonTree::SetSetBinding(TWeakObjectPtr<UAbstractSkeletonSetBinding> InSetBinding)
 	{
+		if (SetBinding.IsValid())
+		{
+			SetBinding->UnregisterOnBindingsChanged(OnBindingsChangedHandle);
+		}
+
 		SetBinding = InSetBinding;
+
+		if (SetBinding.IsValid())
+		{
+			OnBindingsChangedHandle = SetBinding->RegisterOnBindingsChanged(FSimpleMulticastDelegate::FDelegate::CreateSP(this, &SSetsSkeletonTree::HandleBindingsChanged));
+		}
+
 		RepopulateTreeData();
 	}
 
 	void SSetsSkeletonTree::RepopulateTreeData()
 	{
-		if (!bRepopulating)
+		// Save expansion state of each tree item to restore later
+		TSet<FName> CollapsedBoneNames;
+		for (const TSharedPtr<FTreeItem>& TreeItem : GetAllTreeItems())
 		{
-			TGuardValue<bool> RecursionGuard(bRepopulating, true);
-
-			// Save expansion state of each tree item to restore later
-			TSet<FName> CollapsedBoneNames;
-			for (const TSharedPtr<FTreeItem>& TreeItem : GetAllTreeItems())
+			if (!TreeView->IsItemExpanded(TreeItem))
 			{
-				if (!TreeView->IsItemExpanded(TreeItem))
-				{
-					CollapsedBoneNames.Add(TreeItem->BoneName);
-				}
+				CollapsedBoneNames.Add(TreeItem->BoneName);
 			}
+		}
 
-			RootItems.Empty();
+		RootItems.Empty();
 
-			if (SetBinding.IsValid() && SetBinding->GetSkeleton())
-			{
-				const FReferenceSkeleton& RefSkeleton = SetBinding->GetSkeleton()->GetReferenceSkeleton();
+		if (SetBinding.IsValid() && SetBinding->GetSkeleton())
+		{
+			const FReferenceSkeleton& RefSkeleton = SetBinding->GetSkeleton()->GetReferenceSkeleton();
 
-				TFunction<TSharedPtr<FTreeItem>(int32, const TSharedPtr<FTreeItem>&)> AddBoneToTree;
+			TFunction<TSharedPtr<FTreeItem>(int32, const TSharedPtr<FTreeItem>&)> AddBoneToTree;
 
-				AddBoneToTree = [this, &RefSkeleton, &AddBoneToTree](const int32 BoneIndex, const TSharedPtr<FTreeItem>& Parent) -> TSharedPtr<FTreeItem>
-				{
-					TSharedPtr<FTreeItem> TreeItem = MakeShared<FTreeItem>();
-					TreeItem->BoneName = RefSkeleton.GetBoneName(BoneIndex);
-					TreeItem->BoundSet = SetBinding->GetBoneSet(TreeItem->BoneName);
-					TreeItem->Parent = Parent;
+			AddBoneToTree = [this, &RefSkeleton, &AddBoneToTree](const int32 BoneIndex, const TSharedPtr<FTreeItem>& Parent) -> TSharedPtr<FTreeItem>
+			{
+				TSharedPtr<FTreeItem> TreeItem = MakeShared<FTreeItem>();
+				TreeItem->BoneName = RefSkeleton.GetBoneName(BoneIndex);
+				TreeItem->BoundSet = SetBinding->GetBoneSet(TreeItem->BoneName);
+				TreeItem->Parent = Parent;
 
-					TArray<int32> Children;
-					RefSkeleton.GetDirectChildBones(BoneIndex, Children);
+				TArray<int32> Children;
+				RefSkeleton.GetDirectChildBones(BoneIndex, Children);
 
-					for (const int32 ChildBoneIndex : Children)
-					{
-						TSharedPtr<FTreeItem> ChildTreeItem = AddBoneToTree(ChildBoneIndex, TreeItem);
-						TreeItem->Children.Add(ChildTreeItem);
-					}
+				for (const int32 ChildBoneIndex : Children)
+				{
+					TSharedPtr<FTreeItem> ChildTreeItem = AddBoneToTree(ChildBoneIndex, TreeItem);
+					TreeItem->Children.Add(ChildTreeItem);
+				}
 
-					return TreeItem;
-				};
+				return TreeItem;
+			};
 				
-				TSharedPtr<FTreeItem> RootBoneItem = AddBoneToTree(0, nullptr);
+			TSharedPtr<FTreeItem> RootBoneItem = AddBoneToTree(0, nullptr);
 
-				RootItems.Add(RootBoneItem);
-			}
+			RootItems.Add(RootBoneItem);
+		}
 
-			TreeView->RebuildList();
+		TreeView->RebuildList();
 
-			for (TSharedPtr<FTreeItem> TreeItem : GetAllTreeItems())
-			{
-				const bool bExpanded = !CollapsedBoneNames.Contains(TreeItem->BoneName);
-				TreeView->SetItemExpansion(TreeItem, bExpanded);
-			}
-
-			OnTreeRefreshed.ExecuteIfBound();
+		for (TSharedPtr<FTreeItem> TreeItem : GetAllTreeItems())
+		{
+			const bool bExpanded = !CollapsedBoneNames.Contains(TreeItem->BoneName);
+			TreeView->SetItemExpansion(TreeItem, bExpanded);
 		}
 	}
+
+	void SSetsSkeletonTree::HandleBindingsChanged()
+	{
+		bTreeDirty = true;
+	}
 }
 
 #undef LOCTEXT_NAMESPACE
\ No newline at end of file
diff --git a/Engine/Plugins/Experimental/UAF/UAF/Source/UAFEditor/Private/UAF/AbstractSkeleton/Sets/SSetSkeletonTree.h b/Engine/Plugins/Experimental/UAF/UAF/Source/UAFEditor/Private/UAF/AbstractSkeleton/Sets/SSetSkeletonTree.h
--- a/Engine/Plugins/Experimental/UAF/UAF/Source/UAFEditor/Private/UAF/AbstractSkeleton/Sets/SSetSkeletonTree.h
+++ b/Engine/Plugins/Experimental/UAF/UAF/Source/UAFEditor/Private/UAF/AbstractSkeleton/Sets/SSetSkeletonTree.h
@@ -3,6 +3,7 @@
 #pragma once
 
 #include "CoreMinimal.h"
+#include "EditorUndoClient.h"
 #include "Widgets/SCompoundWidget.h"
 
 class ITableRow;
@@ -12,7 +13,7 @@ class UAbstractSkeletonSetBinding;
 
 namespace UE::UAF::Sets
 {
-	class SSetsSkeletonTree : public SCompoundWidget
+	class SSetsSkeletonTree : public SCompoundWidget, public FSelfRegisteringEditorUndoClient
 	{
 	public:
 		struct FTreeItem
@@ -25,22 +26,31 @@ namespace UE::UAF::Sets
 		};
 
 		SLATE_BEGIN_ARGS(SSetsSkeletonTree) {}
-			SLATE_ARGUMENT(TWeakObjectPtr<UAbstractSkeletonSetBinding>, SetBinding)
-			SLATE_EVENT(FSimpleDelegate, OnTreeRefreshed)
 		SLATE_END_ARGS()
 
-		virtual void Construct(const FArguments& InArgs);
-
-		void RepopulateTreeData();
+		void Construct(const FArguments& InArgs, TWeakObjectPtr<UAbstractSkeletonSetBinding> SetBinding);
 
 		void SetSetBinding(TWeakObjectPtr<UAbstractSkeletonSetBinding> InSetBinding);
 
 	private:
-		TSharedRef<ITableRow> TreeView_OnGenerateRow(TSharedPtr<FTreeItem> Item, const TSharedRef<STableViewBase>& OwnerTable);
-		void TreeView_OnGetChildren(TSharedPtr<FTreeItem> InParent, TArray<TSharedPtr<FTreeItem>>& OutChildren);
+		/** SCompoundWidget */
+		virtual void Tick(const FGeometry& AllocatedGeometry, const double InCurrentTime, const float InDeltaTime) override;
+
+		/** FSelfRegisteringEditorUndoClient interface */
+		virtual void PostUndo(bool bSuccess) override;
+		virtual void PostRedo(bool bSuccess) override;
 
+		void RepopulateTreeData();
+
+		// Called when the bindings inside the set binding asset we're editing have been modified
+		void HandleBindingsChanged();
+		
 		TArray<TSharedPtr<FTreeItem>> GetAllTreeItems();
 
+		TSharedRef<ITableRow> TreeView_OnGenerateRow(TSharedPtr<FTreeItem> Item, const TSharedRef<STableViewBase>& OwnerTable);
+		
+		void TreeView_OnGetChildren(TSharedPtr<FTreeItem> InParent, TArray<TSharedPtr<FTreeItem>>& OutChildren);
+
 	private:
 		TWeakObjectPtr<UAbstractSkeletonSetBinding> SetBinding;
 
@@ -48,9 +58,9 @@ namespace UE::UAF::Sets
 
 		TSharedPtr<STreeView<TSharedPtr<FTreeItem>>> TreeView;
 
-		FSimpleDelegate OnTreeRefreshed;
+		FDelegateHandle OnBindingsChangedHandle;
 
-		bool bRepopulating = false;
+		bool bTreeDirty = false;
 	};
 
 }
\ No newline at end of file
diff --git a/Engine/Plugins/Experimental/UAF/UAF/Source/UAFEditor/Private/UAF/AbstractSkeleton/Sets/SSetsTab.cpp b/Engine/Plugins/Experimental/UAF/UAF/Source/UAFEditor/Private/UAF/AbstractSkeleton/Sets/SSetsTab.cpp
--- a/Engine/Plugins/Experimental/UAF/UAF/Source/UAFEditor/Private/UAF/AbstractSkeleton/Sets/SSetsTab.cpp
+++ b/Engine/Plugins/Experimental/UAF/UAF/Source/UAFEditor/Private/UAF/AbstractSkeleton/Sets/SSetsTab.cpp
@@ -7,6 +7,7 @@
 #include "Framework/Docking/LayoutService.h"
 #include "Framework/MultiBox/MultiBoxBuilder.h"
 #include "PropertyCustomizationHelpers.h"
+#include "ScopedTransaction.h"
 #include "UAF/AbstractSkeleton/AbstractSkeletonSetBinding.h"
 #include "UAF/AbstractSkeleton/AbstractSkeletonSetCollection.h"
 #include "UAF/AbstractSkeleton/Sets/SSetSkeletonTree.h"
@@ -59,21 +60,7 @@ namespace UE::UAF::Sets
 				return SNew(SDockTab)
 					.CanEverClose(false)
 					[
-						SAssignNew(SkeletonTreeWidget, SSetsSkeletonTree)
-							.SetBinding(SetBinding)
-							.OnTreeRefreshed_Lambda([this]()
-								{
-									if (SetBindingWidget.IsValid())
-									{
-										SetBindingWidget->RepopulateTreeData();
-									}
-
-									if (AttributesListWidget.IsValid())
-									{
-										AttributesListWidget->RepopulateListData();
-									}
-								})
-
+						SAssignNew(SkeletonTreeWidget, SSetsSkeletonTree, SetBinding)
 					];
 			}))
 			.SetDisplayName(LOCTEXT("SkeletonTree_TabDisplayName", "Skeleton Tree"))
@@ -84,20 +71,7 @@ namespace UE::UAF::Sets
 				return SNew(SDockTab)
 					.CanEverClose(false)
 					[
-						SAssignNew(AttributesListWidget, SAttributesList)
-							.SetBinding(SetBinding)
-							.OnListRefreshed_Lambda([this]()
-								{
-									if (SkeletonTreeWidget.IsValid())
-									{
-										SkeletonTreeWidget->RepopulateTreeData();
-									}
-
-									if (SetBindingWidget.IsValid())
-									{
-										SetBindingWidget->RepopulateTreeData();
-									}
-								})
+						SAssignNew(AttributesListWidget, SAttributesList, SetBinding)
 					];
 			}))
 			.SetDisplayName(LOCTEXT("Attributes_TabDisplayName", "Attributes"))
@@ -109,18 +83,6 @@ namespace UE::UAF::Sets
 					.CanEverClose(false)
 					[
 						SAssignNew(SetBindingWidget, SSetBinding, SetBinding)
-							.OnTreeRefreshed_Lambda([this]()
-								{
-									if (SkeletonTreeWidget.IsValid())
-									{
-										SkeletonTreeWidget->RepopulateTreeData();
-									}
-
-									if (AttributesListWidget.IsValid())
-									{
-										AttributesListWidget->RepopulateListData();
-									}
-								})
 					];
 			}))
 			.SetDisplayName(LOCTEXT("SetBindings_TabDisplayName", "Set Bindings"))
@@ -209,12 +171,11 @@ namespace UE::UAF::Sets
 							{
 								if (SetBinding.IsValid())
 								{
-									const bool bSuccess = SetBinding->SetSetCollection(CastChecked<UAbstractSkeletonSetCollection>(AssetData.GetAsset()));
-									check(bSuccess);
+									const FScopedTransaction Transaction(LOCTEXT("SetSetCollection", "Set Set Collection"));
+									SetBinding->Modify();
 
-									SetBindingWidget->RepopulateTreeData();
-									AttributesListWidget->RepopulateListData();
-									SetBindingWidget->RepopulateTreeData();
+									const bool bSuccess = SetBinding->SetSetCollection(CastChecked<UAbstractSkeletonSetCollection>(AssetData.GetAsset()));
+									ensure(bSuccess);
 								}
 							})
 						.AllowClear(false)
@@ -259,12 +220,11 @@ namespace UE::UAF::Sets
 							{
 								if (SetBinding.IsValid())
 								{
-									const bool bSuccess = SetBinding->SetSkeleton(CastChecked<USkeleton>(AssetData.GetAsset()));
-									check(bSuccess);
+									const FScopedTransaction Transaction(LOCTEXT("SetBindingSkeleton", "Set Binding Skeleton"));
+									SetBinding->Modify();
 
-									SetBindingWidget->RepopulateTreeData();
-									AttributesListWidget->RepopulateListData();
-									SetBindingWidget->RepopulateTreeData();
+									const bool bSuccess = SetBinding->SetSkeleton(CastChecked<USkeleton>(AssetData.GetAsset()));
+									ensure(bSuccess);
 								}
 							})
 						.AllowClear(false)
