Add context menu action to unbind bones via the skeleton tree
* Assign 'Backspace' as hotkey to unbind

#Synced-CL 48947880

diff --git a/Engine/Plugins/Experimental/UAF/UAF/Source/UAFEditor/Internal/UAFStyle.h b/Engine/Plugins/Experimental/UAF/UAF/Source/UAFEditor/Internal/UAFStyle.h
--- a/Engine/Plugins/Experimental/UAF/UAF/Source/UAFEditor/Internal/UAFStyle.h
+++ b/Engine/Plugins/Experimental/UAF/UAF/Source/UAFEditor/Internal/UAFStyle.h
@@ -25,6 +25,8 @@ private:
 	UE_API void SetTraitStyles();
 	
 	UE_API void SetAssetWizardStyles();
+
+	UE_API void SetSetBindingStyles();
 };
 
 #undef UE_API
diff --git a/Engine/Plugins/Experimental/UAF/UAF/Source/UAFEditor/Private/AnimNextEditorModule.cpp b/Engine/Plugins/Experimental/UAF/UAF/Source/UAFEditor/Private/AnimNextEditorModule.cpp
--- a/Engine/Plugins/Experimental/UAF/UAF/Source/UAFEditor/Private/AnimNextEditorModule.cpp
+++ b/Engine/Plugins/Experimental/UAF/UAF/Source/UAFEditor/Private/AnimNextEditorModule.cpp
@@ -90,6 +90,7 @@
 #include "UAF/Debug/VariableWatchTabSummoner.h"
 #include "WorkspaceSchema.h"
 #include "Styling/SlateIconFinder.h"
+#include "UAF/AbstractSkeleton/Sets/Commands.h"
 
 #define LOCTEXT_NAMESPACE "AnimNextEditorModule"
 
@@ -144,6 +145,7 @@ void FAnimNextEditorModule::StartupModule()
 	FRigVMCommands::Register();
 	FAnimNextRigVMAssetCommands::Register();
 	FVariablesOutlinerCommands::Register();
+	FSetBindingEditorCommands::Register();
 	FOutlinerCommands::Register();
 	FUAFStyle::Get();
 
diff --git a/Engine/Plugins/Experimental/UAF/UAF/Source/UAFEditor/Private/UAF/AbstractSkeleton/Sets/Commands.cpp b/Engine/Plugins/Experimental/UAF/UAF/Source/UAFEditor/Private/UAF/AbstractSkeleton/Sets/Commands.cpp
new file mode 100644
--- /dev/null
+++ b/Engine/Plugins/Experimental/UAF/UAF/Source/UAFEditor/Private/UAF/AbstractSkeleton/Sets/Commands.cpp
@@ -0,0 +1,17 @@
+// Copyright Epic Games, Inc. All Rights Reserved.
+
+#include "UAF/AbstractSkeleton/Sets/Commands.h"
+
+#include "Framework/Commands/UICommandList.h"
+
+#define LOCTEXT_NAMESPACE "SetBindingEditorCommands"
+
+namespace UE::UAF::Editor
+{
+	void FSetBindingEditorCommands::RegisterCommands()
+	{
+		UI_COMMAND(UnbindSelection, "Unbind Selection", "Unbinds the current selection", EUserInterfaceActionType::Button, FInputChord(EKeys::BackSpace));
+	}
+}
+
+#undef LOCTEXT_NAMESPACE
\ No newline at end of file
diff --git a/Engine/Plugins/Experimental/UAF/UAF/Source/UAFEditor/Private/UAF/AbstractSkeleton/Sets/Commands.h b/Engine/Plugins/Experimental/UAF/UAF/Source/UAFEditor/Private/UAF/AbstractSkeleton/Sets/Commands.h
new file mode 100644
--- /dev/null
+++ b/Engine/Plugins/Experimental/UAF/UAF/Source/UAFEditor/Private/UAF/AbstractSkeleton/Sets/Commands.h
@@ -0,0 +1,32 @@
+// Copyright Epic Games, Inc. All Rights Reserved.
+
+#pragma once
+
+#include "CoreMinimal.h"
+#include "Framework/Commands/Commands.h"
+#include "UAFStyle.h"
+
+class FUICommandInfo;
+
+namespace UE::UAF::Editor
+{
+	class FSetBindingEditorCommands : public TCommands<FSetBindingEditorCommands>
+	{
+	public:
+		FSetBindingEditorCommands()
+			: TCommands<FSetBindingEditorCommands>
+			(
+				TEXT("SetBindingEditorCommands"),
+				NSLOCTEXT("SetBindingEditorCommands", "SetBindingEditorCommands", "Set Binding Editor Commands"),
+				NAME_None,
+				FUAFStyle::Get().GetStyleSetName()
+			)
+		{
+		}
+
+		virtual void RegisterCommands() override;
+
+		TSharedPtr<FUICommandInfo> UnbindSelection;
+	};
+
+}
\ No newline at end of file
diff --git a/Engine/Plugins/Experimental/UAF/UAF/Source/UAFEditor/Private/UAF/AbstractSkeleton/Sets/SSetSkeletonTree.cpp b/Engine/Plugins/Experimental/UAF/UAF/Source/UAFEditor/Private/UAF/AbstractSkeleton/Sets/SSetSkeletonTree.cpp
--- a/Engine/Plugins/Experimental/UAF/UAF/Source/UAFEditor/Private/UAF/AbstractSkeleton/Sets/SSetSkeletonTree.cpp
+++ b/Engine/Plugins/Experimental/UAF/UAF/Source/UAFEditor/Private/UAF/AbstractSkeleton/Sets/SSetSkeletonTree.cpp
@@ -3,11 +3,14 @@
 #include "UAF/AbstractSkeleton/Sets/SSetSkeletonTree.h"
 
 #include "Animation/Skeleton.h"
+#include "Framework/Commands/UICommandList.h"
+#include "Framework/MultiBox/MultiBoxBuilder.h"
 #include "ReferenceSkeleton.h"
 #include "ScopedTransaction.h"
 #include "SetDragDrop.h"
 #include "SkeletonTreeDragDrop.h"
 #include "UAF/AbstractSkeleton/AbstractSkeletonSetBinding.h"
+#include "UAF/AbstractSkeleton/Sets/Commands.h"
 #include "Widgets/Input/SSearchBox.h"
 #include "Widgets/Views/STreeView.h"
 
@@ -49,6 +52,7 @@ namespace UE::UAF::Sets
 							.TreeItemsSource(&RootItems)
 							.OnGenerateRow(this, &SSetsSkeletonTree::TreeView_OnGenerateRow)
 							.OnGetChildren(this, &SSetsSkeletonTree::TreeView_OnGetChildren)
+							.OnContextMenuOpening(this, &SSetsSkeletonTree::TreeView_OnContextMenuOpening)
 							.HeaderRow(
 								SNew(SHeaderRow)
 								+ SHeaderRow::Column("Set")
@@ -62,6 +66,7 @@ namespace UE::UAF::Sets
 			];
 
 		RepopulateTreeData();
+		BindCommands();
 	}
 
 	void SSetsSkeletonTree::Tick(const FGeometry& AllocatedGeometry, const double InCurrentTime, const float InDeltaTime)
@@ -73,6 +78,16 @@ namespace UE::UAF::Sets
 		}
 	}
 
+	FReply SSetsSkeletonTree::OnKeyDown(const FGeometry& MyGeometry, const FKeyEvent& InKeyEvent)
+	{
+		if (CommandList->ProcessCommandBindings(InKeyEvent))
+		{
+			return FReply::Handled();
+		}
+
+		return FReply::Unhandled();
+	}
+
 	void SSetsSkeletonTree::PostUndo(bool bSuccess)
 	{
 		RepopulateTreeData();
@@ -235,6 +250,17 @@ namespace UE::UAF::Sets
 		OutChildren = InParent->Children;
 	}
 
+	TSharedPtr<SWidget> SSetsSkeletonTree::TreeView_OnContextMenuOpening()
+	{
+		const Editor::FSetBindingEditorCommands& Commands = Editor::FSetBindingEditorCommands::Get();
+
+		FMenuBuilder MenuBuilder(true, CommandList.ToSharedRef());
+
+		MenuBuilder.AddMenuEntry(Commands.UnbindSelection);
+
+		return MenuBuilder.MakeWidget();
+	}
+
 	TArray<TSharedPtr<SSetsSkeletonTree::FTreeItem>> SSetsSkeletonTree::GetAllTreeItems()
 	{
 		TArray<TSharedPtr<FTreeItem>> AllItems;
@@ -318,10 +344,39 @@ namespace UE::UAF::Sets
 		}
 	}
 
+	void SSetsSkeletonTree::BindCommands()
+	{
+		CommandList = MakeShareable(new FUICommandList);
+
+		const Editor::FSetBindingEditorCommands& Commands = Editor::FSetBindingEditorCommands::Get();
+
+		CommandList->MapAction(Commands.UnbindSelection,
+			FExecuteAction::CreateSP(this, &SSetsSkeletonTree::HandleUnbindSelection),
+			FCanExecuteAction::CreateSP(this, &SSetsSkeletonTree::IsSelectionNonEmpty));
+	}
+
 	void SSetsSkeletonTree::HandleBindingsChanged()
 	{
 		bTreeDirty = true;
 	}
+
+	void SSetsSkeletonTree::HandleUnbindSelection()
+	{
+		const FScopedTransaction Transaction(LOCTEXT("UnbindSelection", "Unbind Selection"));
+		SetBinding->Modify();
+
+		TArray<TSharedPtr<FTreeItem>> Selection = TreeView->GetSelectedItems();
+
+		for (const TSharedPtr<FTreeItem>& Selected : Selection)
+		{
+			SetBinding->RemoveBoneFromSet(Selected->BoneName);
+		}
+	}
+
+	bool SSetsSkeletonTree::IsSelectionNonEmpty() const
+	{
+		return TreeView->GetNumItemsSelected() > 0;
+	}
 }
 
 #undef LOCTEXT_NAMESPACE
\ No newline at end of file
diff --git a/Engine/Plugins/Experimental/UAF/UAF/Source/UAFEditor/Private/UAF/AbstractSkeleton/Sets/SSetSkeletonTree.h b/Engine/Plugins/Experimental/UAF/UAF/Source/UAFEditor/Private/UAF/AbstractSkeleton/Sets/SSetSkeletonTree.h
--- a/Engine/Plugins/Experimental/UAF/UAF/Source/UAFEditor/Private/UAF/AbstractSkeleton/Sets/SSetSkeletonTree.h
+++ b/Engine/Plugins/Experimental/UAF/UAF/Source/UAFEditor/Private/UAF/AbstractSkeleton/Sets/SSetSkeletonTree.h
@@ -10,6 +10,7 @@ class ITableRow;
 class STableViewBase;
 template <typename> class STreeView;
 class UAbstractSkeletonSetBinding;
+class FUICommandList;
 
 namespace UE::UAF::Sets
 {
@@ -35,6 +36,7 @@ namespace UE::UAF::Sets
 	private:
 		/** SCompoundWidget */
 		virtual void Tick(const FGeometry& AllocatedGeometry, const double InCurrentTime, const float InDeltaTime) override;
+		virtual FReply OnKeyDown(const FGeometry& MyGeometry, const FKeyEvent& InKeyEvent) override;
 
 		/** FSelfRegisteringEditorUndoClient interface */
 		virtual void PostUndo(bool bSuccess) override;
@@ -42,15 +44,23 @@ namespace UE::UAF::Sets
 
 		void RepopulateTreeData();
 
+		void BindCommands();
+
 		// Called when the bindings inside the set binding asset we're editing have been modified
 		void HandleBindingsChanged();
-		
+
+		void HandleUnbindSelection();
+
+		bool IsSelectionNonEmpty() const;
+
 		TArray<TSharedPtr<FTreeItem>> GetAllTreeItems();
 
 		TSharedRef<ITableRow> TreeView_OnGenerateRow(TSharedPtr<FTreeItem> Item, const TSharedRef<STableViewBase>& OwnerTable);
 		
 		void TreeView_OnGetChildren(TSharedPtr<FTreeItem> InParent, TArray<TSharedPtr<FTreeItem>>& OutChildren);
 
+		TSharedPtr<SWidget> TreeView_OnContextMenuOpening();
+
 	private:
 		TWeakObjectPtr<UAbstractSkeletonSetBinding> SetBinding;
 
@@ -61,6 +71,8 @@ namespace UE::UAF::Sets
 		FDelegateHandle OnBindingsChangedHandle;
 
 		bool bTreeDirty = false;
+
+		TSharedPtr<FUICommandList> CommandList;
 	};
 
 }
\ No newline at end of file
diff --git a/Engine/Plugins/Experimental/UAF/UAF/Source/UAFEditor/Private/UAFStyle.cpp b/Engine/Plugins/Experimental/UAF/UAF/Source/UAFEditor/Private/UAFStyle.cpp
--- a/Engine/Plugins/Experimental/UAF/UAF/Source/UAFEditor/Private/UAFStyle.cpp
+++ b/Engine/Plugins/Experimental/UAF/UAF/Source/UAFEditor/Private/UAFStyle.cpp
@@ -1,4 +1,4 @@
-﻿// Copyright Epic Games, Inc. All Rights Reserved.
+// Copyright Epic Games, Inc. All Rights Reserved.
 
 #include "UAFStyle.h"
 #include "Brushes/SlateBoxBrush.h"
@@ -25,6 +25,7 @@ FUAFStyle::FUAFStyle()
 	SetVariablesOutlinerStyles();
 	SetTraitStyles();
 	SetAssetWizardStyles();
+	SetSetBindingStyles();
 
 	FSlateStyleRegistry::RegisterSlateStyle(*this);
 }
@@ -108,6 +109,13 @@ void FUAFStyle::SetAssetWizardStyles()
 	Set("AssetWizard.TemplateFolderHoverSelectedBackground", new FSlateRoundedBoxBrush(FStyleColors::PrimaryHover, 4.0f));
 }
 
+void FUAFStyle::SetSetBindingStyles()
+{
+	const FVector2D Icon16x16(16.0f, 16.0f);
+
+	Set("SetBindingEditorCommands.UnbindSelection", new CORE_IMAGE_BRUSH_SVG("Starship/Common/Unlinked", Icon16x16));
+}
+
 FUAFStyle& FUAFStyle::Get()
 {
 	static FUAFStyle Instance;
