UAF: Change to FTransactionallySafeRWLock as we can access the system pool allocator from Verse transactions

#Synced-CL 48962394

diff --git a/Engine/Plugins/Experimental/UAF/UAF/Source/UAF/Private/Module/SystemReference.cpp b/Engine/Plugins/Experimental/UAF/UAF/Source/UAF/Private/Module/SystemReference.cpp
--- a/Engine/Plugins/Experimental/UAF/UAF/Source/UAF/Private/Module/SystemReference.cpp
+++ b/Engine/Plugins/Experimental/UAF/UAF/Source/UAF/Private/Module/SystemReference.cpp
@@ -25,7 +25,7 @@ namespace UE::UAF::Private
 		TPool<FAnimNextModuleInstance> SystemInstances;
 
 		// Lock for concurrent allocations/resolves
-		FRWLock SystemInstancesLock;
+		FTransactionallySafeRWLock SystemInstancesLock;
 
 		// Actions to perform once world tick is completed each frame
 		TMpscQueue<TUniqueFunction<void()>> PendingActions;
