Disallow creating sets named 'None' in set collection editor

#Synced-CL 49029692

diff --git a/Engine/Plugins/Experimental/UAF/UAF/Source/UAF/Private/UAF/AbstractSkeleton/AbstractSkeletonSetCollection.cpp b/Engine/Plugins/Experimental/UAF/UAF/Source/UAF/Private/UAF/AbstractSkeleton/AbstractSkeletonSetCollection.cpp
--- a/Engine/Plugins/Experimental/UAF/UAF/Source/UAF/Private/UAF/AbstractSkeleton/AbstractSkeletonSetCollection.cpp
+++ b/Engine/Plugins/Experimental/UAF/UAF/Source/UAF/Private/UAF/AbstractSkeleton/AbstractSkeletonSetCollection.cpp
@@ -7,6 +7,11 @@
 
 bool UAbstractSkeletonSetCollection::AddSet(const FName InSetName, const FName InParentSetName)
 {
+	if (!IsValidSetName(InSetName))
+	{
+		return false;
+	}
+
 	if (HasSet(InSetName))
 	{
 		return false; // Set with this name already exists
@@ -32,6 +37,11 @@ bool UAbstractSkeletonSetCollection::HasSet(const FName InSetName) const
 
 bool UAbstractSkeletonSetCollection::RenameSet(const FName OldSetName, const FName NewSetName)
 {
+	if (!IsValidSetName(NewSetName))
+	{
+		return false;
+	}
+
 	if (HasSet(NewSetName))
 	{
 		return false; // Set with this name already exists
@@ -130,6 +140,11 @@ const TArrayView<const FAbstractSkeletonSet> UAbstractSkeletonSetCollection::Get
 	return SetHierarchy;
 }
 
+bool UAbstractSkeletonSetCollection::IsValidSetName(const FName SetName) const
+{
+	return !SetName.IsNone();
+}
+
 int32 UAbstractSkeletonSetCollection::GetSetIndex(const FName SetName) const
 {
 	for (int32 SetIndex = 0; SetIndex < SetHierarchy.Num(); ++SetIndex)
diff --git a/Engine/Plugins/Experimental/UAF/UAF/Source/UAF/Private/UAF/Attributes/AttributeBindingData.cpp b/Engine/Plugins/Experimental/UAF/UAF/Source/UAF/Private/UAF/Attributes/AttributeBindingData.cpp
--- a/Engine/Plugins/Experimental/UAF/UAF/Source/UAF/Private/UAF/Attributes/AttributeBindingData.cpp
+++ b/Engine/Plugins/Experimental/UAF/UAF/Source/UAF/Private/UAF/Attributes/AttributeBindingData.cpp
@@ -41,6 +41,7 @@ namespace UE::UAF
 
 				for (const FAbstractSkeletonSet& Set : SetCollection->GetSetHierarchy())
 				{
+					ensure(!Set.SetName.IsNone()); // The set collection should maintain this invariant
 					SetNames.AddUnique(Set.SetName);
 				}
 			}
diff --git a/Engine/Plugins/Experimental/UAF/UAF/Source/UAF/Public/UAF/AbstractSkeleton/AbstractSkeletonSetCollection.h b/Engine/Plugins/Experimental/UAF/UAF/Source/UAF/Public/UAF/AbstractSkeleton/AbstractSkeletonSetCollection.h
--- a/Engine/Plugins/Experimental/UAF/UAF/Source/UAF/Public/UAF/AbstractSkeleton/AbstractSkeletonSetCollection.h
+++ b/Engine/Plugins/Experimental/UAF/UAF/Source/UAF/Public/UAF/AbstractSkeleton/AbstractSkeletonSetCollection.h
@@ -52,6 +52,8 @@ public:
 #endif
 
 private:
+	bool IsValidSetName(const FName SetName) const;
+
 	int32 GetSetIndex(const FName SetName) const;
 
 	void SortSetHierarchy();
diff --git a/Engine/Plugins/Experimental/UAF/UAF/Source/UAFEditor/Private/UAF/AbstractSkeleton/Sets/SSetCollectionTreeView.cpp b/Engine/Plugins/Experimental/UAF/UAF/Source/UAFEditor/Private/UAF/AbstractSkeleton/Sets/SSetCollectionTreeView.cpp
--- a/Engine/Plugins/Experimental/UAF/UAF/Source/UAFEditor/Private/UAF/AbstractSkeleton/Sets/SSetCollectionTreeView.cpp
+++ b/Engine/Plugins/Experimental/UAF/UAF/Source/UAFEditor/Private/UAF/AbstractSkeleton/Sets/SSetCollectionTreeView.cpp
@@ -162,13 +162,20 @@ namespace UE::UAF::Editor
 								return true;
 							}
 
-							if (CanRenameSet(OldName, NewName))
+							if (NewName.IsNone())
 							{
-								return true;
+								OutErrorMessage = LOCTEXT("SetNameInvalid", "Set name cannot be 'None'");
+								return false;
+							}
+
+							const bool bSetNameAlreadyExists = SetCollection->HasSet(NewName);
+							if (bSetNameAlreadyExists)
+							{
+								OutErrorMessage = LOCTEXT("SetNameTaken", "Set name already taken");
+								return false;
 							}
 								
-							OutErrorMessage = LOCTEXT("SetNameTaken", "Name already taken");
-							return false;
+							return true;
 						})
 					.OnTextCommitted_Lambda([InItem, this](const FText& InText, ETextCommit::Type CommitInfo)
 						{
@@ -235,11 +242,6 @@ namespace UE::UAF::Editor
 		SetCollection->RenameSet(OldSetName, NewSetName);
 	}
 
-	bool SSetCollectionTreeView::CanRenameSet(const FName OldSetName, const FName NewSetName) const
-	{
-		return !SetCollection->HasSet(NewSetName);
-	}
-
 	FSetCollectionTreeItemPtr SSetCollectionTreeView::GetTreeItem(const FName SetName)
 	{
 		for (FSetCollectionTreeItemPtr TreeItem : GetAllTreeItems())
diff --git a/Engine/Plugins/Experimental/UAF/UAF/Source/UAFEditor/Private/UAF/AbstractSkeleton/Sets/SSetCollectionTreeView.h b/Engine/Plugins/Experimental/UAF/UAF/Source/UAFEditor/Private/UAF/AbstractSkeleton/Sets/SSetCollectionTreeView.h
--- a/Engine/Plugins/Experimental/UAF/UAF/Source/UAFEditor/Private/UAF/AbstractSkeleton/Sets/SSetCollectionTreeView.h
+++ b/Engine/Plugins/Experimental/UAF/UAF/Source/UAFEditor/Private/UAF/AbstractSkeleton/Sets/SSetCollectionTreeView.h
@@ -53,7 +53,6 @@ namespace UE::UAF::Editor
 		void TreeView_OnItemScrolledIntoView(FSetCollectionTreeItemPtr InItem, const TSharedPtr<ITableRow>& InWidget);
 		void TreeView_OnMouseButtonDoubleClick(FSetCollectionTreeItemPtr InItem);
 
-		bool CanRenameSet(const FName OldSetName, const FName NewSetName) const;
 		void RenameSet(const FName OldSetName, const FName NewSetName);
 
 		virtual FReply OnDrop(const FGeometry& MyGeometry, const FDragDropEvent& DragDropEvent) override;
