Added a basic AnimNext sample that broadcasts a single pose to N instances.

#virtualized
#Synced-CL 49075818

diff --git a/Engine/Plugins/Experimental/UAF/UAF/Source/UAF/Internal/Graph/RigUnit_UAFWriteInstancedSkinnedMeshPose.h b/Engine/Plugins/Experimental/UAF/UAF/Source/UAF/Internal/Graph/RigUnit_UAFWriteInstancedSkinnedMeshPose.h
new file mode 100644
--- /dev/null
+++ b/Engine/Plugins/Experimental/UAF/UAF/Source/UAF/Internal/Graph/RigUnit_UAFWriteInstancedSkinnedMeshPose.h
@@ -0,0 +1,38 @@
+﻿// Copyright Epic Games, Inc. All Rights Reserved.
+
+#pragma once
+
+#include "AnimNextExecuteContext.h"
+#include "Graph/AnimNext_LODPose.h"
+#include "Graph/RigUnit_AnimNextBase.h"
+class UInstancedSkinnedMeshComponent;
+
+#include "RigUnit_UAFWriteInstancedSkinnedMeshPose.generated.h"
+
+#define UE_API UAF_API
+
+/** Writes a pose to an instanced skinned mesh component */
+USTRUCT(meta=(DisplayName="Write Pose Instanced", Category="Animation Graph", NodeColor="0, 1, 1", Keywords="Output,Pose,Port"))
+struct FRigUnit_UAFWriteInstancedSkinnedMeshPose : public FRigUnit_AnimNextBase
+{
+	GENERATED_BODY()
+
+	RIGVM_METHOD()
+	UE_API void Execute();
+
+	virtual FString GetUnitSubTitle() const { return TEXT("Skeletal Mesh Component"); };
+
+	// Pose to write
+	UPROPERTY(EditAnywhere, Category = "Graph", meta = (Input))
+	FAnimNextGraphLODPose Pose;
+
+	// Mesh to write to. 
+	UPROPERTY(EditAnywhere, Category = "Graph", meta = (Input))
+	TObjectPtr<class UInstancedSkinnedMeshComponent> SkinnedMeshComponent;
+
+	// The execution result
+	UPROPERTY(EditAnywhere, DisplayName = "Execute", Category = "BeginExecution", meta = (Input, Output))
+	FAnimNextExecuteContext ExecuteContext;
+};
+
+#undef UE_API
\ No newline at end of file
diff --git a/Engine/Plugins/Experimental/UAF/UAF/Source/UAF/Private/Graph/RigUnit_UAFWriteInstancedSkinnedMeshPose.cpp b/Engine/Plugins/Experimental/UAF/UAF/Source/UAF/Private/Graph/RigUnit_UAFWriteInstancedSkinnedMeshPose.cpp
new file mode 100644
--- /dev/null
+++ b/Engine/Plugins/Experimental/UAF/UAF/Source/UAF/Private/Graph/RigUnit_UAFWriteInstancedSkinnedMeshPose.cpp
@@ -0,0 +1,115 @@
+// Copyright Epic Games, Inc. All Rights Reserved.
+
+#include "Graph/RigUnit_UAFWriteInstancedSkinnedMeshPose.h"
+
+#include "AnimationRuntime.h"
+#include "Animation/NamedValueArray.h"
+#include "Animation/AnimRuntimeTransformProviderData.h"
+#include "AnimNextStats.h"
+#include "GenerationTools.h"
+#include "Component/SkinnedMeshComponentExtensions.h"
+#include "Components/InstancedSkinnedMeshComponent.h"
+#include "Engine/SkeletalMesh.h"
+#include "Logging/StructuredLog.h"
+#include "Module/AnimNextModuleInstance.h"
+
+#include UE_INLINE_GENERATED_CPP_BY_NAME(RigUnit_UAFWriteInstancedSkinnedMeshPose)
+
+namespace UE::Anim
+{
+
+extern void ConvertLocalSpaceToComponentSpace(
+	const USkinnedMeshComponent* InComponent,
+	const UE::UAF::FLODPoseHeap& InLODPose,
+	TArrayView<FTransform> OutComponentSpaceTransforms);
+
+}
+
+FRigUnit_UAFWriteInstancedSkinnedMeshPose_Execute()
+{
+	FTaskTagScope TaskTag(ETaskTag::EParallelGameThread);
+	TRACE_CPUPROFILER_EVENT_SCOPE(AnimNext::WriteInstancedSkinnedMeshPose);
+
+	if (!SkinnedMeshComponent)
+	{
+		return;
+	}
+
+	USkinnedAsset* SkinnedAsset = SkinnedMeshComponent->GetSkinnedAsset();
+	if (!SkinnedAsset)
+	{
+		return;
+	}
+
+	UAnimRuntimeTransformProviderData* AnimRuntimeData = Cast<UAnimRuntimeTransformProviderData>(SkinnedMeshComponent->GetTransformProvider());
+	if (!AnimRuntimeData)
+	{
+		return;
+	}
+
+	if (!Pose.LODPose.IsValid() || Pose.LODPose.LODLevel != 0)
+	{
+		return;
+	}
+
+	const TArray<FMatrix44f>& RefBasesInvMatrix = SkinnedAsset->GetRefBasesInvMatrix();
+
+	TArray<FTransform> ComponentSpaceTransforms;
+	TArray<FCompressedBoneTransform> ReferenceToLocalTransforms;
+	ComponentSpaceTransforms.SetNumUninitialized(RefBasesInvMatrix.Num());
+	ReferenceToLocalTransforms.SetNumUninitialized(RefBasesInvMatrix.Num());
+	UE::Anim::ConvertLocalSpaceToComponentSpace(SkinnedMeshComponent, Pose.LODPose, ComponentSpaceTransforms);
+
+	for (int32 Index = 0; Index < ComponentSpaceTransforms.Num(); ++Index)
+	{
+		FMatrix44f Transform = (FMatrix44f)ComponentSpaceTransforms[Index].ToMatrixWithScale();
+		VectorMatrixMultiply(&Transform, &RefBasesInvMatrix[Index], &Transform);
+		StoreCompressedBoneTransform(&ReferenceToLocalTransforms[Index], Transform);
+	}
+
+	for (int32 InstanceIndex = 0; InstanceIndex < SkinnedMeshComponent->GetInstanceCount(); ++InstanceIndex)
+	{
+		FPrimitiveInstanceId InstanceId = SkinnedMeshComponent->GetInstanceId(InstanceIndex);
+		int32 TrackIndex;
+
+		if (!SkinnedMeshComponent->GetInstanceAnimationIndex(InstanceId, TrackIndex))
+		{
+			continue;
+		}
+
+		if (!AnimRuntimeData->IsValidIndex(TrackIndex))
+		{
+			continue;
+		}
+
+		FAnimRuntimeTrackAllocation Allocation = AnimRuntimeData->UpdateTrack(TrackIndex, EPreviousBoneTransformUpdateMode::None);
+
+		if (!Allocation.IsValid())
+		{
+			continue;
+		}
+
+		Allocation.Current->SetRevisionNumber(GFrameNumber);
+		TArrayView<FCompressedBoneTransform> OutReferenceToLocal = Allocation.Current->GetTransforms();
+		check(OutReferenceToLocal.Num() == ReferenceToLocalTransforms.Num());
+
+		for (int32 TransformIndex = 0; TransformIndex < ReferenceToLocalTransforms.Num(); ++TransformIndex)
+		{
+			OutReferenceToLocal[TransformIndex] = ReferenceToLocalTransforms[TransformIndex];
+		}
+	}
+
+	FAnimNextModuleInstance::RunTaskOnGameThread([WeakComponent = TWeakObjectPtr<UInstancedSkinnedMeshComponent>(SkinnedMeshComponent)]()
+		{
+			UInstancedSkinnedMeshComponent* Component = WeakComponent.Get();
+			if (Component == nullptr)
+			{
+				return;
+			}
+
+			Component->InvalidateCachedBounds();
+			Component->UpdateBounds();
+			Component->MarkRenderTransformDirty();
+			Component->MarkRenderInstancesDirty();
+		});
+}
\ No newline at end of file
diff --git a/Engine/Plugins/Experimental/UAF/UAFAnimGraph/Source/UAFAnimGraph/Private/AnimNextAnimGraphModule.cpp b/Engine/Plugins/Experimental/UAF/UAFAnimGraph/Source/UAFAnimGraph/Private/AnimNextAnimGraphModule.cpp
--- a/Engine/Plugins/Experimental/UAF/UAFAnimGraph/Source/UAFAnimGraph/Private/AnimNextAnimGraphModule.cpp
+++ b/Engine/Plugins/Experimental/UAF/UAFAnimGraph/Source/UAFAnimGraph/Private/AnimNextAnimGraphModule.cpp
@@ -10,6 +10,7 @@
 #include "Animation/AnimSequence.h"
 #include "Animation/BlendSpace.h"
 #include "Components/SkeletalMeshComponent.h"
+#include "Components/InstancedSkinnedMeshComponent.h"
 #include "Factory/AnimGraphFactory.h"
 #include "TraitInterfaces/IUpdate.h"
 #include "Graph/AnimNextAnimationGraph.h"
@@ -65,6 +66,7 @@ void FAnimNextAnimGraphModule::StartupModule()
 	{
 		{ UAnimNextAnimationGraph::StaticClass(), FRigVMRegistry::ERegisterObjectOperation::Class },
 		{ USkeletalMeshComponent::StaticClass(), FRigVMRegistry::ERegisterObjectOperation::ClassAndParents },
+		{ UInstancedSkinnedMeshComponent::StaticClass(), FRigVMRegistry::ERegisterObjectOperation::ClassAndParents },
 	};
 
 	RigVMRegistry.RegisterObjectTypes(AllowedObjectTypes);
