Update PostLoad of set binding asset to remove bindings to deleted sets

#Synced-CL 49029692

diff --git a/Engine/Plugins/Experimental/UAF/UAF/Source/UAF/Private/UAF/AbstractSkeleton/AbstractSkeletonSetBinding.cpp b/Engine/Plugins/Experimental/UAF/UAF/Source/UAF/Private/UAF/AbstractSkeleton/AbstractSkeletonSetBinding.cpp
--- a/Engine/Plugins/Experimental/UAF/UAF/Source/UAF/Private/UAF/AbstractSkeleton/AbstractSkeletonSetBinding.cpp
+++ b/Engine/Plugins/Experimental/UAF/UAF/Source/UAF/Private/UAF/AbstractSkeleton/AbstractSkeletonSetBinding.cpp
@@ -11,6 +11,10 @@ void UAbstractSkeletonSetBinding::PostLoad()
 {
 	Super::PostLoad();
 
+#if WITH_EDITOR
+	FilterOrphanedBindings();
+#endif
+
 	RebuildBindingData();
 
 #if WITH_EDITOR
@@ -267,6 +271,48 @@ void UAbstractSkeletonSetBinding::RebuildBindingData()
 	AttributeBindingData = UE::UAF::MakeAttributeBindingData(*this);
 }
 
+#if WITH_EDITOR
+
+void UAbstractSkeletonSetBinding::FilterOrphanedBindings()
+{
+	TSet<FName> ValidSetNames;
+
+	// While 'None' is not a valid set name, it is used internally
+	// to contain known unbound attributes
+	ValidSetNames.Add(NAME_None);
+
+	if (SetCollection)
+	{
+		for (const FAbstractSkeletonSet& Set : SetCollection->GetSetHierarchy())
+		{
+			ValidSetNames.Add(Set.SetName);
+		}
+	}
+
+	int32 CountRemoved = 0;
+
+	CountRemoved += BoneBindings.RemoveAll([&](FAbstractSkeleton_BoneBinding& Binding)
+		{
+			return !ValidSetNames.Contains(Binding.SetName);
+		});
+
+	CountRemoved += AttributeBindings.RemoveAll([&](FAbstractSkeleton_AttributeBinding& Binding)
+		{
+			return !ValidSetNames.Contains(Binding.SetName);
+		});
+
+	if (CountRemoved > 0)
+	{
+		OnBindingsChanged();
+		UE_LOG(LogAnimation, Warning, TEXT("Removed %d bindings whilst loading '%s' due to those sets no longer existing in Set Collection '%s'"),
+			CountRemoved,
+			*GetFullName(),
+			SetCollection ? *SetCollection->GetFullName() : TEXT("None"));
+	}
+}
+
+#endif // WITH_EDITOR
+
 bool UAbstractSkeletonSetBinding::SetSkeleton(const TObjectPtr<USkeleton> InSkeleton)
 {
 	if (Skeleton == InSkeleton)
@@ -350,6 +396,7 @@ void UAbstractSkeletonSetBinding::OnSkeletonChanged()
 
 void UAbstractSkeletonSetBinding::OnSetCollectionChanged()
 {
+	FilterOrphanedBindings();
 	RebuildBindingData();
 }
 
diff --git a/Engine/Plugins/Experimental/UAF/UAF/Source/UAF/Public/UAF/AbstractSkeleton/AbstractSkeletonSetBinding.h b/Engine/Plugins/Experimental/UAF/UAF/Source/UAF/Public/UAF/AbstractSkeleton/AbstractSkeletonSetBinding.h
--- a/Engine/Plugins/Experimental/UAF/UAF/Source/UAF/Public/UAF/AbstractSkeleton/AbstractSkeletonSetBinding.h
+++ b/Engine/Plugins/Experimental/UAF/UAF/Source/UAF/Public/UAF/AbstractSkeleton/AbstractSkeletonSetBinding.h
@@ -102,7 +102,12 @@ public:
 private:
 #if WITH_EDITOR
 	void OnSkeletonChanged();
+
 	void OnSetCollectionChanged();
+
+	// If a change is made to the set collection while we are loaded/unloaded we need
+	// to remove any bindings that are bound to invalid sets
+	void FilterOrphanedBindings();
 #endif
 
 	// Called internally whenever a change is made to the attribute bindings (bones + attributes)
