Sort bindings in the set binding tree view
* A set tree item's children are ordered by bone bindings -> attribute bindings -> children sets.
* Neighbouring bone bindings are ordered by their topology
* Neighbouring attribute bindings are ordered alphabetically
* Children sets are ordered according to their ordering in the set colleciton

#Synced-CL 49029692

diff --git a/Engine/Plugins/Experimental/UAF/UAF/Source/UAFEditor/Private/UAF/AbstractSkeleton/Sets/SSetBinding.cpp b/Engine/Plugins/Experimental/UAF/UAF/Source/UAFEditor/Private/UAF/AbstractSkeleton/Sets/SSetBinding.cpp
--- a/Engine/Plugins/Experimental/UAF/UAF/Source/UAFEditor/Private/UAF/AbstractSkeleton/Sets/SSetBinding.cpp
+++ b/Engine/Plugins/Experimental/UAF/UAF/Source/UAFEditor/Private/UAF/AbstractSkeleton/Sets/SSetBinding.cpp
@@ -191,6 +191,8 @@ namespace UE::UAF::Sets
 		// Rebuild tree data
 		RootItems.Reset();
 
+		int32 SetIndex = 0;
+
 		if (SetBinding.IsValid())
 		{
 			TMap<FName, TSharedPtr<FTreeItem_Set>> TreeItemMap;
@@ -199,7 +201,7 @@ namespace UE::UAF::Sets
 			{
 				for (const FAbstractSkeletonSet& Set : SetBinding->GetSetCollection()->GetSetHierarchy())
 				{
-					TSharedPtr<FTreeItem_Set> SetItem = MakeShared<FTreeItem_Set>(*this, Set, false);
+					TSharedPtr<FTreeItem_Set> SetItem = MakeShared<FTreeItem_Set>(*this, Set, false, SetIndex++);
 
 					if (Set.ParentSetName == NAME_None)
 					{
@@ -217,9 +219,15 @@ namespace UE::UAF::Sets
 				}
 			}
 
+			const USkeleton* Skeleton = SetBinding->GetSkeleton();
+			const FReferenceSkeleton& RefSkeleton = Skeleton->GetReferenceSkeleton();
+
 			for (const FAbstractSkeleton_BoneBinding& Binding : SetBinding->GetBoneBindings())
 			{
-				TSharedPtr<FTreeItem_Bone> BindingItem = MakeShared<FTreeItem_Bone>(*this, Binding);
+				const int32 BoneIndex = RefSkeleton.FindBoneIndex(Binding.BoneName);
+				check(RefSkeleton.IsValidIndex(BoneIndex));
+
+				TSharedPtr<FTreeItem_Bone> BindingItem = MakeShared<FTreeItem_Bone>(*this, Binding, BoneIndex);
 
 				if (TreeItemMap.Contains(Binding.SetName))
 				{
@@ -233,7 +241,7 @@ namespace UE::UAF::Sets
 					// This binding can be fixed by re-adding the set in the set collection or
 					// by removing this binding. Show this binding as an orphan.
 
-					TSharedPtr<FTreeItem_Set> DeletedSet = MakeShared<FTreeItem_Set>(*this, FAbstractSkeletonSet(Binding.SetName, NAME_None), true);
+					TSharedPtr<FTreeItem_Set> DeletedSet = MakeShared<FTreeItem_Set>(*this, FAbstractSkeletonSet(Binding.SetName, NAME_None), true, SetIndex++);
 					DeletedSet->Children.Insert(BindingItem, 0);
 
 					TreeItemMap.Add(Binding.SetName, DeletedSet);
@@ -253,6 +261,11 @@ namespace UE::UAF::Sets
 			}
 		}
 
+		for (TSharedPtr<ITreeItem> RootItem : RootItems)
+		{
+			RootItem->SortChildren();
+		}
+
 		TreeView->RebuildList();
 
 		// Restore expanded/collapsed state
@@ -575,6 +588,49 @@ namespace UE::UAF::Sets
 		return true;
 	}
 
+	void SSetBinding::FTreeItem_Set::SortChildren()
+	{
+		for (const auto& Child : Children)
+		{
+			Child->SortChildren();
+		}
+
+		Algo::StableSort(Children, [](const TSharedPtr<ITreeItem>& A, TSharedPtr<ITreeItem> B) -> bool
+			{
+				if (A->GetType() != B->GetType())
+				{
+					return static_cast<uint8>(A->GetType()) < static_cast<uint8>(B->GetType());
+				}
+
+				switch (A->GetType())
+				{
+					case ETreeItem::Set:
+					{
+						TSharedPtr<FTreeItem_Set> SetA = StaticCastSharedPtr<FTreeItem_Set>(A);
+						TSharedPtr<FTreeItem_Set> SetB = StaticCastSharedPtr<FTreeItem_Set>(B);
+
+						return SetA->SetIndex < SetB->SetIndex;
+					}
+					case ETreeItem::Bone:
+					{
+						TSharedPtr<FTreeItem_Bone> BoneA = StaticCastSharedPtr<FTreeItem_Bone>(A);
+						TSharedPtr<FTreeItem_Bone> BoneB = StaticCastSharedPtr<FTreeItem_Bone>(B);
+
+						return BoneA->BoneIndex < BoneB->BoneIndex;
+					}
+					case ETreeItem::Attribute:
+					{
+						TSharedPtr<FTreeItem_Attribute> AttributeA = StaticCastSharedPtr<FTreeItem_Attribute>(A);
+						TSharedPtr<FTreeItem_Attribute> AttributeB = StaticCastSharedPtr<FTreeItem_Attribute>(B);
+
+						return AttributeA->Binding.Attribute.GetName().Compare(AttributeB->Binding.Attribute.GetName()) < 0;
+					}
+				}
+
+				return false;
+			});
+	}
+
 	bool SSetBinding::FTreeItem_Bone::GenerateContextMenu(FMenuBuilder& Menu)
 	{
 		FUIAction UnbindBone(
diff --git a/Engine/Plugins/Experimental/UAF/UAF/Source/UAFEditor/Private/UAF/AbstractSkeleton/Sets/SSetBinding.h b/Engine/Plugins/Experimental/UAF/UAF/Source/UAFEditor/Private/UAF/AbstractSkeleton/Sets/SSetBinding.h
--- a/Engine/Plugins/Experimental/UAF/UAF/Source/UAFEditor/Private/UAF/AbstractSkeleton/Sets/SSetBinding.h
+++ b/Engine/Plugins/Experimental/UAF/UAF/Source/UAFEditor/Private/UAF/AbstractSkeleton/Sets/SSetBinding.h
@@ -23,6 +23,15 @@ namespace UE::UAF::Sets
 		SLATE_END_ARGS()
 
 	public:
+
+		// Sorting here determines ordering in the tree view
+		enum class ETreeItem : uint8
+		{
+			Bone,
+			Attribute,
+			Set
+		};
+
 		struct ITreeItem
 		{
 			ITreeItem(SSetBinding& InSetBindingWidget)
@@ -31,55 +40,62 @@ namespace UE::UAF::Sets
 			}
 			
 			virtual ~ITreeItem() = default;
-			virtual FName GetType() const = 0;
+			virtual ETreeItem GetType() const = 0;
 			virtual void GetChildren(TArray<TSharedPtr<ITreeItem>>& OutChildren) = 0;
 			virtual TSharedRef<ITableRow> GenerateRow(const TSharedRef<STableViewBase>& OwnerTable) = 0;
 			virtual bool GenerateContextMenu(FMenuBuilder& Menu) = 0;
-			
+			virtual void SortChildren() = 0;
+
 			SSetBinding& SetBindingWidget;
 		};
 
 		struct FTreeItem_Set : public ITreeItem
 		{
-			FTreeItem_Set(SSetBinding& InSetBindingWidget, FAbstractSkeletonSet InSet, bool bInIsDeletedSet)
+			FTreeItem_Set(SSetBinding& InSetBindingWidget, FAbstractSkeletonSet InSet, bool bInIsDeletedSet, int32 InSetIndex)
 				: ITreeItem(InSetBindingWidget)
 				, Set(InSet)
 				, bIsDeletedSet(bInIsDeletedSet)
+				, SetIndex(InSetIndex)
 			{
 			}
 			
 			virtual ~FTreeItem_Set() override = default;
 			
-			static FName StaticGetType() { return FName("Set"); }
+			static ETreeItem StaticGetType() { return ETreeItem::Set; }
 			
-			virtual FName GetType() const override { return StaticGetType(); }
+			virtual ETreeItem GetType() const override { return StaticGetType(); }
 			virtual void GetChildren(TArray<TSharedPtr<ITreeItem>>& OutChildren) override;
 			virtual TSharedRef<ITableRow> GenerateRow(const TSharedRef<STableViewBase>& OwnerTable) override;
 			virtual bool GenerateContextMenu(FMenuBuilder& Menu) override;
+			virtual void SortChildren() override;
 
 			FAbstractSkeletonSet Set;
 			TArray<TSharedPtr<ITreeItem>> Children;
 			bool bIsDeletedSet;
+			int32 SetIndex;
 		};
 
 		struct FTreeItem_Bone : public ITreeItem
 		{
-			FTreeItem_Bone(SSetBinding& InSetBindingWidget, FAbstractSkeleton_BoneBinding InBinding)
+			FTreeItem_Bone(SSetBinding& InSetBindingWidget, FAbstractSkeleton_BoneBinding InBinding, int32 InBoneIndex)
 				: ITreeItem(InSetBindingWidget)
 				, Binding(InBinding)
+				, BoneIndex(InBoneIndex)
 			{
 			}
 			
 			virtual ~FTreeItem_Bone() override = default;
 
-			static FName StaticGetType() { return FName("Bone"); }
+			static ETreeItem StaticGetType() { return ETreeItem::Bone; }
 			
-			virtual FName GetType() const override { return StaticGetType(); }
+			virtual ETreeItem GetType() const override { return StaticGetType(); }
 			virtual void GetChildren(TArray<TSharedPtr<ITreeItem>>& OutChildren) override;
 			virtual TSharedRef<ITableRow> GenerateRow(const TSharedRef<STableViewBase>& OwnerTable) override;
 			virtual bool GenerateContextMenu(FMenuBuilder& Menu) override;
+			virtual void SortChildren() override {};
 
 			FAbstractSkeleton_BoneBinding Binding;
+			int32 BoneIndex;
 		};
 		
 		struct FTreeItem_Attribute : public ITreeItem
@@ -92,12 +108,13 @@ namespace UE::UAF::Sets
 			
 			virtual ~FTreeItem_Attribute() override = default;
 
-			static FName StaticGetType() { return FName("Attribute"); }
+			static ETreeItem StaticGetType() { return ETreeItem::Attribute; }
 			
-			virtual FName GetType() const override { return StaticGetType(); }
+			virtual ETreeItem GetType() const override { return StaticGetType(); }
 			virtual void GetChildren(TArray<TSharedPtr<ITreeItem>>& OutChildren) override;
 			virtual TSharedRef<ITableRow> GenerateRow(const TSharedRef<STableViewBase>& OwnerTable) override;
 			virtual bool GenerateContextMenu(FMenuBuilder& Menu) override;
+			virtual void SortChildren() override {};
 
 			FAbstractSkeleton_AttributeBinding Binding;
 		};
