UAF: Fix additional AutoRTFM hazards

#Synced-CL 49071325

diff --git a/Engine/Plugins/Experimental/UAF/UAF/Source/UAF/Internal/AnimNextDebugDraw.h b/Engine/Plugins/Experimental/UAF/UAF/Source/UAF/Internal/AnimNextDebugDraw.h
--- a/Engine/Plugins/Experimental/UAF/UAF/Source/UAF/Internal/AnimNextDebugDraw.h
+++ b/Engine/Plugins/Experimental/UAF/UAF/Source/UAF/Internal/AnimNextDebugDraw.h
@@ -1,4 +1,4 @@
-﻿// Copyright Epic Games, Inc. All Rights Reserved.
+// Copyright Epic Games, Inc. All Rights Reserved.
 
 #pragma once
 
@@ -7,6 +7,7 @@
 #include "PrimitiveSceneProxy.h"
 #include "PrimitiveSceneProxyDesc.h"
 #include "RigVMCore/RigVMDrawInterface.h"
+#include "Misc/TransactionallySafeRWLock.h"
 
 class UObject;
 struct FRigVMDrawInterface;
@@ -86,7 +87,7 @@ private:
 	bool bIsRegistered = false;
 
 	// Lock to prevent unregistration during WT drawing
-	FRWLock Lock;
+	FTransactionallySafeRWLock Lock;
 
 	friend struct ::FAnimNextModuleInstance;
 	friend struct UE::UAF::FModuleEventTickFunction;
diff --git a/Engine/Plugins/Experimental/UAF/UAF/Source/UAF/Internal/Module/AnimNextModuleInstance.h b/Engine/Plugins/Experimental/UAF/UAF/Source/UAF/Internal/Module/AnimNextModuleInstance.h
--- a/Engine/Plugins/Experimental/UAF/UAF/Source/UAF/Internal/Module/AnimNextModuleInstance.h
+++ b/Engine/Plugins/Experimental/UAF/UAF/Source/UAF/Internal/Module/AnimNextModuleInstance.h
@@ -14,6 +14,7 @@
 #include "StructUtils/InstancedStruct.h"
 #include "Module/AnimNextModuleInitMethod.h"
 #include "Module/TaskRunLocation.h"
+#include "Misc/TransactionallySafeRWLock.h"
 #include "AnimNextModuleInstance.generated.h"
 
 #define UE_API UAF_API
@@ -222,13 +223,13 @@ private:
 	UE::UAF::FTraitEventList OutputEventList;
 
 	// Lock to ensure event list actions are thread safe
-	FRWLock EventListLock;
+	FTransactionallySafeRWLock EventListLock;
 
 	// Lock to ensure proxy access is thread safe
-	FRWLock ProxyLock;
+	FTransactionallySafeRWLock ProxyLock;
 
 	// Lock guarding against gross concurrent access via system references (safety only: should rarely contend if dependencies are set correctly)
-	FRWLock SystemReferenceLock;
+	FTransactionallySafeRWLock SystemReferenceLock;
 
 	// Proxy public variables, double-buffered
 	UPROPERTY(Transient)
diff --git a/Engine/Plugins/Experimental/UAF/UAF/Source/UAF/Internal/Module/ModuleTickFunction.h b/Engine/Plugins/Experimental/UAF/UAF/Source/UAF/Internal/Module/ModuleTickFunction.h
--- a/Engine/Plugins/Experimental/UAF/UAF/Source/UAF/Internal/Module/ModuleTickFunction.h
+++ b/Engine/Plugins/Experimental/UAF/UAF/Source/UAF/Internal/Module/ModuleTickFunction.h
@@ -7,6 +7,9 @@
 #include "Containers/SpscQueue.h"
 #include "TaskSyncManager.h"
 #include "Script/UAFScriptComponent.h"
+#if UE_AUTORTFM
+#include "Containers/TransactionallySafeSpscQueue.h"
+#endif
 
 struct FAnimNextModuleInstance;
 
@@ -57,8 +60,14 @@ struct FModuleEventTickFunction : public FTickFunction
 	FAnimNextModuleInstance* ModuleInstance = nullptr;
 
 	// Tasks queued to execute once per event
-	TMpscQueue<TUniqueFunction<void(const FModuleTaskContext&)>> PreExecuteTasks;
-	TMpscQueue<TUniqueFunction<void(const FModuleTaskContext&)>> PostExecuteTasks;
+#if UE_AUTORTFM
+	using FTaskQueueType = TTransactionallySafeSpscQueue<TUniqueFunction<void(const FModuleTaskContext&)>>;
+#else
+	using FTaskQueueType = TMpscQueue<TUniqueFunction<void(const FModuleTaskContext&)>>;
+#endif
+
+	FTaskQueueType PreExecuteTasks;
+	FTaskQueueType PostExecuteTasks;
 
 	// Tasks queued to execute on the game thread after the current tick completes
 	TArray<TUniqueFunction<void()>> PostExecuteGameThreadTaskList;
diff --git a/Engine/Plugins/Experimental/UAF/UAF/Source/UAF/Private/AnimNextDebugDraw.cpp b/Engine/Plugins/Experimental/UAF/UAF/Source/UAF/Private/AnimNextDebugDraw.cpp
--- a/Engine/Plugins/Experimental/UAF/UAF/Source/UAF/Private/AnimNextDebugDraw.cpp
+++ b/Engine/Plugins/Experimental/UAF/UAF/Source/UAF/Private/AnimNextDebugDraw.cpp
@@ -168,7 +168,7 @@ FDebugDraw::FDebugDraw(UObject* InOwner)
 
 void FDebugDraw::RemovePrimitive()
 {
-	FWriteScopeLock WriteLock(Lock);
+	UE::TWriteScopeLock WriteLock(Lock);
 
 	Scene->RemovePrimitive(&SceneDesc);
 	bIsRegistered = false;
@@ -182,7 +182,7 @@ void FDebugDraw::RemovePrimitive()
 
 void FDebugDraw::Draw()
 {
-	FWriteScopeLock WriteLock(Lock);
+	UE::TWriteScopeLock WriteLock(Lock);
 
 	if (!bIsRegistered)
 	{
diff --git a/Engine/Plugins/Experimental/UAF/UAF/Source/UAF/Private/Module/AnimNextModuleInstance.cpp b/Engine/Plugins/Experimental/UAF/UAF/Source/UAF/Private/Module/AnimNextModuleInstance.cpp
--- a/Engine/Plugins/Experimental/UAF/UAF/Source/UAF/Private/Module/AnimNextModuleInstance.cpp
+++ b/Engine/Plugins/Experimental/UAF/UAF/Source/UAF/Private/Module/AnimNextModuleInstance.cpp
@@ -404,7 +404,7 @@ void FAnimNextModuleInstance::CopyProxyVariables()
 {
 	int32 ProxyReadIndex = 0;
 	{
-		FWriteScopeLock WriteLock(ProxyLock);
+		UE::TWriteScopeLock WriteLock(ProxyLock);
 
 		// Flip the write buffer index
 		ProxyReadIndex = ProxyWriteIndex;
@@ -543,14 +543,14 @@ void FAnimNextModuleInstance::RaiseTraitEvents(const UE::UAF::FTraitEventList& E
 
 EPropertyBagResult FAnimNextModuleInstance::SetProxyVariable(const FAnimNextVariableReference& InVariable, const FAnimNextParamType& InType, TConstArrayView<uint8> InData)
 {
-	FWriteScopeLock WriteLock(ProxyLock);
+	UE::TWriteScopeLock WriteLock(ProxyLock);
 	FUAFInstanceVariableDataProxy& Proxy = ProxyVariables[ProxyWriteIndex];
 	return Proxy.SetVariable(InVariable, InType, InData);
 }
 
 EPropertyBagResult FAnimNextModuleInstance::WriteProxyVariable(const FAnimNextVariableReference& InVariable, const FAnimNextParamType& InType, TFunctionRef<void(TArrayView<uint8>)> InFunction)
 {
-	FWriteScopeLock WriteLock(ProxyLock);
+	UE::TWriteScopeLock WriteLock(ProxyLock);
 	FUAFInstanceVariableDataProxy& Proxy = ProxyVariables[ProxyWriteIndex];
 	TArrayView<uint8> Data;
 	EPropertyBagResult Result = Proxy.WriteVariable(InVariable, InType, Data);
@@ -563,7 +563,7 @@ EPropertyBagResult FAnimNextModuleInstance::WriteProxyVariable(const FAnimNextVa
 
 EPropertyBagResult FAnimNextModuleInstance::ReadVariableFromExternal(const FAnimNextVariableReference& InVariable, const FAnimNextParamType& InType, TFunctionRef<void(TConstArrayView<uint8>)> InFunction)
 {
-	FReadScopeLock Lock(SystemReferenceLock);
+	UE::TReadScopeLock Lock(SystemReferenceLock);
 
 	TArrayView<uint8> Data;
 	EPropertyBagResult Result = Variables.AccessVariable(InVariable, InType, Data);
@@ -576,7 +576,7 @@ EPropertyBagResult FAnimNextModuleInstance::ReadVariableFromExternal(const FAnim
 
 bool FAnimNextModuleInstance::ReadComponentFromExternal(UScriptStruct* InComponentType, TFunctionRef<void(FConstStructView InComponentStruct)> InFunction)
 {
-	FReadScopeLock Lock(SystemReferenceLock);
+	UE::TReadScopeLock Lock(SystemReferenceLock);
 
 	if (const FUAFAssetInstanceComponent* Component = TryGetComponent(InComponentType))
 	{
@@ -635,6 +635,11 @@ void FAnimNextModuleInstance::QueueTask(FName InEventName, TUniqueFunction<void(
 {
 	using namespace UE::UAF;
 
+#if UE_AUTORTFM
+	// As we only have a SPSC queue under AutoRTFM, we limit enquing tasks to the game thread for now
+	check(IsInGameThread());
+#endif
+
 	FModuleEventTickFunction* FoundTickFunction = nullptr;
 	if(TickFunctions.Num() > 0)
 	{
@@ -657,7 +662,7 @@ void FAnimNextModuleInstance::QueueTask(FName InEventName, TUniqueFunction<void(
 		}
 	}
 
-	TMpscQueue<TUniqueFunction<void(const FModuleTaskContext&)>>* Queue = nullptr;
+	FModuleEventTickFunction::FTaskQueueType* Queue = nullptr;
 	if (FoundTickFunction)
 	{
 		switch (InLocation)
diff --git a/Engine/Plugins/Experimental/UAF/UAF/Source/UAF/Private/Module/SystemReference.cpp b/Engine/Plugins/Experimental/UAF/UAF/Source/UAF/Private/Module/SystemReference.cpp
--- a/Engine/Plugins/Experimental/UAF/UAF/Source/UAF/Private/Module/SystemReference.cpp
+++ b/Engine/Plugins/Experimental/UAF/UAF/Source/UAF/Private/Module/SystemReference.cpp
@@ -16,6 +16,7 @@
 #include "Module/SystemDependency.h"
 #include "Module/UAFWeakSystemReference.h"
 #include "Containers/MpscQueue.h"
+#include "Containers/TransactionallySafeSpscQueue.h"
 
 namespace UE::UAF::Private
 {
@@ -26,6 +27,10 @@ namespace UE::UAF::Private
 
 		// Lock for concurrent allocations/resolves
 		FTransactionallySafeRWLock SystemInstancesLock;
+		
+		// Currently dont support a transactionally safe TMpscQueue, so we need to just use Spsc for now for actions that might run in transactions
+		// like deletes
+		TTransactionallySafeSpscQueue<TUniqueFunction<void()>> PendingDeletes;
 
 		// Actions to perform once world tick is completed each frame
 		TMpscQueue<TUniqueFunction<void()>> PendingActions;
@@ -36,6 +41,11 @@ namespace UE::UAF::Private
 			using namespace UE::UAF;
 
 			TUniqueFunction<void()> PendingAction;
+			while (PendingDeletes.Dequeue(PendingAction))
+			{
+				PendingAction();
+			}
+
 			while (PendingActions.Dequeue(PendingAction))
 			{
 				PendingAction();
@@ -98,7 +108,7 @@ struct FCustomSystemDeleter
 			InInstance->bHasBeenQueuedForRelease = true;
 #endif // WITH_EDITOR
 
-			Private::GetPool().PendingActions.Enqueue([Handle = Handle]()
+			Private::GetPool().PendingDeletes.Enqueue([Handle = Handle]()
 			{
 				UE::TWriteScopeLock InstancesLockScope(Private::GetPool().SystemInstancesLock);
 				Private::GetPool().SystemInstances.Release(Handle);
@@ -164,6 +174,7 @@ void FSystemReference::Destroy()
 {
 	checkf(Private::GPool != nullptr, TEXT("UAF: System pool has already been destroyed"));
 	checkf(Private::GPool->PendingActions.IsEmpty(), TEXT("UAF: System pool has pending actions that were added after Engine PreExit"));
+	checkf(Private::GPool->PendingDeletes.IsEmpty(), TEXT("UAF: System pool has pending deletes that were added after Engine PreExit"));
 	checkf(Private::GPool->SystemInstances.Num() == 0, TEXT("UAF: System pool is being destroyed with %d allocated entries"), Private::GPool->SystemInstances.Num());
 
 	delete Private::GPool;
diff --git a/Engine/Plugins/Experimental/UAF/UAFAnimGraph/Source/UAFAnimGraph/Private/Graph/RigUnit_AnimNextRunAnimationGraph_v1.cpp b/Engine/Plugins/Experimental/UAF/UAFAnimGraph/Source/UAFAnimGraph/Private/Graph/RigUnit_AnimNextRunAnimationGraph_v1.cpp
--- a/Engine/Plugins/Experimental/UAF/UAFAnimGraph/Source/UAFAnimGraph/Private/Graph/RigUnit_AnimNextRunAnimationGraph_v1.cpp
+++ b/Engine/Plugins/Experimental/UAF/UAFAnimGraph/Source/UAFAnimGraph/Private/Graph/RigUnit_AnimNextRunAnimationGraph_v1.cpp
@@ -81,7 +81,7 @@ FRigUnit_AnimNextRunAnimationGraph_v1_Execute()
 
 	// A schedule can contain multiple graphs, we copy the input event list since it might be appended to during our update
 	{
-		FReadScopeLock ReadLock(ModuleInstance.EventListLock);
+		UE::TReadScopeLock ReadLock(ModuleInstance.EventListLock);
 		InputEventList = ModuleInstance.InputEventList;
 	}
 
@@ -119,7 +119,7 @@ FRigUnit_AnimNextRunAnimationGraph_v1_Execute()
 	{
 		const int32 NumInputEvents = InputEventList.Num();
 
-		FWriteScopeLock WriteLock(ModuleInstance.EventListLock);
+		UE::TWriteScopeLock WriteLock(ModuleInstance.EventListLock);
 
 		// Append the new input events
 		for (int32 EventIndex = NumOriginalInputEvents; EventIndex < NumInputEvents; ++EventIndex)
diff --git a/Engine/Plugins/Experimental/UAF/UAFAnimGraph/Source/UAFAnimGraph/Private/Graph/RigUnit_AnimNextRunAnimationGraph_v2.cpp b/Engine/Plugins/Experimental/UAF/UAFAnimGraph/Source/UAFAnimGraph/Private/Graph/RigUnit_AnimNextRunAnimationGraph_v2.cpp
--- a/Engine/Plugins/Experimental/UAF/UAFAnimGraph/Source/UAFAnimGraph/Private/Graph/RigUnit_AnimNextRunAnimationGraph_v2.cpp
+++ b/Engine/Plugins/Experimental/UAF/UAFAnimGraph/Source/UAFAnimGraph/Private/Graph/RigUnit_AnimNextRunAnimationGraph_v2.cpp
@@ -112,7 +112,7 @@ FRigUnit_AnimNextRunAnimationGraph_v2_Execute()
 
 	// A module can contain multiple graphs, we copy the input event list since it might be appended to during our update
 	{
-		FReadScopeLock ReadLock(ModuleInstance.EventListLock);
+		UE::TReadScopeLock ReadLock(ModuleInstance.EventListLock);
 		InputEventList = ModuleInstance.InputEventList;
 	}
 
@@ -152,7 +152,7 @@ FRigUnit_AnimNextRunAnimationGraph_v2_Execute()
 	{
 		const int32 NumInputEvents = InputEventList.Num();
 
-		FWriteScopeLock WriteLock(ModuleInstance.EventListLock);
+		UE::TWriteScopeLock WriteLock(ModuleInstance.EventListLock);
 
 		// Append the new input events
 		for (int32 EventIndex = NumOriginalInputEvents; EventIndex < NumInputEvents; ++EventIndex)
diff --git a/Engine/Plugins/Runtime/RigVM/Source/RigVM/Public/RigVMCore/RigVMRegistry.h b/Engine/Plugins/Runtime/RigVM/Source/RigVM/Public/RigVMCore/RigVMRegistry.h
--- a/Engine/Plugins/Runtime/RigVM/Source/RigVM/Public/RigVMCore/RigVMRegistry.h
+++ b/Engine/Plugins/Runtime/RigVM/Source/RigVM/Public/RigVMCore/RigVMRegistry.h
@@ -555,7 +555,10 @@ protected:
 			else if (DesiredLockType == LockType_Write)
 			{
 				Registry->Lock.WriteLock();
-				Registry->WriteLockOwnerThread = FPlatformTLS::GetCurrentThreadId();
+				UE_AUTORTFM_OPEN
+				{
+					Registry->WriteLockOwnerThread.store(FPlatformTLS::GetCurrentThreadId());
+				};
 			}
 
 			// Increment the matching lock count within the RigVM registry.
@@ -632,7 +635,10 @@ protected:
 			else if (DesiredLockType == LockType_Write)
 			{
 				ensure(CurrentLockCount == 0);
-				Registry->WriteLockOwnerThread = static_cast<uint32>(-1);
+				UE_AUTORTFM_OPEN
+				{
+					Registry->WriteLockOwnerThread.store(static_cast<uint32>(-1));
+				};
 				Registry->Lock.WriteUnlock();
 			}
 		}
