Move and rename UAF ValueRuntime related files

#Synced-CL 49105642

diff --git a/Engine/Plugins/Experimental/MovieSceneAnimMixer/Source/MovieSceneAnimMixer/Private/Systems/MovieSceneMixedSkeletalAnimationSystem.cpp b/Engine/Plugins/Experimental/MovieSceneAnimMixer/Source/MovieSceneAnimMixer/Private/Systems/MovieSceneMixedSkeletalAnimationSystem.cpp
--- a/Engine/Plugins/Experimental/MovieSceneAnimMixer/Source/MovieSceneAnimMixer/Private/Systems/MovieSceneMixedSkeletalAnimationSystem.cpp
+++ b/Engine/Plugins/Experimental/MovieSceneAnimMixer/Source/MovieSceneAnimMixer/Private/Systems/MovieSceneMixedSkeletalAnimationSystem.cpp
@@ -38,7 +38,7 @@
 #include "Animation/AnimSequence.h"
 #include "Animation/AnimRootMotionProvider.h"
 #include "BonePose.h"
-#include "DecompressionTools.h"
+#include "EvaluationVM/Tasks/DecompressionTools.h"
 #include "EvaluationVM/EvaluationVM.h"
 #include "EvaluationVM/KeyframeState.h"
 #include "Animation/BuiltInAttributeTypes.h"
diff --git a/Engine/Plugins/Experimental/MovieScenePoseSearchTracks/Source/MovieScenePoseSearchTracks/Private/Sections/MovieSceneStitchAnimSection.cpp b/Engine/Plugins/Experimental/MovieScenePoseSearchTracks/Source/MovieScenePoseSearchTracks/Private/Sections/MovieSceneStitchAnimSection.cpp
--- a/Engine/Plugins/Experimental/MovieScenePoseSearchTracks/Source/MovieScenePoseSearchTracks/Private/Sections/MovieSceneStitchAnimSection.cpp
+++ b/Engine/Plugins/Experimental/MovieScenePoseSearchTracks/Source/MovieScenePoseSearchTracks/Private/Sections/MovieSceneStitchAnimSection.cpp
@@ -7,7 +7,7 @@
 #include "Animation/BuiltInAttributeTypes.h"
 #include "AnimMixerComponentTypes.h"
 #include "Channels/MovieSceneChannelProxy.h"
-#include "DecompressionTools.h"
+#include "EvaluationVM/Tasks/DecompressionTools.h"
 #include "MovieScene.h"
 #include "MovieSceneSequence.h"
 #include "EntitySystem/BuiltInComponentTypes.h"
diff --git a/Engine/Plugins/Experimental/UAF/UAF/Source/UAF/Private/AnimNextModuleImpl.cpp b/Engine/Plugins/Experimental/UAF/UAF/Source/UAF/Private/AnimNextModuleImpl.cpp
--- a/Engine/Plugins/Experimental/UAF/UAF/Source/UAF/Private/AnimNextModuleImpl.cpp
+++ b/Engine/Plugins/Experimental/UAF/UAF/Source/UAF/Private/AnimNextModuleImpl.cpp
@@ -32,7 +32,7 @@
 #include "Features/IModularFeatures.h"
 #include "RewindDebuggerRuntimeInterface/IRewindDebuggerRuntimeExtension.h"
 #include "RewindDebugger/RewindDebuggerUAFRuntime.h"
-#include "UAF/Attributes/AttributeRuntimeRegistry.h"
+#include "UAF/ValueRuntime/ValueRuntimeRegistry.h"
 #include "GameFramework/CharacterMovementComponent.h"
 #include "UAF/Attributes/EngineAttributes.h"
 
@@ -94,7 +94,7 @@ namespace UE::UAF
 
 		FSystemFactory::Init();
 
-		FAttributeRuntimeRegistry::Init();
+		FValueRuntimeRegistry::Init();
 		RegisterEngineAttributes();
 
 #if UAF_TRACE_ENABLED
@@ -140,7 +140,7 @@ namespace UE::UAF
 		IModularFeatures::Get().UnregisterModularFeature(IRewindDebuggerRuntimeExtension::ModularFeatureName, &GRewindDebuggerUAFRuntime);
 #endif
 		
-		FAttributeRuntimeRegistry::Destroy();
+		FValueRuntimeRegistry::Destroy();
 		FRigVMRuntimeDataRegistry::Destroy();
 		FDataRegistry::Destroy();
 		FSystemFactory::Destroy();
diff --git a/Engine/Plugins/Experimental/UAF/UAF/Source/UAF/Private/UAF/Attributes/EngineAttributes.cpp b/Engine/Plugins/Experimental/UAF/UAF/Source/UAF/Private/UAF/Attributes/EngineAttributes.cpp
--- a/Engine/Plugins/Experimental/UAF/UAF/Source/UAF/Private/UAF/Attributes/EngineAttributes.cpp
+++ b/Engine/Plugins/Experimental/UAF/UAF/Source/UAF/Private/UAF/Attributes/EngineAttributes.cpp
@@ -2,19 +2,19 @@
 
 #include "UAF/Attributes/EngineAttributes.h"
 
-#include "UAF/Attributes/AttributeRuntimeRegistry.h"
-#include "UAF/Attributes/Transformers/Accumulate.h"
-#include "UAF/Attributes/Transformers/AdditiveSpace.h"
-#include "UAF/Attributes/Transformers/Interpolate.h"
-#include "UAF/Attributes/Transformers/Layer.h"
-#include "UAF/Attributes/Transformers/Overwrite.h"
-#include "UAF/Attributes/Transformers/Sanitize.h"
+#include "UAF/ValueRuntime/ValueRuntimeRegistry.h"
+#include "UAF/ValueRuntime/Transformers/Accumulate.h"
+#include "UAF/ValueRuntime/Transformers/AdditiveSpace.h"
+#include "UAF/ValueRuntime/Transformers/Interpolate.h"
+#include "UAF/ValueRuntime/Transformers/Layer.h"
+#include "UAF/ValueRuntime/Transformers/Overwrite.h"
+#include "UAF/ValueRuntime/Transformers/Sanitize.h"
 
 namespace UE::UAF
 {
 	void RegisterEngineAttributes()
 	{
-		FAttributeRuntimeRegistry& RuntimeRegistry = FAttributeRuntimeRegistry::Get();
+		FValueRuntimeRegistry& RuntimeRegistry = FValueRuntimeRegistry::Get();
 
 		RuntimeRegistry.RegisterSetMapInitializer<FFloatAnimationAttribute>();
 		RuntimeRegistry.RegisterSetMapInitializer<FIntegerAnimationAttribute>();
diff --git a/Engine/Plugins/Experimental/UAF/UAF/Source/UAF/Private/UAF/ValueRuntime/AttributeCollection.cpp b/Engine/Plugins/Experimental/UAF/UAF/Source/UAF/Private/UAF/ValueRuntime/AttributeCollection.cpp
new file mode 100644
--- /dev/null
+++ b/Engine/Plugins/Experimental/UAF/UAF/Source/UAF/Private/UAF/ValueRuntime/AttributeCollection.cpp
@@ -0,0 +1,741 @@
+// Copyright Epic Games, Inc. All Rights Reserved.
+
+#include "UAF/ValueRuntime/AttributeCollection.h"
+
+#include "Algo/BinarySearch.h"
+#include "UAF/ValueRuntime/ValueRuntimeRegistry.h"
+
+namespace UE::UAF
+{
+	FAttributeCollection::FAttributeCollection(const FAttributeCollection& Other)
+		: NamedSet(Other.NamedSet)
+		, ValueSpace(Other.ValueSpace)
+		, SetMapCount(Other.SetMapCount)
+		, SetMapCapacity(Other.SetMapCount)
+		, MapCount(Other.MapCount)
+		, MapCapacity(Other.MapCount)
+		, ReallocFun(Other.ReallocFun)
+	{
+		if (Other.SetMaps != nullptr)
+		{
+			const int32 BufferSize = sizeof(FSetMapEntry) * Other.SetMapCount;
+			SetMaps = reinterpret_cast<FSetMapEntry*>((*Other.ReallocFun)(nullptr, 0, BufferSize));
+			FMemory::Memcpy(SetMaps, Other.SetMaps, BufferSize);
+
+			for (int32 MapIndex = 0; MapIndex < Other.SetMapCount; ++MapIndex)
+			{
+				SetMaps[MapIndex].Map = Other.SetMaps[MapIndex].Map->Duplicate(Other.ReallocFun);
+			}
+		}
+		else
+		{
+			SetMaps = nullptr;
+		}
+
+		if (Other.Maps != nullptr)
+		{
+			const int32 BufferSize = sizeof(FMapEntry) * Other.MapCount;
+			Maps = reinterpret_cast<FMapEntry*>((*Other.ReallocFun)(nullptr, 0, BufferSize));
+			FMemory::Memcpy(Maps, Other.Maps, BufferSize);
+
+			for (int32 MapIndex = 0; MapIndex < Other.MapCount; ++MapIndex)
+			{
+				Maps[MapIndex].Map = Other.Maps[MapIndex].Map->Duplicate(Other.ReallocFun);
+			}
+		}
+		else
+		{
+			Maps = nullptr;
+		}
+	}
+
+	FAttributeCollection::FAttributeCollection(FAttributeCollection&& Other)
+		: NamedSet(MoveTemp(Other.NamedSet))
+		, ValueSpace(Other.ValueSpace)
+		, SetMaps(Other.SetMaps)
+		, SetMapCount(Other.SetMapCount)
+		, SetMapCapacity(Other.SetMapCapacity)
+		, Maps(Other.Maps)
+		, MapCount(Other.MapCount)
+		, MapCapacity(Other.MapCapacity)
+		, ReallocFun(Other.ReallocFun)
+	{
+		Other.SetMaps = nullptr;
+		Other.SetMapCount = 0;
+		Other.SetMapCapacity = 0;
+		Other.Maps = nullptr;
+		Other.MapCount = 0;
+		Other.MapCapacity = 0;
+	}
+
+	FAttributeCollection::FAttributeCollection(const FAttributeNamedSetPtr& InNamedSet, FReallocFun InReallocFun)
+		: NamedSet(InNamedSet)
+		, ReallocFun(InReallocFun)
+	{
+		check(InReallocFun != nullptr);
+	}
+
+	FAttributeCollection::~FAttributeCollection()
+	{
+		Empty();
+	}
+
+	FAttributeCollection& FAttributeCollection::operator=(const FAttributeCollection& Other)
+	{
+		if (this != &Other)
+		{
+			checkf(ReallocFun != nullptr, TEXT("When we copy from another collection, we retain our original allocator"));
+
+			// Free anything we might be holding
+			Empty();
+
+			NamedSet = Other.NamedSet;
+			ValueSpace = Other.ValueSpace;
+			SetMapCount = Other.SetMapCount;
+			SetMapCapacity = Other.SetMapCount;
+
+			if (Other.SetMaps != nullptr)
+			{
+				const int32 BufferSize = sizeof(FSetMapEntry) * Other.SetMapCount;
+				SetMaps = reinterpret_cast<FSetMapEntry*>((*ReallocFun)(nullptr, 0, BufferSize));
+				FMemory::Memcpy(SetMaps, Other.SetMaps, BufferSize);
+
+				for (int32 MapIndex = 0; MapIndex < Other.SetMapCount; ++MapIndex)
+				{
+					SetMaps[MapIndex].Map = Other.SetMaps[MapIndex].Map->Duplicate(ReallocFun);
+				}
+			}
+			else
+			{
+				SetMaps = nullptr;
+			}
+
+			MapCount = Other.MapCount;
+			MapCapacity = Other.MapCount;
+
+			if (Other.Maps != nullptr)
+			{
+				const int32 BufferSize = sizeof(FMapEntry) * Other.MapCount;
+				Maps = reinterpret_cast<FMapEntry*>((*ReallocFun)(nullptr, 0, BufferSize));
+				FMemory::Memcpy(Maps, Other.Maps, BufferSize);
+
+				for (int32 MapIndex = 0; MapIndex < Other.MapCount; ++MapIndex)
+				{
+					Maps[MapIndex].Map = Other.Maps[MapIndex].Map->Duplicate(ReallocFun);
+				}
+			}
+			else
+			{
+				Maps = nullptr;
+			}
+		}
+
+		return *this;
+	}
+
+	FAttributeCollection& FAttributeCollection::operator=(FAttributeCollection&& Other)
+	{
+		if (this != &Other)
+		{
+			checkf(ReallocFun == Other.ReallocFun, TEXT("When we move from another collection, our allocators must match"));
+
+			// Free anything we might be holding
+			Empty();
+
+			NamedSet = MoveTemp(Other.NamedSet);
+			ValueSpace = Other.ValueSpace;
+			SetMaps = Other.SetMaps;
+			SetMapCount = Other.SetMapCount;
+			SetMapCapacity = Other.SetMapCapacity;
+			Maps = Other.Maps;
+			MapCount = Other.MapCount;
+			MapCapacity = Other.MapCapacity;
+
+			Other.SetMaps = nullptr;
+			Other.SetMapCount = 0;
+			Other.SetMapCapacity = 0;
+			Other.Maps = nullptr;
+			Other.MapCount = 0;
+			Other.MapCapacity = 0;
+		}
+
+		return *this;
+	}
+
+	bool FAttributeCollection::IsEmpty() const
+	{
+		return SetMapCount == 0 && MapCount == 0;
+	}
+
+	const FAttributeNamedSetPtr& FAttributeCollection::GetNamedSet() const
+	{
+		return NamedSet;
+	}
+
+	FValueSpace FAttributeCollection::GetValueSpace() const
+	{
+		return ValueSpace;
+	}
+
+	void FAttributeCollection::SetValueSpace(FValueSpace InValueSpace)
+	{
+		ValueSpace = InValueSpace;
+	}
+
+	void FAttributeCollection::InitWithValueSpace(FValueSpace InValueSpace)
+	{
+		if (!NamedSet)
+		{
+			// Nothing to initialize if we have no named set
+			return;
+		}
+
+		// Empty the collection
+		Reset(NamedSet);
+
+		const bool bIsAdditive = InValueSpace.IsAdditive();
+
+		bool bIsInitialized = false;
+		if (!bIsAdditive)
+		{
+			const EValueSpaceType SpaceType = InValueSpace.GetType();
+			switch (SpaceType)
+			{
+			case EValueSpaceType::Local:
+				// From local space reference pose
+				*this = *NamedSet->GetDefaultAttributeValues();
+				bIsInitialized = true;
+				break;
+			case EValueSpaceType::Component:
+				// From component space reference pose
+				checkf(false, TEXT("Not yet supported"));
+				break;
+			}
+		}
+
+		if (!bIsInitialized)
+		{
+			for (auto It = NamedSet->CreateTypedSetIterator(); It; ++It)
+			{
+				const FAttributeTypedSetPtr TypedSet = *It;
+
+				// By default, attributes map to values of their type
+				const FAttributeMappingKey MappingKey = FAttributeMappingKey::MakeFromTo(TypedSet->GetType());
+
+				FAttributeSetMap* Copy = AddSetMap(MappingKey);
+				Copy->FillWithIdentity(bIsAdditive);
+			}
+		}
+
+		ValueSpace = InValueSpace;
+	}
+
+	bool FAttributeCollection::IsAdditive() const
+	{
+		return ValueSpace.IsAdditive();
+	}
+
+	void FAttributeCollection::Reset(const FAttributeNamedSetPtr& InNamedSet)
+	{
+		// Free our maps but retain the arrays that contains them
+		for (int32 MapIndex = 0; MapIndex < SetMapCount; ++MapIndex)
+		{
+			ReleaseAttributeSetMap(SetMaps[MapIndex].Map);
+		}
+		SetMapCount = 0;
+
+		for (int32 MapIndex = 0; MapIndex < MapCount; ++MapIndex)
+		{
+			ReleaseAttributeMap(Maps[MapIndex].Map);
+		}
+		MapCount = 0;
+
+		NamedSet = InNamedSet;
+		ValueSpace = FValueSpace();
+	}
+
+	void FAttributeCollection::Empty()
+	{
+		if (SetMaps != nullptr)
+		{
+			for (int32 MapIndex = 0; MapIndex < SetMapCount; ++MapIndex)
+			{
+				ReleaseAttributeSetMap(SetMaps[MapIndex].Map);
+			}
+
+			(*ReallocFun)(reinterpret_cast<uint8*>(SetMaps), sizeof(FSetMapEntry) * SetMapCapacity, 0);
+			SetMaps = nullptr;
+			SetMapCount = 0;
+			SetMapCapacity = 0;
+		}
+
+		if (Maps != nullptr)
+		{
+			for (int32 MapIndex = 0; MapIndex < MapCount; ++MapIndex)
+			{
+				ReleaseAttributeMap(Maps[MapIndex].Map);
+			}
+
+			(*ReallocFun)(reinterpret_cast<uint8*>(Maps), sizeof(FMapEntry) * MapCapacity, 0);
+			Maps = nullptr;
+			MapCount = 0;
+			MapCapacity = 0;
+		}
+
+		NamedSet = nullptr;
+		ValueSpace = FValueSpace();
+	}
+
+	FReallocFun FAttributeCollection::GetAllocator() const
+	{
+		return ReallocFun;
+	}
+
+	int32 FAttributeCollection::NumSetMaps() const
+	{
+		return SetMapCount;
+	}
+
+	FAttributeSetMap* FAttributeCollection::AddSetMap(const FAttributeMappingKey& MappingKey)
+	{
+		if (!NamedSet)
+		{
+			// We are not initialized
+			return nullptr;
+		}
+
+		FAttributeTypedSetPtr TypedSet = NamedSet->FindTypedSet(MappingKey.GetAttributeType());
+		if (!TypedSet)
+		{
+			// Unknown attribute type, it isn't contained within our named set
+			return nullptr;
+		}
+
+		// Use default construction arguments
+		const FAttributeSetMap::FConstructArgs Args(TypedSet, MappingKey.GetValueType(), ReallocFun);
+
+		// We don't know the concrete type of the set map, use indirection to construct it
+		constexpr bool bConstructSetMap = true;
+		FAttributeSetMap** SetMap = AddSetMapImpl(Args, bConstructSetMap);
+		return SetMap != nullptr ? *SetMap : nullptr;
+	}
+
+	FAttributeSetMap** FAttributeCollection::AddSetMapImpl(const FAttributeSetMap::FConstructArgs& Args, bool bConstructSetMap)
+	{
+		UScriptStruct* AttributeType = Args.TypedSet->GetType();
+		UScriptStruct* ValueType = Args.ValueType;
+
+		const FAttributeMappingKey Key(AttributeType, ValueType);
+
+		const int32 MapIndex = Algo::LowerBoundBy(TArrayView<const FSetMapEntry>(SetMaps, SetMapCount), Key, [](const FSetMapEntry& Entry) { return Entry.Key; });
+		if (MapIndex < SetMapCount && SetMaps[MapIndex].Key == Key)
+		{
+			// A map already exists that matches the specified types
+			return nullptr;
+		}
+
+		FAttributeSetMap* SetMap = nullptr;
+		if (bConstructSetMap)
+		{
+			SetMap = FValueRuntimeRegistry::Get().ConstructSetMap(Args);
+
+			if (SetMap == nullptr)
+			{
+				// Failed to construct a set map, the type has probably not been registered
+				return nullptr;
+			}
+		}
+
+		if (SetMapCount >= SetMapCapacity)
+		{
+			// The first time we add a set, we reserve space for an entry per typed set plus some extra room
+			const bool bAllowQuantize = true;
+			const int32 DesiredCapacity = SetMapCapacity == 0 ? (NamedSet->NumTypedSets() + 4) : (SetMapCapacity + 1);
+			const int32 NewCapacity = DefaultCalculateSlackGrow(DesiredCapacity, SetMapCapacity, sizeof(FSetMapEntry), bAllowQuantize);
+
+			SetMaps = reinterpret_cast<FSetMapEntry*>((*ReallocFun)(
+				reinterpret_cast<uint8*>(SetMaps),
+				sizeof(FSetMapEntry) * SetMapCapacity,
+				sizeof(FSetMapEntry) * NewCapacity));
+
+			SetMapCapacity = NewCapacity;
+		}
+
+		const int32 NumToMove = SetMapCount - MapIndex;
+		FMemory::Memmove(SetMaps + MapIndex + 1, SetMaps + MapIndex, sizeof(FSetMapEntry) * NumToMove);
+		SetMaps[MapIndex].Key = Key;
+		SetMaps[MapIndex].Map = SetMap;
+		SetMapCount++;
+
+		return &SetMaps[MapIndex].Map;
+	}
+
+	bool FAttributeCollection::AddSetMap(FAttributeSetMap* Map)
+	{
+		if (Map == nullptr)
+		{
+			// Invalid set map
+			return false;
+		}
+		else if (Map->GetAllocator() != ReallocFun)
+		{
+			// Allocator mismatch
+			return false;
+		}
+
+		const FAttributeMappingKey Key = Map->GetMappingKey();
+
+		const int32 MapIndex = Algo::LowerBoundBy(TArrayView<const FSetMapEntry>(SetMaps, SetMapCount), Key, [](const FSetMapEntry& Entry) { return Entry.Key; });
+		if (MapIndex < SetMapCount && SetMaps[MapIndex].Key == Key)
+		{
+			// A map already exists that matches the specified types
+			return false;
+		}
+
+		if (SetMapCount >= SetMapCapacity)
+		{
+			// The first time we add a set, we reserve space for an entry per typed set plus some extra room
+			const bool bAllowQuantize = true;
+			const int32 DesiredCapacity = SetMapCapacity == 0 ? (NamedSet->NumTypedSets() + 4) : (SetMapCapacity + 1);
+			const int32 NewCapacity = DefaultCalculateSlackGrow(DesiredCapacity, SetMapCapacity, sizeof(FSetMapEntry), bAllowQuantize);
+
+			SetMaps = reinterpret_cast<FSetMapEntry*>((*ReallocFun)(
+				reinterpret_cast<uint8*>(SetMaps),
+				sizeof(FSetMapEntry) * SetMapCapacity,
+				sizeof(FSetMapEntry) * NewCapacity));
+
+			SetMapCapacity = NewCapacity;
+		}
+
+		const int32 NumToMove = SetMapCount - MapIndex;
+		FMemory::Memmove(SetMaps + MapIndex + 1, SetMaps + MapIndex, sizeof(FSetMapEntry) * NumToMove);
+		SetMaps[MapIndex].Key = Key;
+		SetMaps[MapIndex].Map = Map;
+		SetMapCount++;
+
+		return true;
+	}
+
+	bool FAttributeCollection::AppendSetMap(FAttributeSetMap* Map)
+	{
+		if (Map == nullptr)
+		{
+			// Invalid set map
+			return false;
+		}
+		else if (Map->GetAllocator() != ReallocFun)
+		{
+			// Allocator mismatch
+			return false;
+		}
+
+		const FAttributeMappingKey Key = Map->GetMappingKey();
+
+		if (SetMapCount != 0 && !(SetMaps[SetMapCount - 1].Key < Key))
+		{
+			// We don't belong at the end of the set maps
+			return false;
+		}
+
+		if (SetMapCount >= SetMapCapacity)
+		{
+			// The first time we add a set, we reserve space for an entry per typed set plus some extra room
+			const bool bAllowQuantize = true;
+			const int32 DesiredCapacity = SetMapCapacity == 0 ? (NamedSet->NumTypedSets() + 4) : (SetMapCapacity + 1);
+			const int32 NewCapacity = DefaultCalculateSlackGrow(DesiredCapacity, SetMapCapacity, sizeof(FSetMapEntry), bAllowQuantize);
+
+			SetMaps = reinterpret_cast<FSetMapEntry*>((*ReallocFun)(
+				reinterpret_cast<uint8*>(SetMaps),
+				sizeof(FSetMapEntry) * SetMapCapacity,
+				sizeof(FSetMapEntry) * NewCapacity));
+
+			SetMapCapacity = NewCapacity;
+		}
+
+		SetMaps[SetMapCount].Key = Key;
+		SetMaps[SetMapCount].Map = Map;
+		SetMapCount++;
+
+		return true;
+	}
+
+	bool FAttributeCollection::RemoveSetMap(FAttributeSetMap* Map)
+	{
+		if (Map == nullptr)
+		{
+			// Invalid set map
+			return false;
+		}
+
+		return RemoveSetMap(Map->GetMappingKey());
+	}
+
+	bool FAttributeCollection::RemoveSetMap(const FAttributeMappingKey& MappingKey)
+	{
+		const int32 MapIndex = Algo::LowerBoundBy(TArrayView<const FSetMapEntry>(SetMaps, SetMapCount), MappingKey, [](const FSetMapEntry& Entry) { return Entry.Key; });
+		if (MapIndex >= SetMapCount || SetMaps[MapIndex].Key != MappingKey)
+		{
+			// No map found that matches the desired types
+			return false;
+		}
+
+		// Destroy and free our set map
+		ReleaseAttributeSetMap(SetMaps[MapIndex].Map);
+
+		// Remove our entry
+		const int32 NumToMove = SetMapCount - MapIndex - 1;
+		FMemory::Memmove(SetMaps + MapIndex, SetMaps + MapIndex + 1, sizeof(FSetMapEntry) * NumToMove);
+		SetMapCount--;
+
+		return true;
+	}
+
+	FAttributeSetMap* FAttributeCollection::FindSetMap(const FAttributeMappingKey& MappingKey)
+	{
+		const int32 MapIndex = Algo::LowerBoundBy(TArrayView<const FSetMapEntry>(SetMaps, SetMapCount), MappingKey, [](const FSetMapEntry& Entry) { return Entry.Key; });
+		if (MapIndex >= SetMapCount || SetMaps[MapIndex].Key != MappingKey)
+		{
+			// No map found that matches the desired types
+			return nullptr;
+		}
+
+		return SetMaps[MapIndex].Map;
+	}
+
+	const FAttributeSetMap* FAttributeCollection::FindSetMap(const FAttributeMappingKey& MappingKey) const
+	{
+		const int32 MapIndex = Algo::LowerBoundBy(TArrayView<const FSetMapEntry>(SetMaps, SetMapCount), MappingKey, [](const FSetMapEntry& Entry) { return Entry.Key; });
+		if (MapIndex >= SetMapCount || SetMaps[MapIndex].Key != MappingKey)
+		{
+			// No map found that matches the desired types
+			return nullptr;
+		}
+
+		return SetMaps[MapIndex].Map;
+	}
+
+	FAttributeSetMap* FAttributeCollection::FindOrAddSetMap(const FAttributeMappingKey& MappingKey)
+	{
+		if (FAttributeSetMap* SetMap = FindSetMap(MappingKey))
+		{
+			return SetMap;
+		}
+
+		return AddSetMap(MappingKey);
+	}
+
+	int32 FAttributeCollection::NumMaps() const
+	{
+		return MapCount;
+	}
+
+	FAttributeMap* FAttributeCollection::AddMap(UScriptStruct* ValueType)
+	{
+		if (!NamedSet)
+		{
+			// We are not initialized
+			return nullptr;
+		}
+
+		// We don't know the concrete type of the attribute map, use indirection to construct it
+		constexpr bool bConstructMap = true;
+		FAttributeMap** Map = AddMapImpl(ValueType, bConstructMap);
+		return Map != nullptr ? *Map : nullptr;
+	}
+
+	FAttributeMap** FAttributeCollection::AddMapImpl(UScriptStruct* ValueType, bool bConstructMap)
+	{
+		const int32 MapIndex = Algo::LowerBoundBy(TArrayView<const FMapEntry>(Maps, MapCount), ValueType, [](const FMapEntry& Entry) { return Entry.ValueType; });
+		if (MapIndex < MapCount && Maps[MapIndex].ValueType == ValueType)
+		{
+			// A map already exists that matches the specified type
+			return nullptr;
+		}
+
+		FAttributeMap* Map = nullptr;
+		if (bConstructMap)
+		{
+			Map = FValueRuntimeRegistry::Get().ConstructMap(ValueType, ReallocFun);
+
+			if (Map == nullptr)
+			{
+				// Failed to construct a map, the type has probably not been registered
+				return nullptr;
+			}
+		}
+
+		if (MapCount >= MapCapacity)
+		{
+			const bool bAllowQuantize = true;
+			const int32 DesiredCapacity = MapCapacity + 1;
+			const int32 NewCapacity = DefaultCalculateSlackGrow(DesiredCapacity, MapCapacity, sizeof(FMapEntry), bAllowQuantize);
+
+			Maps = reinterpret_cast<FMapEntry*>((*ReallocFun)(
+				reinterpret_cast<uint8*>(Maps),
+				sizeof(FMapEntry) * MapCapacity,
+				sizeof(FMapEntry) * NewCapacity));
+
+			MapCapacity = NewCapacity;
+		}
+
+		const int32 NumToMove = MapCount - MapIndex;
+		FMemory::Memmove(Maps + MapIndex + 1, Maps + MapIndex, sizeof(FMapEntry) * NumToMove);
+		Maps[MapIndex].ValueType = ValueType;
+		Maps[MapIndex].Map = Map;
+		MapCount++;
+
+		return &Maps[MapIndex].Map;
+	}
+
+	bool FAttributeCollection::AddMap(FAttributeMap* Map)
+	{
+		if (Map == nullptr)
+		{
+			// Invalid set map
+			return false;
+		}
+		else if (Map->GetAllocator() != ReallocFun)
+		{
+			// Allocator mismatch
+			return false;
+		}
+
+		UScriptStruct* ValueType = Map->GetValueType();
+
+		const int32 MapIndex = Algo::LowerBoundBy(TArrayView<const FMapEntry>(Maps, MapCount), ValueType, [](const FMapEntry& Entry) { return Entry.ValueType; });
+		if (MapIndex < MapCount && Maps[MapIndex].ValueType == ValueType)
+		{
+			// A map already exists that matches the specified types
+			return false;
+		}
+
+		if (MapCount >= MapCapacity)
+		{
+			const bool bAllowQuantize = true;
+			const int32 DesiredCapacity = MapCapacity + 1;
+			const int32 NewCapacity = DefaultCalculateSlackGrow(DesiredCapacity, MapCapacity, sizeof(FMapEntry), bAllowQuantize);
+
+			Maps = reinterpret_cast<FMapEntry*>((*ReallocFun)(
+				reinterpret_cast<uint8*>(Maps),
+				sizeof(FMapEntry) * MapCapacity,
+				sizeof(FMapEntry) * NewCapacity));
+
+			MapCapacity = NewCapacity;
+		}
+
+		const int32 NumToMove = MapCount - MapIndex;
+		FMemory::Memmove(Maps + MapIndex + 1, Maps + MapIndex, sizeof(FMapEntry) * NumToMove);
+		Maps[MapIndex].ValueType = ValueType;
+		Maps[MapIndex].Map = Map;
+		MapCount++;
+
+		return true;
+	}
+
+	bool FAttributeCollection::AppendMap(FAttributeMap* Map)
+	{
+		if (Map == nullptr)
+		{
+			// Invalid set map
+			return false;
+		}
+		else if (Map->GetAllocator() != ReallocFun)
+		{
+			// Allocator mismatch
+			return false;
+		}
+
+		UScriptStruct* ValueType = Map->GetValueType();
+
+		if (MapCount != 0 && !(Maps[MapCount - 1].ValueType < ValueType))
+		{
+			// We don't belong at the end of the name/value maps
+			return false;
+		}
+
+		if (MapCount >= MapCapacity)
+		{
+			const bool bAllowQuantize = true;
+			const int32 DesiredCapacity = MapCapacity + 1;
+			const int32 NewCapacity = DefaultCalculateSlackGrow(DesiredCapacity, MapCapacity, sizeof(FMapEntry), bAllowQuantize);
+
+			Maps = reinterpret_cast<FMapEntry*>((*ReallocFun)(
+				reinterpret_cast<uint8*>(Maps),
+				sizeof(FMapEntry) * MapCapacity,
+				sizeof(FMapEntry) * NewCapacity));
+
+			MapCapacity = NewCapacity;
+		}
+
+		Maps[MapCount].ValueType = ValueType;
+		Maps[MapCount].Map = Map;
+		MapCount++;
+
+		return true;
+	}
+
+	bool FAttributeCollection::RemoveMap(FAttributeMap* Map)
+	{
+		if (Map == nullptr)
+		{
+			// Invalid map
+			return false;
+		}
+
+		return RemoveMap(Map->GetValueType());
+	}
+
+	bool FAttributeCollection::RemoveMap(UScriptStruct* ValueType)
+	{
+		const int32 MapIndex = Algo::LowerBoundBy(TArrayView<const FMapEntry>(Maps, MapCount), ValueType, [](const FMapEntry& Entry) { return Entry.ValueType; });
+		if (MapIndex >= MapCount || Maps[MapIndex].ValueType != ValueType)
+		{
+			// No map found that matches the desired types
+			return false;
+		}
+
+		// Destroy and free our map
+		ReleaseAttributeMap(Maps[MapIndex].Map);
+
+		// Remove our entry
+		const int32 NumToMove = MapCount - MapIndex - 1;
+		FMemory::Memmove(Maps + MapIndex, Maps + MapIndex + 1, sizeof(FMapEntry) * NumToMove);
+		MapCount--;
+
+		return true;
+	}
+
+	FAttributeMap* FAttributeCollection::FindMap(UScriptStruct* ValueType)
+	{
+		const int32 MapIndex = Algo::LowerBoundBy(TArrayView<const FMapEntry>(Maps, MapCount), ValueType, [](const FMapEntry& Entry) { return Entry.ValueType; });
+		if (MapIndex >= MapCount || Maps[MapIndex].ValueType != ValueType)
+		{
+			// No map found that matches the desired types
+			return nullptr;
+		}
+
+		return Maps[MapIndex].Map;
+	}
+
+	const FAttributeMap* FAttributeCollection::FindMap(UScriptStruct* ValueType) const
+	{
+		const int32 MapIndex = Algo::LowerBoundBy(TArrayView<const FMapEntry>(Maps, MapCount), ValueType, [](const FMapEntry& Entry) { return Entry.ValueType; });
+		if (MapIndex >= MapCount || Maps[MapIndex].ValueType != ValueType)
+		{
+			// No map found that matches the desired types
+			return nullptr;
+		}
+
+		return Maps[MapIndex].Map;
+	}
+
+	FAttributeMap* FAttributeCollection::FindOrAddMap(UScriptStruct* ValueType)
+	{
+		if (FAttributeMap* Map = FindMap(ValueType))
+		{
+			return Map;
+		}
+
+		return AddMap(ValueType);
+	}
+}
diff --git a/Engine/Plugins/Experimental/UAF/UAF/Source/UAF/Private/UAF/ValueRuntime/AttributeMap.cpp b/Engine/Plugins/Experimental/UAF/UAF/Source/UAF/Private/UAF/ValueRuntime/AttributeMap.cpp
new file mode 100644
--- /dev/null
+++ b/Engine/Plugins/Experimental/UAF/UAF/Source/UAF/Private/UAF/ValueRuntime/AttributeMap.cpp
@@ -0,0 +1,72 @@
+// Copyright Epic Games, Inc. All Rights Reserved.
+
+#include "UAF/ValueRuntime/AttributeMap.h"
+
+#include "UObject/Class.h"
+
+namespace UE::UAF
+{
+	FAttributeMap::FAttributeMap() = default;
+
+	FAttributeMap::FAttributeMap(UScriptStruct* InValueType, FReallocFun InReallocFun, int32 InContainerSize)
+		: ValueTypeStruct(InValueType)
+		, NumValues(0)
+		, MaxValues(0)
+		, DataPtr(nullptr)
+		, ReallocFun(InReallocFun)
+		, ContainerSize(InContainerSize)
+	{
+		check(InReallocFun != nullptr);
+	}
+
+	FAttributeMap::~FAttributeMap()
+	{
+		checkf(DataPtr == nullptr, TEXT("Derived types should handle destruction"));
+	}
+
+	bool FAttributeMap::IsEmpty() const
+	{
+		return NumValues == 0;
+	}
+
+	int32 FAttributeMap::Num() const
+	{
+		return NumValues;
+	}
+
+	int32 FAttributeMap::Max() const
+	{
+		return MaxValues;
+	}
+
+	UScriptStruct* FAttributeMap::GetValueType() const
+	{
+		return ValueTypeStruct;
+	}
+
+	int32 FAttributeMap::GetAllocatedSize() const
+	{
+		const int32 ValueTypeSize = ValueTypeStruct != nullptr ? ValueTypeStruct->GetStructureSize() : 0;
+		const int32 ValueTypeAlignment = ValueTypeSize >= 16 ? 16 : 8;
+		return ContainerSize + Align(sizeof(FName) * MaxValues, ValueTypeAlignment) + (ValueTypeSize * MaxValues);
+	}
+
+	FReallocFun FAttributeMap::GetAllocator() const
+	{
+		return ReallocFun;
+	}
+
+	void FAttributeMap::MoveTo(FAttributeMap& Other)
+	{
+		if (!ensureMsgf(ValueTypeStruct == Other.GetValueType(), TEXT("Value types must match to allow moving")))
+		{
+			return;
+		}
+
+		Swap(NumValues, Other.NumValues);
+		Swap(MaxValues, Other.MaxValues);
+		Swap(DataPtr, Other.DataPtr);
+		Swap(ReallocFun, Other.ReallocFun);
+		Swap(ContainerSize, Other.ContainerSize);
+	}
+}
diff --git a/Engine/Plugins/Experimental/UAF/UAF/Source/UAF/Private/UAF/ValueRuntime/AttributeSetMap.cpp b/Engine/Plugins/Experimental/UAF/UAF/Source/UAF/Private/UAF/ValueRuntime/AttributeSetMap.cpp
new file mode 100644
--- /dev/null
+++ b/Engine/Plugins/Experimental/UAF/UAF/Source/UAF/Private/UAF/ValueRuntime/AttributeSetMap.cpp
@@ -0,0 +1,65 @@
+// Copyright Epic Games, Inc. All Rights Reserved.
+
+#include "UAF/ValueRuntime/AttributeSetMap.h"
+
+namespace UE::UAF
+{
+	FAttributeSetMap::FAttributeSetMap() = default;
+
+	FAttributeSetMap::FAttributeSetMap(const FConstructArgs& Args, int32 InAllocatedSize)
+		: TypedSet(Args.TypedSet)
+		, ValueTypeStruct(Args.ValueType)
+		, AllocatedSize(InAllocatedSize)
+		, NumValues(0)
+		, ValuesPtr(nullptr)
+		, ReallocFun(Args.ReallocFun)
+	{
+		check(Args.TypedSet.IsValid());
+		check(Args.ReallocFun != nullptr);
+	}
+
+	FAttributeSetMap::~FAttributeSetMap()
+	{
+		checkf(ValuesPtr == nullptr, TEXT("Derived types should handle destruction"));
+	}
+
+	bool FAttributeSetMap::IsEmpty() const
+	{
+		return ValuesPtr == nullptr;
+	}
+
+	int32 FAttributeSetMap::Num() const
+	{
+		return NumValues;
+	}
+
+	const FAttributeTypedSetPtr& FAttributeSetMap::GetTypedSet() const
+	{
+		return TypedSet;
+	}
+
+	UScriptStruct* FAttributeSetMap::GetAttributeType() const
+	{
+		return TypedSet ? TypedSet->GetType() : nullptr;
+	}
+
+	UScriptStruct* FAttributeSetMap::GetValueType() const
+	{
+		return ValueTypeStruct;
+	}
+
+	FAttributeMappingKey FAttributeSetMap::GetMappingKey() const
+	{
+		return TypedSet ? FAttributeMappingKey(TypedSet->GetType(), ValueTypeStruct) : FAttributeMappingKey();
+	}
+
+	int32 FAttributeSetMap::GetAllocatedSize() const
+	{
+		return AllocatedSize;
+	}
+
+	FReallocFun FAttributeSetMap::GetAllocator() const
+	{
+		return ReallocFun;
+	}
+}
diff --git a/Engine/Plugins/Experimental/UAF/UAF/Source/UAF/Private/UAF/ValueRuntime/Transformers/Accumulate.cpp b/Engine/Plugins/Experimental/UAF/UAF/Source/UAF/Private/UAF/ValueRuntime/Transformers/Accumulate.cpp
new file mode 100644
--- /dev/null
+++ b/Engine/Plugins/Experimental/UAF/UAF/Source/UAF/Private/UAF/ValueRuntime/Transformers/Accumulate.cpp
@@ -0,0 +1,567 @@
+// Copyright Epic Games, Inc. All Rights Reserved.
+
+#include "UAF/ValueRuntime/Transformers/Accumulate.h"
+
+#include "UAF/ValueRuntime/IteratorUtils.h"
+
+namespace UE::UAF::Transformers
+{
+	FName FAccumulate::TransformerName = TEXT("UAF::Transformers::Accumulate");
+
+	void FAccumulate::Apply(const FAttributeValueTransformerMapPtr& TransformerMap, const FAttributeCollection& InputA, const FAttributeCollection& InputB, float WeightB, FAttributeCollection& Output)
+	{
+		// Named sets must match as it ensures that our inputs have the same sizes/shapes
+		checkf(InputA.GetNamedSet() == InputB.GetNamedSet(), TEXT("Both input collections must use the same named set"));
+		checkf(InputA.GetValueSpace() == InputB.GetValueSpace(), TEXT("Both input collections must have matching value spaces: [%s] != [%s]"), *InputA.GetValueSpace().ToString(), *InputB.GetValueSpace().ToString());
+
+		const bool bIsOutputA = &InputA == &Output;
+		if (bIsOutputA && !FAnimWeight::IsRelevant(WeightB))
+		{
+			// Our input A is our output and B is not relevant, nothing to do
+			return;
+		}
+
+		// Allocator is irrelevant since it will remain empty
+		FAttributeCollectionHeap EmptyCollection(InputA.GetNamedSet());
+		FAccumulate::Apply(TransformerMap, InputA, InputB, EmptyCollection, WeightB, Output);
+	}
+
+	void FAccumulate::Apply(const FAttributeValueTransformerMapPtr& TransformerMap, const FAttributeCollection& InputA, const FAttributeCollection& InputB, const FAttributeCollection& PerValueWeightsB, float DefaultWeightB, FAttributeCollection& Output)
+	{
+		// Named sets must match as it ensures that our inputs have the same sizes/shapes
+		checkf(InputA.GetNamedSet() == InputB.GetNamedSet(), TEXT("Both input collections must use the same named set"));
+		checkf(InputA.GetNamedSet() == PerValueWeightsB.GetNamedSet(), TEXT("Input and per-value weight collections must use the same named set"));
+		checkf(InputA.GetValueSpace() == InputB.GetValueSpace(), TEXT("Both input collections must have matching value spaces: [%s] != [%s]"), *InputA.GetValueSpace().ToString(), *InputB.GetValueSpace().ToString());
+
+		const FAttributeValueTransformerList* TransformerList = TransformerMap->Find(TransformerName);
+		if (TransformerList == nullptr)
+		{
+			// No transformer implementations have been registered
+			return;
+		}
+
+		UScriptStruct* FloatValueType = FFloatAnimationAttribute::StaticStruct();
+
+		const bool bIsOutputA = &InputA == &Output;
+		const bool bIsOutputB = &InputB == &Output;
+		if (bIsOutputA || bIsOutputB)
+		{
+			// Our output collection matches one of our inputs, we'll modify it in place
+			FReallocFun OutputAllocator = Output.GetAllocator();
+
+			// Set maps
+			{
+				// Keep a list of set maps that needs to be appended to our output
+				TArray<FAttributeSetMap*, TInlineAllocator<8>> PendingSetMaps;
+
+				OuterJoinBy(
+					FAttributeValueTransformerListWithAttributeCollectionSetMapJoinOp(),
+					[&PerValueWeightsB, DefaultWeightB, bIsOutputA, FloatValueType, OutputAllocator, &PendingSetMaps](FRawTransformerFunc Transformer, FAttributeSetMap* SetMapA, FAttributeSetMap* SetMapB)
+					{
+						if (SetMapA == nullptr)
+						{
+							if (SetMapB == nullptr)
+							{
+								// No maps found for the transformer type
+							}
+							else
+							{
+								// A doesn't have a matching pair
+								if (bIsOutputA)
+								{
+									// A is our output, duplicate it
+									PendingSetMaps.Add(SetMapB->Duplicate(OutputAllocator));
+								}
+								else
+								{
+									// B is our output, nothing to do
+								}
+							}
+						}
+						else
+						{
+							if (SetMapB == nullptr)
+							{
+								// B doesn't have a matching pair
+								if (bIsOutputA)
+								{
+									// A is our output, nothing to do
+								}
+								else
+								{
+									// B is our output, duplicate it
+									PendingSetMaps.Add(SetMapA->Duplicate(OutputAllocator));
+								}
+							}
+							else
+							{
+								if (Transformer == nullptr)
+								{
+									// No transformer was found for this value type, retain the one with the largest weight
+									if (DefaultWeightB >= 0.5f)
+									{
+										// B has more weight
+										if (bIsOutputA)
+										{
+											// A is our output, copy it over
+											SetMapB->CopyTo(*SetMapA);
+										}
+										else
+										{
+											// B is our output, nothing to do
+										}
+									}
+									else
+									{
+										// A has more weight
+										if (bIsOutputA)
+										{
+											// A is our output, nothing to do
+										}
+										else
+										{
+											// B is our output, copy it over
+											SetMapA->CopyTo(*SetMapB);
+										}
+									}
+								}
+								else
+								{
+									// Transform our matching pair in place
+									FAttributeSetMap* OutputSetMap = bIsOutputA ? SetMapA : SetMapB;
+
+									// We cannot join our weights because their value type differs
+									const FAttributeMappingKey MappingKey = SetMapA->GetMappingKey().To(FloatValueType);
+									const FAttributeSetMap* PerValueWeightMapB = PerValueWeightsB.FindSetMap(MappingKey);
+
+									FTransformSetMapFunc TypedTransformer = reinterpret_cast<FTransformSetMapFunc>(Transformer);
+									(*TypedTransformer)(SetMapA, SetMapB, PerValueWeightMapB, DefaultWeightB, OutputSetMap, nullptr);
+								}
+							}
+						}
+					},
+					TransformerList->CreateSetMapTransformerIterator(),
+					// Remove the const qualifier since we'll mutate one in place as our output
+					const_cast<FAttributeCollection&>(InputA).CreateSetMapIterator(),
+					const_cast<FAttributeCollection&>(InputB).CreateSetMapIterator());
+
+				// Append our pending entries
+				for (FAttributeSetMap* SetMap : PendingSetMaps)
+				{
+					Output.AddSetMap(SetMap);
+				}
+			}
+
+			// Name/value maps
+			{
+				// Keep a list of maps that needs to be appended to our output
+				TArray<FAttributeMap*, TInlineAllocator<8>> PendingMaps;
+
+				OuterJoinBy(
+					FAttributeValueTransformerListWithAttributeCollectionMapJoinOp(),
+					[DefaultWeightB, bIsOutputA, OutputAllocator, &PendingMaps](FRawTransformerFunc Transformer, FAttributeMap* MapA, FAttributeMap* MapB)
+					{
+						if (MapA == nullptr)
+						{
+							if (MapB == nullptr)
+							{
+								// No maps found for the transformer type
+							}
+							else
+							{
+								// A doesn't have a matching pair
+								if (bIsOutputA)
+								{
+									// A is our output, duplicate it
+									PendingMaps.Add(MapB->Duplicate(OutputAllocator));
+								}
+								else
+								{
+									// B is our output, nothing to do
+								}
+							}
+						}
+						else
+						{
+							if (MapB == nullptr)
+							{
+								// B doesn't have a matching pair
+								if (bIsOutputA)
+								{
+									// A is our output, nothing to do
+								}
+								else
+								{
+									// B is our output, duplicate it
+									PendingMaps.Add(MapA->Duplicate(OutputAllocator));
+								}
+							}
+							else
+							{
+								if (Transformer == nullptr)
+								{
+									// No transformer was found for this value type, retain the one with the largest weight
+									if (DefaultWeightB >= 0.5f)
+									{
+										// B has more weight
+										if (bIsOutputA)
+										{
+											// A is our output, copy it over
+											MapB->CopyTo(*MapA);
+										}
+										else
+										{
+											// B is our output, nothing to do
+										}
+									}
+									else
+									{
+										// A has more weight
+										if (bIsOutputA)
+										{
+											// A is our output, nothing to do
+										}
+										else
+										{
+											// B is our output, copy it over
+											MapA->CopyTo(*MapB);
+										}
+									}
+								}
+								else
+								{
+									// Transform our matching pair in a new map
+									FAttributeMap* OutputMap = nullptr;
+
+									FTransformMapFunc TypedTransformer = reinterpret_cast<FTransformMapFunc>(Transformer);
+									(*TypedTransformer)(MapA, MapB, DefaultWeightB, OutputMap, OutputAllocator);
+
+									checkf(OutputMap != nullptr, TEXT("Transformer function is responsible for allocating an output buffer"));
+									checkf(OutputMap->GetAllocator() == OutputAllocator, TEXT("Transformer function is responsible for allocating using the output allocator"));
+
+									// Overwrite our old output
+									OutputMap->MoveTo(bIsOutputA ? *MapA : *MapB);
+
+									// Destroy our temporary output map
+									ReleaseAttributeMap(OutputMap);
+								}
+							}
+						}
+					},
+					TransformerList->CreateMapTransformerIterator(),
+					// Remove the const qualifier since we'll mutate one in place as our output
+					const_cast<FAttributeCollection&>(InputA).CreateMapIterator(),
+					const_cast<FAttributeCollection&>(InputB).CreateMapIterator());
+
+				// Append our pending entries
+				for (FAttributeMap* Map : PendingMaps)
+				{
+					Output.AddMap(Map);
+				}
+			}
+		}
+		else
+		{
+			// Our output collection isn't either of our inputs, reset it
+			Output.Reset(InputA.GetNamedSet());
+			Output.SetValueSpace(InputA.GetValueSpace());
+
+			// Set maps
+			OuterJoinBy(
+				FAttributeValueTransformerListWithAttributeCollectionSetMapJoinOp(),
+				[&PerValueWeightsB, DefaultWeightB, FloatValueType, &Output](FRawTransformerFunc Transformer, const FAttributeSetMap* SetMapA, const FAttributeSetMap* SetMapB)
+				{
+					if (SetMapA == nullptr)
+					{
+						if (SetMapB == nullptr)
+						{
+							// No maps found for the transformer type
+						}
+						else
+						{
+							// This set map doesn't have a matching pair, duplicate it
+							Output.AppendSetMap(SetMapB->Duplicate(Output.GetAllocator()));
+						}
+					}
+					else
+					{
+						if (SetMapB == nullptr)
+						{
+							// This set map doesn't have a matching pair, duplicate it
+							Output.AppendSetMap(SetMapA->Duplicate(Output.GetAllocator()));
+						}
+						else
+						{
+							if (Transformer == nullptr)
+							{
+								// No transformer was found for this value type, retain the one with the largest weight
+								FAttributeSetMap* SetMapCopy;
+								if (DefaultWeightB >= 0.5f)
+								{
+									SetMapCopy = SetMapB->Duplicate(Output.GetAllocator());
+								}
+								else
+								{
+									SetMapCopy = SetMapA->Duplicate(Output.GetAllocator());
+								}
+
+								Output.AppendSetMap(SetMapCopy);
+							}
+							else
+							{
+								// Transform our matching pair
+								FAttributeSetMap* OutputSetMap = nullptr;
+
+								// We cannot join our weights because their value type differs
+								const FAttributeMappingKey MappingKey = SetMapA->GetMappingKey().To(FloatValueType);
+								const FAttributeSetMap* PerValueWeightMapB = PerValueWeightsB.FindSetMap(MappingKey);
+
+								FTransformSetMapFunc TypedTransformer = reinterpret_cast<FTransformSetMapFunc>(Transformer);
+								(*TypedTransformer)(SetMapA, SetMapB, PerValueWeightMapB, DefaultWeightB, OutputSetMap, Output.GetAllocator());
+
+								checkf(OutputSetMap != nullptr, TEXT("Transformer function is responsible for allocating an output buffer"));
+								checkf(OutputSetMap->GetAllocator() == Output.GetAllocator(), TEXT("Transformer function is responsible for allocating using the output allocator"));
+								Output.AppendSetMap(OutputSetMap);
+							}
+						}
+					}
+				},
+				TransformerList->CreateSetMapTransformerIterator(),
+				InputA.CreateConstSetMapIterator(),
+				InputB.CreateConstSetMapIterator());
+
+			// Name/value maps
+			OuterJoinBy(
+				FAttributeValueTransformerListWithAttributeCollectionMapJoinOp(),
+				[DefaultWeightB, &Output](FRawTransformerFunc Transformer, const FAttributeMap* MapA, const FAttributeMap* MapB)
+				{
+					if (MapA == nullptr)
+					{
+						if (MapB == nullptr)
+						{
+							// No maps found for the transformer type
+						}
+						else
+						{
+							// This map doesn't have a matching pair, duplicate it
+							Output.AppendMap(MapB->Duplicate(Output.GetAllocator()));
+						}
+					}
+					else
+					{
+						if (MapB == nullptr)
+						{
+							// This map doesn't have a matching pair
+							// Our output is A and thus already contains this
+						}
+						else
+						{
+							if (Transformer == nullptr)
+							{
+								// No transformer was found for this value type, retain the one with the largest weight
+								FAttributeMap* MapCopy;
+								if (DefaultWeightB >= 0.5f)
+								{
+									MapCopy = MapB->Duplicate(Output.GetAllocator());
+								}
+								else
+								{
+									MapCopy = MapA->Duplicate(Output.GetAllocator());
+								}
+
+								Output.AppendMap(MapCopy);
+							}
+							else
+							{
+								// Transform our matching pair in place
+								FAttributeMap* OutputMap = nullptr;
+
+								FTransformMapFunc TypedTransformer = reinterpret_cast<FTransformMapFunc>(Transformer);
+								(*TypedTransformer)(MapA, MapB, DefaultWeightB, OutputMap, Output.GetAllocator());
+
+								checkf(OutputMap != nullptr, TEXT("Transformer function is responsible for allocating an output buffer"));
+								checkf(OutputMap->GetAllocator() == Output.GetAllocator(), TEXT("Transformer function is responsible for allocating using the output allocator"));
+								Output.AppendMap(OutputMap);
+							}
+						}
+					}
+				},
+				TransformerList->CreateMapTransformerIterator(),
+				InputA.CreateConstMapIterator(),
+				InputB.CreateConstMapIterator());
+		}
+	}
+
+	namespace Private
+	{
+		void FAccumulate_BoneTransformAttribute::TransformSetMap(const FAttributeSetMap* SetMapA, const FAttributeSetMap* SetMapB, const FAttributeSetMap* PerValueWeightMapB, float DefaultWeightB, FAttributeSetMap*& OutputSetMap, FReallocFun OutputAllocator)
+		{
+			check(SetMapA->GetTypedSet() == SetMapB->GetTypedSet());
+			check(SetMapA->Num() == SetMapB->Num());
+
+			if (OutputSetMap == nullptr)
+			{
+				const FAttributeSetMap::FConstructArgs Args(SetMapA->GetTypedSet(), SetMapA->GetValueType(), OutputAllocator);
+				OutputSetMap = MakeAttributeSetMap<FBoneTransformAnimationAttribute>(Args);
+			}
+
+			check(SetMapA->GetTypedSet() == OutputSetMap->GetTypedSet());
+			check(SetMapA->Num() == OutputSetMap->Num());
+
+			const TAttributeSetMap<FBoneTransformAnimationAttribute>* BoneTransformsA = CastChecked<FBoneTransformAnimationAttribute>(SetMapA);
+			const TAttributeSetMap<FBoneTransformAnimationAttribute>* BoneTransformsB = CastChecked<FBoneTransformAnimationAttribute>(SetMapB);
+			TAttributeSetMap<FBoneTransformAnimationAttribute>* OutputBoneTransforms = CastChecked<FBoneTransformAnimationAttribute>(OutputSetMap);
+
+			const FBoneTransformAnimationAttribute* TransformPtrA = BoneTransformsA->GetData();
+			const FBoneTransformAnimationAttribute* TransformPtrB = BoneTransformsB->GetData();
+			FBoneTransformAnimationAttribute* OutputTransformPtr = OutputBoneTransforms->GetData();
+			const FBoneTransformAnimationAttribute* TransformPtrAEnd = TransformPtrA + BoneTransformsA->Num();
+
+			if (const TAttributeSetMap<FFloatAnimationAttribute>* PerValueWeightsB = Cast<FFloatAnimationAttribute>(PerValueWeightMapB))
+			{
+				check(SetMapA->GetTypedSet() == PerValueWeightsB->GetTypedSet());
+				check(SetMapA->Num() == PerValueWeightsB->Num());
+
+				const FFloatAnimationAttribute* PerValueWeightsBPtr = PerValueWeightsB->GetData();
+
+				while (TransformPtrA < TransformPtrAEnd)
+				{
+					FTransform Result = TransformPtrA->Value;
+					Result.AccumulateWithShortestRotation(TransformPtrB->Value, ScalarRegister(PerValueWeightsBPtr->Value));
+
+					OutputTransformPtr->Value = Result;
+
+					TransformPtrA++;
+					TransformPtrB++;
+					PerValueWeightsBPtr++;
+					OutputTransformPtr++;
+				}
+			}
+			else
+			{
+				const ScalarRegister VDefaultWeightB(DefaultWeightB);
+
+				while (TransformPtrA < TransformPtrAEnd)
+				{
+					FTransform Result = TransformPtrA->Value;
+					Result.AccumulateWithShortestRotation(TransformPtrB->Value, VDefaultWeightB);
+
+					OutputTransformPtr->Value = Result;
+
+					TransformPtrA++;
+					TransformPtrB++;
+					OutputTransformPtr++;
+				}
+			}
+		}
+
+		void FAccumulate_BoneTransformAttribute::TransformMap(const FAttributeMap* MapA, const FAttributeMap* MapB, float WeightB, FAttributeMap*& OutputMap, FReallocFun OutputAllocator)
+		{
+			check(OutputMap == nullptr);
+
+			const TAttributeMap<FBoneTransformAnimationAttribute>* BoneTransformsA = CastChecked<FBoneTransformAnimationAttribute>(MapA);
+			const TAttributeMap<FBoneTransformAnimationAttribute>* BoneTransformsB = CastChecked<FBoneTransformAnimationAttribute>(MapB);
+			TAttributeMap<FBoneTransformAnimationAttribute>* OutputBoneTransforms = MakeAttributeMap<FBoneTransformAnimationAttribute>(OutputAllocator);
+
+			const ScalarRegister VWeightB(WeightB);
+
+			OuterJoinBy(
+				FAttributeMapJoinOp(),
+				[VWeightB, OutputBoneTransforms](FName Name, const FBoneTransformAnimationAttribute* InputA, const FBoneTransformAnimationAttribute* InputB)
+				{
+					// Missing values are set to the identity
+					const FTransform& ValueA = InputA != nullptr ? InputA->Value : FTransform::Identity;
+					const FTransform& ValueB = InputB != nullptr ? InputB->Value : FTransform::Identity;
+
+					FTransform Result = ValueA;
+					Result.AccumulateWithShortestRotation(ValueB, VWeightB);
+
+					OutputBoneTransforms->Append(Name, FBoneTransformAnimationAttribute{ Result });
+				},
+				BoneTransformsA->CreateConstIterator(),
+				BoneTransformsB->CreateConstIterator());
+
+			OutputMap = OutputBoneTransforms;
+		}
+
+		void FAccumulate_FloatAttribute::TransformSetMap(const FAttributeSetMap* SetMapA, const FAttributeSetMap* SetMapB, const FAttributeSetMap* PerValueWeightMapB, float DefaultWeightB, FAttributeSetMap*& OutputSetMap, FReallocFun OutputAllocator)
+		{
+			check(SetMapA->GetTypedSet() == SetMapB->GetTypedSet());
+			check(SetMapA->Num() == SetMapB->Num());
+
+			if (OutputSetMap == nullptr)
+			{
+				const FAttributeSetMap::FConstructArgs Args(SetMapA->GetTypedSet(), SetMapA->GetValueType(), OutputAllocator);
+				OutputSetMap = MakeAttributeSetMap<FFloatAnimationAttribute>(Args);
+			}
+
+			check(SetMapA->GetTypedSet() == OutputSetMap->GetTypedSet());
+			check(SetMapA->Num() == OutputSetMap->Num());
+
+			const TAttributeSetMap<FFloatAnimationAttribute>* FloatsA = CastChecked<FFloatAnimationAttribute>(SetMapA);
+			const TAttributeSetMap<FFloatAnimationAttribute>* FloatsB = CastChecked<FFloatAnimationAttribute>(SetMapB);
+			TAttributeSetMap<FFloatAnimationAttribute>* OutputFloats = CastChecked<FFloatAnimationAttribute>(OutputSetMap);
+
+			const FFloatAnimationAttribute* FloatsPtrA = FloatsA->GetData();
+			const FFloatAnimationAttribute* FloatsPtrB = FloatsB->GetData();
+			FFloatAnimationAttribute* OutputFloatsPtr = OutputFloats->GetData();
+			const FFloatAnimationAttribute* FloatsPtrAEnd = FloatsPtrA + FloatsA->Num();
+
+			if (const TAttributeSetMap<FFloatAnimationAttribute>* PerValueWeightsB = Cast<FFloatAnimationAttribute>(PerValueWeightMapB))
+			{
+				check(SetMapA->GetTypedSet() == PerValueWeightsB->GetTypedSet());
+				check(SetMapA->Num() == PerValueWeightsB->Num());
+
+				const FFloatAnimationAttribute* PerValueWeightsBPtr = PerValueWeightsB->GetData();
+
+				while (FloatsPtrA < FloatsPtrAEnd)
+				{
+					OutputFloatsPtr->Value = FloatsPtrA->Value + (FloatsPtrB->Value * PerValueWeightsBPtr->Value);
+
+					FloatsPtrA++;
+					FloatsPtrB++;
+					PerValueWeightsBPtr++;
+					OutputFloatsPtr++;
+				}
+			}
+			else
+			{
+				while (FloatsPtrA < FloatsPtrAEnd)
+				{
+					OutputFloatsPtr->Value = FloatsPtrA->Value + (FloatsPtrB->Value * DefaultWeightB);
+
+					FloatsPtrA++;
+					FloatsPtrB++;
+					OutputFloatsPtr++;
+				}
+			}
+		}
+
+		void FAccumulate_FloatAttribute::TransformMap(const FAttributeMap* MapA, const FAttributeMap* MapB, float WeightB, FAttributeMap*& OutputMap, FReallocFun OutputAllocator)
+		{
+			check(OutputMap == nullptr);
+
+			const TAttributeMap<FFloatAnimationAttribute>* FloatsA = CastChecked<FFloatAnimationAttribute>(MapA);
+			const TAttributeMap<FFloatAnimationAttribute>* FloatsB = CastChecked<FFloatAnimationAttribute>(MapB);
+			TAttributeMap<FFloatAnimationAttribute>* OutputFloats = MakeAttributeMap<FFloatAnimationAttribute>(OutputAllocator);
+
+			OuterJoinBy(
+				FAttributeMapJoinOp(),
+				[&WeightB, OutputFloats](FName Name, const FFloatAnimationAttribute* InputA, const FFloatAnimationAttribute* InputB)
+				{
+					// Missing values are set to 0.0
+					const float ValueA = InputA != nullptr ? InputA->Value : 0.0f;
+					const float ValueB = InputB != nullptr ? InputB->Value : 0.0f;
+
+					FFloatAnimationAttribute Output;
+					Output.Value = ValueA + (ValueB * WeightB);
+
+					OutputFloats->Append(Name, Output);
+				},
+				FloatsA->CreateConstIterator(),
+				FloatsB->CreateConstIterator());
+
+			OutputMap = OutputFloats;
+		}
+	}
+}
diff --git a/Engine/Plugins/Experimental/UAF/UAF/Source/UAF/Private/UAF/ValueRuntime/Transformers/Interpolate.cpp b/Engine/Plugins/Experimental/UAF/UAF/Source/UAF/Private/UAF/ValueRuntime/Transformers/Interpolate.cpp
new file mode 100644
--- /dev/null
+++ b/Engine/Plugins/Experimental/UAF/UAF/Source/UAF/Private/UAF/ValueRuntime/Transformers/Interpolate.cpp
@@ -0,0 +1,575 @@
+// Copyright Epic Games, Inc. All Rights Reserved.
+
+#include "UAF/ValueRuntime/Transformers/Interpolate.h"
+
+#include "UAF/ValueRuntime/IteratorUtils.h"
+
+namespace UE::UAF::Transformers
+{
+	FName FInterpolate::TransformerName = TEXT("UAF::Transformers::Interpolate");
+
+	void FInterpolate::Apply(
+		const FAttributeValueTransformerMapPtr& TransformerMap,
+		const FAttributeCollection& InputA,
+		const FAttributeCollection& InputB,
+		float WeightB,
+		FAttributeCollection& Output)
+	{
+		// Named sets must match as it ensures that our inputs have the same sizes/shapes
+		checkf(InputA.GetNamedSet() == InputB.GetNamedSet(), TEXT("Both input collections must use the same named set"));
+		checkf(InputA.GetValueSpace() == InputB.GetValueSpace(), TEXT("Both input collections must have matching value spaces: [%s] != [%s]"), *InputA.GetValueSpace().ToString(), *InputB.GetValueSpace().ToString());
+
+		const bool bIsOutputA = &InputA == &Output;
+		if (bIsOutputA && !FAnimWeight::IsRelevant(WeightB))
+		{
+			// Our input A is our output and B is not relevant, nothing to do
+			return;
+		}
+
+		const bool bIsOutputB = &InputB == &Output;
+		if (bIsOutputB && FAnimWeight::IsFullWeight(WeightB))
+		{
+			// Our input B is our output and A is not relevant, nothing to do
+			return;
+		}
+
+		// Allocator is irrelevant since it will remain empty
+		FAttributeCollectionHeap EmptyCollection(InputA.GetNamedSet());
+		FInterpolate::Apply(TransformerMap, InputA, InputB, EmptyCollection, WeightB, Output);
+	}
+
+	void FInterpolate::Apply(
+		const FAttributeValueTransformerMapPtr& TransformerMap,
+		const FAttributeCollection& InputA,
+		const FAttributeCollection& InputB,
+		const FAttributeCollection& PerValueWeightsB,
+		float DefaultWeightB,
+		FAttributeCollection& Output)
+	{
+		// Named sets must match as it ensures that our inputs have the same sizes/shapes
+		checkf(InputA.GetNamedSet() == InputB.GetNamedSet(), TEXT("Both input collections must use the same named set"));
+		checkf(InputA.GetNamedSet() == PerValueWeightsB.GetNamedSet(), TEXT("Input and per-value weight collections must use the same named set"));
+		checkf(InputA.GetValueSpace() == InputB.GetValueSpace(), TEXT("Both input collections must have matching value spaces: [%s] != [%s]"), *InputA.GetValueSpace().ToString(), *InputB.GetValueSpace().ToString());
+
+		const FAttributeValueTransformerList* TransformerList = TransformerMap->Find(TransformerName);
+		if (TransformerList == nullptr)
+		{
+			// No transformer implementations have been registered
+			return;
+		}
+
+		UScriptStruct* FloatValueType = FFloatAnimationAttribute::StaticStruct();
+
+		const bool bIsOutputA = &InputA == &Output;
+		const bool bIsOutputB = &InputB == &Output;
+		if (bIsOutputA || bIsOutputB)
+		{
+			// Our output collection matches one of our inputs, we'll modify it in place
+			FReallocFun OutputAllocator = Output.GetAllocator();
+
+			// Set maps
+			{
+				// Keep a list of set maps that needs to be appended to our output
+				TArray<FAttributeSetMap*, TInlineAllocator<8>> PendingSetMaps;
+
+				OuterJoinBy(
+					FAttributeValueTransformerListWithAttributeCollectionSetMapJoinOp(),
+					[&PerValueWeightsB, DefaultWeightB, bIsOutputA, FloatValueType, OutputAllocator, &PendingSetMaps](FRawTransformerFunc Transformer, FAttributeSetMap* SetMapA, FAttributeSetMap* SetMapB)
+					{
+						if (SetMapA == nullptr)
+						{
+							if (SetMapB == nullptr)
+							{
+								// No maps found for the transformer type
+							}
+							else
+							{
+								// A doesn't have a matching pair
+								if (bIsOutputA)
+								{
+									// A is our output, duplicate it
+									PendingSetMaps.Add(SetMapB->Duplicate(OutputAllocator));
+								}
+								else
+								{
+									// B is our output, nothing to do
+								}
+							}
+						}
+						else
+						{
+							if (SetMapB == nullptr)
+							{
+								// B doesn't have a matching pair
+								if (bIsOutputA)
+								{
+									// A is our output, nothing to do
+								}
+								else
+								{
+									// B is our output, duplicate it
+									PendingSetMaps.Add(SetMapA->Duplicate(OutputAllocator));
+								}
+							}
+							else
+							{
+								if (Transformer == nullptr)
+								{
+									// No transformer was found for this value type, retain the one with the largest weight
+									if (DefaultWeightB >= 0.5f)
+									{
+										// B has more weight
+										if (bIsOutputA)
+										{
+											// A is our output, copy it over
+											SetMapB->CopyTo(*SetMapA);
+										}
+										else
+										{
+											// B is our output, nothing to do
+										}
+									}
+									else
+									{
+										// A has more weight
+										if (bIsOutputA)
+										{
+											// A is our output, nothing to do
+										}
+										else
+										{
+											// B is our output, copy it over
+											SetMapA->CopyTo(*SetMapB);
+										}
+									}
+								}
+								else
+								{
+									// Transform our matching pair in place
+									FAttributeSetMap* OutputSetMap = bIsOutputA ? SetMapA : SetMapB;
+
+									// We cannot join our weights because their value type differs
+									const FAttributeMappingKey MappingKey = SetMapA->GetMappingKey().To(FloatValueType);
+									const FAttributeSetMap* PerValueWeightMapB = PerValueWeightsB.FindSetMap(MappingKey);
+
+									FTransformSetMapFunc TypedTransformer = reinterpret_cast<FTransformSetMapFunc>(Transformer);
+									(*TypedTransformer)(SetMapA, SetMapB, PerValueWeightMapB, DefaultWeightB, OutputSetMap, nullptr);
+								}
+							}
+						}
+					},
+					TransformerList->CreateSetMapTransformerIterator(),
+					// Remove the const qualifier since we'll mutate one in place as our output
+					const_cast<FAttributeCollection&>(InputA).CreateSetMapIterator(),
+					const_cast<FAttributeCollection&>(InputB).CreateSetMapIterator());
+
+				// Append our pending entries
+				for (FAttributeSetMap* SetMap : PendingSetMaps)
+				{
+					Output.AddSetMap(SetMap);
+				}
+			}
+
+			// Name/value maps
+			{
+				// Keep a list of maps that needs to be appended to our output
+				TArray<FAttributeMap*, TInlineAllocator<8>> PendingMaps;
+
+				OuterJoinBy(
+					FAttributeValueTransformerListWithAttributeCollectionMapJoinOp(),
+					[DefaultWeightB, bIsOutputA, OutputAllocator, &PendingMaps](FRawTransformerFunc Transformer, FAttributeMap* MapA, FAttributeMap* MapB)
+					{
+						if (MapA == nullptr)
+						{
+							if (MapB == nullptr)
+							{
+								// No maps found for the transformer type
+							}
+							else
+							{
+								// A doesn't have a matching pair
+								if (bIsOutputA)
+								{
+									// A is our output, duplicate it
+									PendingMaps.Add(MapB->Duplicate(OutputAllocator));
+								}
+								else
+								{
+									// B is our output, nothing to do
+								}
+							}
+						}
+						else
+						{
+							if (MapB == nullptr)
+							{
+								// B doesn't have a matching pair
+								if (bIsOutputA)
+								{
+									// A is our output, nothing to do
+								}
+								else
+								{
+									// B is our output, duplicate it
+									PendingMaps.Add(MapA->Duplicate(OutputAllocator));
+								}
+							}
+							else
+							{
+								if (Transformer == nullptr)
+								{
+									// No transformer was found for this value type, retain the one with the largest weight
+									if (DefaultWeightB >= 0.5f)
+									{
+										// B has more weight
+										if (bIsOutputA)
+										{
+											// A is our output, copy it over
+											MapB->CopyTo(*MapA);
+										}
+										else
+										{
+											// B is our output, nothing to do
+										}
+									}
+									else
+									{
+										// A has more weight
+										if (bIsOutputA)
+										{
+											// A is our output, nothing to do
+										}
+										else
+										{
+											// B is our output, copy it over
+											MapA->CopyTo(*MapB);
+										}
+									}
+								}
+								else
+								{
+									// Transform our matching pair in a new map
+									FAttributeMap* OutputMap = nullptr;
+
+									FTransformMapFunc TypedTransformer = reinterpret_cast<FTransformMapFunc>(Transformer);
+									(*TypedTransformer)(MapA, MapB, DefaultWeightB, OutputMap, OutputAllocator);
+
+									checkf(OutputMap != nullptr, TEXT("Transformer function is responsible for allocating an output buffer"));
+									checkf(OutputMap->GetAllocator() == OutputAllocator, TEXT("Transformer function is responsible for allocating using the output allocator"));
+
+									// Overwrite our old output
+									OutputMap->MoveTo(bIsOutputA ? *MapA : *MapB);
+
+									// Destroy our temporary output map
+									ReleaseAttributeMap(OutputMap);
+								}
+							}
+						}
+					},
+					TransformerList->CreateMapTransformerIterator(),
+					// Remove the const qualifier since we'll mutate one in place as our output
+					const_cast<FAttributeCollection&>(InputA).CreateMapIterator(),
+					const_cast<FAttributeCollection&>(InputB).CreateMapIterator());
+
+				// Append our pending entries
+				for (FAttributeMap* Map : PendingMaps)
+				{
+					Output.AddMap(Map);
+				}
+			}
+		}
+		else
+		{
+			// Our output collection isn't either of our inputs, reset it
+			Output.Reset(InputA.GetNamedSet());
+			Output.SetValueSpace(InputA.GetValueSpace());
+
+			// Set maps
+			OuterJoinBy(
+				FAttributeValueTransformerListWithAttributeCollectionSetMapJoinOp(),
+				[&PerValueWeightsB, DefaultWeightB, FloatValueType, &Output](FRawTransformerFunc Transformer, const FAttributeSetMap* SetMapA, const FAttributeSetMap* SetMapB)
+				{
+					if (SetMapA == nullptr)
+					{
+						if (SetMapB == nullptr)
+						{
+							// No maps found for the transformer type
+						}
+						else
+						{
+							// This set map doesn't have a matching pair, duplicate it
+							Output.AppendSetMap(SetMapB->Duplicate(Output.GetAllocator()));
+						}
+					}
+					else
+					{
+						if (SetMapB == nullptr)
+						{
+							// This set map doesn't have a matching pair, duplicate it
+							Output.AppendSetMap(SetMapA->Duplicate(Output.GetAllocator()));
+						}
+						else
+						{
+							if (Transformer == nullptr)
+							{
+								// No transformer was found for this value type, retain the one with the largest weight
+								FAttributeSetMap* SetMapCopy;
+								if (DefaultWeightB >= 0.5f)
+								{
+									SetMapCopy = SetMapB->Duplicate(Output.GetAllocator());
+								}
+								else
+								{
+									SetMapCopy = SetMapA->Duplicate(Output.GetAllocator());
+								}
+
+								Output.AppendSetMap(SetMapCopy);
+							}
+							else
+							{
+								// Transform our matching pair
+								FAttributeSetMap* OutputSetMap = nullptr;
+
+								// We cannot join our weights because their value type differs
+								const FAttributeMappingKey MappingKey = SetMapA->GetMappingKey().To(FloatValueType);
+								const FAttributeSetMap* PerValueWeightMapB = PerValueWeightsB.FindSetMap(MappingKey);
+
+								FTransformSetMapFunc TypedTransformer = reinterpret_cast<FTransformSetMapFunc>(Transformer);
+								(*TypedTransformer)(SetMapA, SetMapB, PerValueWeightMapB, DefaultWeightB, OutputSetMap, Output.GetAllocator());
+
+								checkf(OutputSetMap != nullptr, TEXT("Transformer function is responsible for allocating an output buffer"));
+								checkf(OutputSetMap->GetAllocator() == Output.GetAllocator(), TEXT("Transformer function is responsible for allocating using the output allocator"));
+								Output.AppendSetMap(OutputSetMap);
+							}
+						}
+					}
+				},
+				TransformerList->CreateSetMapTransformerIterator(),
+				InputA.CreateConstSetMapIterator(),
+				InputB.CreateConstSetMapIterator());
+
+			// Name/value maps
+			OuterJoinBy(
+				FAttributeValueTransformerListWithAttributeCollectionMapJoinOp(),
+				[DefaultWeightB, &Output](FRawTransformerFunc Transformer, const FAttributeMap* MapA, const FAttributeMap* MapB)
+				{
+					if (MapA == nullptr)
+					{
+						if (MapB == nullptr)
+						{
+							// No maps found for the transformer type
+						}
+						else
+						{
+							// This map doesn't have a matching pair, duplicate it
+							Output.AppendMap(MapB->Duplicate(Output.GetAllocator()));
+						}
+					}
+					else
+					{
+						if (MapB == nullptr)
+						{
+							// This map doesn't have a matching pair
+							// Our output is A and thus already contains this
+						}
+						else
+						{
+							if (Transformer == nullptr)
+							{
+								// No transformer was found for this value type, retain the one with the largest weight
+								FAttributeMap* MapCopy;
+								if (DefaultWeightB >= 0.5f)
+								{
+									MapCopy = MapB->Duplicate(Output.GetAllocator());
+								}
+								else
+								{
+									MapCopy = MapA->Duplicate(Output.GetAllocator());
+								}
+
+								Output.AppendMap(MapCopy);
+							}
+							else
+							{
+								// Transform our matching pair in place
+								FAttributeMap* OutputMap = nullptr;
+
+								FTransformMapFunc TypedTransformer = reinterpret_cast<FTransformMapFunc>(Transformer);
+								(*TypedTransformer)(MapA, MapB, DefaultWeightB, OutputMap, Output.GetAllocator());
+
+								checkf(OutputMap != nullptr, TEXT("Transformer function is responsible for allocating an output buffer"));
+								checkf(OutputMap->GetAllocator() == Output.GetAllocator(), TEXT("Transformer function is responsible for allocating using the output allocator"));
+								Output.AppendMap(OutputMap);
+							}
+						}
+					}
+				},
+				TransformerList->CreateMapTransformerIterator(),
+				InputA.CreateConstMapIterator(),
+				InputB.CreateConstMapIterator());
+		}
+	}
+
+	namespace Private
+	{
+		void FInterpolate_BoneTransformAttribute::TransformSetMap(const FAttributeSetMap* SetMapA, const FAttributeSetMap* SetMapB, const FAttributeSetMap* PerValueWeightMapB, float DefaultWeightB, FAttributeSetMap*& OutputSetMap, FReallocFun OutputAllocator)
+		{
+			check(SetMapA->GetTypedSet() == SetMapB->GetTypedSet());
+			check(SetMapA->Num() == SetMapB->Num());
+
+			if (OutputSetMap == nullptr)
+			{
+				const FAttributeSetMap::FConstructArgs Args(SetMapA->GetTypedSet(), SetMapA->GetValueType(), OutputAllocator);
+				OutputSetMap = MakeAttributeSetMap<FBoneTransformAnimationAttribute>(Args);
+			}
+
+			check(SetMapA->GetTypedSet() == OutputSetMap->GetTypedSet());
+			check(SetMapA->Num() == OutputSetMap->Num());
+
+			const TAttributeSetMap<FBoneTransformAnimationAttribute>* BoneTransformsA = CastChecked<FBoneTransformAnimationAttribute>(SetMapA);
+			const TAttributeSetMap<FBoneTransformAnimationAttribute>* BoneTransformsB = CastChecked<FBoneTransformAnimationAttribute>(SetMapB);
+			TAttributeSetMap<FBoneTransformAnimationAttribute>* OutputBoneTransforms = CastChecked<FBoneTransformAnimationAttribute>(OutputSetMap);
+
+			const FBoneTransformAnimationAttribute* TransformPtrA = BoneTransformsA->GetData();
+			const FBoneTransformAnimationAttribute* TransformPtrB = BoneTransformsB->GetData();
+			FBoneTransformAnimationAttribute* OutputTransformPtr = OutputBoneTransforms->GetData();
+			const FBoneTransformAnimationAttribute* TransformPtrAEnd = TransformPtrA + BoneTransformsA->Num();
+
+			if (const TAttributeSetMap<FFloatAnimationAttribute>* PerValueWeightsB = Cast<FFloatAnimationAttribute>(PerValueWeightMapB))
+			{
+				check(SetMapA->GetTypedSet() == PerValueWeightsB->GetTypedSet());
+				check(SetMapA->Num() == PerValueWeightsB->Num());
+
+				const FFloatAnimationAttribute* PerValueWeightsBPtr = PerValueWeightsB->GetData();
+
+				while (TransformPtrA < TransformPtrAEnd)
+				{
+					OutputTransformPtr->Value.Blend(TransformPtrA->Value, TransformPtrB->Value, PerValueWeightsBPtr->Value);
+
+					TransformPtrA++;
+					TransformPtrB++;
+					PerValueWeightsBPtr++;
+					OutputTransformPtr++;
+				}
+			}
+			else
+			{
+				while (TransformPtrA < TransformPtrAEnd)
+				{
+					OutputTransformPtr->Value.Blend(TransformPtrA->Value, TransformPtrB->Value, DefaultWeightB);
+
+					TransformPtrA++;
+					TransformPtrB++;
+					OutputTransformPtr++;
+				}
+			}
+		}
+
+		void FInterpolate_BoneTransformAttribute::TransformMap(const FAttributeMap* MapA, const FAttributeMap* MapB, float WeightB, FAttributeMap*& OutputMap, FReallocFun OutputAllocator)
+		{
+			check(OutputMap == nullptr);
+
+			const TAttributeMap<FBoneTransformAnimationAttribute>* BoneTransformsA = CastChecked<FBoneTransformAnimationAttribute>(MapA);
+			const TAttributeMap<FBoneTransformAnimationAttribute>* BoneTransformsB = CastChecked<FBoneTransformAnimationAttribute>(MapB);
+			TAttributeMap<FBoneTransformAnimationAttribute>* OutputBoneTransforms = MakeAttributeMap<FBoneTransformAnimationAttribute>(OutputAllocator);
+
+			OuterJoinBy(
+				FAttributeMapJoinOp(),
+				[WeightB, OutputBoneTransforms](FName Name, const FBoneTransformAnimationAttribute* A, const FBoneTransformAnimationAttribute* B)
+				{
+					// Missing values are set to the identity
+					const FTransform& ValueA = A != nullptr ? A->Value : FTransform::Identity;
+					const FTransform& ValueB = B != nullptr ? B->Value : FTransform::Identity;
+
+					FBoneTransformAnimationAttribute Output;
+					Output.Value.Blend(ValueA, ValueB, WeightB);
+
+					OutputBoneTransforms->Append(Name, Output);
+				},
+				BoneTransformsA->CreateConstIterator(),
+				BoneTransformsB->CreateConstIterator());
+
+			OutputMap = OutputBoneTransforms;
+		}
+
+		void FInterpolate_FloatAttribute::TransformSetMap(const FAttributeSetMap* SetMapA, const FAttributeSetMap* SetMapB, const FAttributeSetMap* PerValueWeightMapB, float DefaultWeightB, FAttributeSetMap*& OutputSetMap, FReallocFun OutputAllocator)
+		{
+			check(SetMapA->GetTypedSet() == SetMapB->GetTypedSet());
+			check(SetMapA->Num() == SetMapB->Num());
+
+			if (OutputSetMap == nullptr)
+			{
+				const FAttributeSetMap::FConstructArgs Args(SetMapA->GetTypedSet(), SetMapA->GetValueType(), OutputAllocator);
+				OutputSetMap = MakeAttributeSetMap<FFloatAnimationAttribute>(Args);
+			}
+
+			check(SetMapA->GetTypedSet() == OutputSetMap->GetTypedSet());
+			check(SetMapA->Num() == OutputSetMap->Num());
+
+			const TAttributeSetMap<FFloatAnimationAttribute>* FloatsA = CastChecked<FFloatAnimationAttribute>(SetMapA);
+			const TAttributeSetMap<FFloatAnimationAttribute>* FloatsB = CastChecked<FFloatAnimationAttribute>(SetMapB);
+			TAttributeSetMap<FFloatAnimationAttribute>* OutputFloats = CastChecked<FFloatAnimationAttribute>(OutputSetMap);
+
+			const FFloatAnimationAttribute* FloatsPtrA = FloatsA->GetData();
+			const FFloatAnimationAttribute* FloatsPtrB = FloatsB->GetData();
+			FFloatAnimationAttribute* OutputFloatsPtr = OutputFloats->GetData();
+			const FFloatAnimationAttribute* FloatsPtrAEnd = FloatsPtrA + FloatsA->Num();
+
+			if (const TAttributeSetMap<FFloatAnimationAttribute>* PerValueWeightsB = Cast<FFloatAnimationAttribute>(PerValueWeightMapB))
+			{
+				check(SetMapA->GetTypedSet() == PerValueWeightsB->GetTypedSet());
+				check(SetMapA->Num() == PerValueWeightsB->Num());
+
+				const FFloatAnimationAttribute* PerValueWeightsBPtr = PerValueWeightsB->GetData();
+
+				while (FloatsPtrA < FloatsPtrAEnd)
+				{
+					OutputFloatsPtr->Value = FMath::Lerp(FloatsPtrA->Value, FloatsPtrB->Value, PerValueWeightsBPtr->Value);
+
+					FloatsPtrA++;
+					FloatsPtrB++;
+					PerValueWeightsBPtr++;
+					OutputFloatsPtr++;
+				}
+			}
+			else
+			{
+				while (FloatsPtrA < FloatsPtrAEnd)
+				{
+					OutputFloatsPtr->Value = FMath::Lerp(FloatsPtrA->Value, FloatsPtrB->Value, DefaultWeightB);
+
+					FloatsPtrA++;
+					FloatsPtrB++;
+					OutputFloatsPtr++;
+				}
+			}
+		}
+
+		void FInterpolate_FloatAttribute::TransformMap(const FAttributeMap* MapA, const FAttributeMap* MapB, float WeightB, FAttributeMap*& OutputMap, FReallocFun OutputAllocator)
+		{
+			check(OutputMap == nullptr);
+
+			const TAttributeMap<FFloatAnimationAttribute>* FloatsA = CastChecked<FFloatAnimationAttribute>(MapA);
+			const TAttributeMap<FFloatAnimationAttribute>* FloatsB = CastChecked<FFloatAnimationAttribute>(MapB);
+			TAttributeMap<FFloatAnimationAttribute>* OutputFloats = MakeAttributeMap<FFloatAnimationAttribute>(OutputAllocator);
+
+			OuterJoinBy(
+				FAttributeMapJoinOp(),
+				[WeightB, OutputFloats](FName Name, const FFloatAnimationAttribute* A, const FFloatAnimationAttribute* B)
+				{
+					// Missing values are set to 0.0
+					const float ValueA = A != nullptr ? A->Value : 0.0f;
+					const float ValueB = B != nullptr ? B->Value : 0.0f;
+
+					FFloatAnimationAttribute Output;
+					Output.Value = FMath::Lerp(ValueA, ValueB, WeightB);
+
+					OutputFloats->Append(Name, Output);
+				},
+				FloatsA->CreateConstIterator(),
+				FloatsB->CreateConstIterator());
+
+			OutputMap = OutputFloats;
+		}
+	}
+}
diff --git a/Engine/Plugins/Experimental/UAF/UAF/Source/UAF/Private/UAF/ValueRuntime/Transformers/Overwrite.cpp b/Engine/Plugins/Experimental/UAF/UAF/Source/UAF/Private/UAF/ValueRuntime/Transformers/Overwrite.cpp
new file mode 100644
--- /dev/null
+++ b/Engine/Plugins/Experimental/UAF/UAF/Source/UAF/Private/UAF/ValueRuntime/Transformers/Overwrite.cpp
@@ -0,0 +1,350 @@
+// Copyright Epic Games, Inc. All Rights Reserved.
+
+#include "UAF/ValueRuntime/Transformers/Overwrite.h"
+
+#include "UAF/ValueRuntime/IteratorUtils.h"
+
+namespace UE::UAF::Transformers
+{
+	FName FOverwrite::TransformerName = TEXT("UAF::Transformers::Overwrite");
+
+	void FOverwrite::Apply(const FAttributeValueTransformerMapPtr& TransformerMap, const FAttributeCollection& Input, float Weight, FAttributeCollection& Output)
+	{
+		const bool bIsInputOutput = &Input == &Output;
+		if (bIsInputOutput && FAnimWeight::IsFullWeight(Weight))
+		{
+			// Our input is our output and it has full weight, nothing to do
+			return;
+		}
+		
+		// Allocator is irrelevant since it will remain empty
+		FAttributeCollectionHeap EmptyCollection(Input.GetNamedSet());
+		FOverwrite::Apply(TransformerMap, Input, EmptyCollection, Weight, Output);
+	}
+
+	void FOverwrite::Apply(const FAttributeValueTransformerMapPtr& TransformerMap, const FAttributeCollection& Input, const FAttributeCollection& PerValueWeights, float DefaultWeight, FAttributeCollection& Output)
+	{
+		checkf(Input.GetNamedSet() == PerValueWeights.GetNamedSet(), TEXT("Input and per-value weight collections must use the same named set"));
+
+		const FAttributeValueTransformerList* TransformerList = TransformerMap->Find(TransformerName);
+		if (TransformerList == nullptr)
+		{
+			// No transformer implementations have been registered
+			return;
+		}
+
+		UScriptStruct* FloatValueType = FFloatAnimationAttribute::StaticStruct();
+
+		const bool bIsInputOutput = &Input == &Output;
+		if (bIsInputOutput)
+		{
+			// Our output collection matches one of our inputs, we'll modify it in place
+			FReallocFun OutputAllocator = Output.GetAllocator();
+
+			// Set maps
+			{
+				OuterJoinBy(
+					FAttributeValueTransformerListWithAttributeCollectionSetMapJoinOp(),
+					[&PerValueWeights, DefaultWeight, FloatValueType, OutputAllocator](FRawTransformerFunc Transformer, FAttributeSetMap* InputSetMap)
+					{
+						if (InputSetMap == nullptr)
+						{
+							// No maps found for the transformer type
+						}
+						else
+						{
+							if (Transformer == nullptr)
+							{
+								// No transformer was found for this value type, retain it as-is
+							}
+							else
+							{
+								// Transform our input in place
+								FAttributeSetMap* OutputSetMap = InputSetMap;
+
+								// We cannot join our weights because their value type differs
+								const FAttributeMappingKey MappingKey = InputSetMap->GetMappingKey().To(FloatValueType);
+								const FAttributeSetMap* PerValueWeightMap = PerValueWeights.FindSetMap(MappingKey);
+
+								FTransformSetMapFunc TypedTransformer = reinterpret_cast<FTransformSetMapFunc>(Transformer);
+								(*TypedTransformer)(InputSetMap, PerValueWeightMap, DefaultWeight, OutputSetMap, nullptr);
+							}
+						}
+					},
+					TransformerList->CreateSetMapTransformerIterator(),
+					// Our input is our output
+					Output.CreateSetMapIterator());
+			}
+
+			// Name/value maps
+			{
+				OuterJoinBy(
+					FAttributeValueTransformerListWithAttributeCollectionMapJoinOp(),
+					[DefaultWeight, OutputAllocator](FRawTransformerFunc Transformer, FAttributeMap* InputMap)
+					{
+						if (InputMap == nullptr)
+						{
+							// No maps found for the transformer type
+						}
+						else
+						{
+							if (Transformer == nullptr)
+							{
+								// No transformer was found for this value type, retain it as-is
+							}
+							else
+							{
+								// Transform our input in place
+								FAttributeMap* OutputMap = InputMap;
+
+								FTransformMapFunc TypedTransformer = reinterpret_cast<FTransformMapFunc>(Transformer);
+								(*TypedTransformer)(InputMap, DefaultWeight, OutputMap, nullptr);
+							}
+						}
+					},
+					TransformerList->CreateMapTransformerIterator(),
+					// Our input is our output
+					Output.CreateMapIterator());
+			}
+		}
+		else
+		{
+			// Our output collection isn't either of our inputs, reset it
+			Output.Reset(Input.GetNamedSet());
+			Output.SetValueSpace(Input.GetValueSpace());
+
+			// Set maps
+			OuterJoinBy(
+				FAttributeValueTransformerListWithAttributeCollectionSetMapJoinOp(),
+				[&PerValueWeights, DefaultWeight, FloatValueType, &Output](FRawTransformerFunc Transformer, const FAttributeSetMap* InputSetMap)
+				{
+					if (InputSetMap == nullptr)
+					{
+						// No maps found for the transformer type
+					}
+					else
+					{
+						if (Transformer == nullptr)
+						{
+							// No transformer was found for this value type, retain it as-is
+							Output.AppendSetMap(InputSetMap->Duplicate(Output.GetAllocator()));
+						}
+						else
+						{
+							// Transform our matching pair
+							FAttributeSetMap* OutputSetMap = nullptr;
+
+							// We cannot join our weights because their value type differs
+							const FAttributeMappingKey MappingKey = InputSetMap->GetMappingKey().To(FloatValueType);
+							const FAttributeSetMap* PerValueWeightMap = PerValueWeights.FindSetMap(MappingKey);
+
+							FTransformSetMapFunc TypedTransformer = reinterpret_cast<FTransformSetMapFunc>(Transformer);
+							(*TypedTransformer)(InputSetMap, PerValueWeightMap, DefaultWeight, OutputSetMap, Output.GetAllocator());
+
+							checkf(OutputSetMap != nullptr, TEXT("Transformer function is responsible for allocating an output buffer"));
+							checkf(OutputSetMap->GetAllocator() == Output.GetAllocator(), TEXT("Transformer function is responsible for allocating using the output allocator"));
+							Output.AppendSetMap(OutputSetMap);
+						}
+					}
+				},
+				TransformerList->CreateSetMapTransformerIterator(),
+				Input.CreateConstSetMapIterator());
+
+			// Name/value maps
+			OuterJoinBy(
+				FAttributeValueTransformerListWithAttributeCollectionMapJoinOp(),
+				[DefaultWeight, &Output](FRawTransformerFunc Transformer, const FAttributeMap* InputMap)
+				{
+					if (InputMap == nullptr)
+					{
+						// No maps found for the transformer type
+					}
+					else
+					{
+						if (Transformer == nullptr)
+						{
+							// No transformer was found for this value type, retain it as-is
+							Output.AppendMap(InputMap->Duplicate(Output.GetAllocator()));
+						}
+						else
+						{
+							// Transform our matching pair
+							FAttributeMap* OutputMap = nullptr;
+
+							FTransformMapFunc TypedTransformer = reinterpret_cast<FTransformMapFunc>(Transformer);
+							(*TypedTransformer)(InputMap, DefaultWeight, OutputMap, Output.GetAllocator());
+
+							checkf(OutputMap != nullptr, TEXT("Transformer function is responsible for allocating an output buffer"));
+							checkf(OutputMap->GetAllocator() == Output.GetAllocator(), TEXT("Transformer function is responsible for allocating using the output allocator"));
+							Output.AppendMap(OutputMap);
+						}
+					}
+				},
+				TransformerList->CreateMapTransformerIterator(),
+				Input.CreateConstMapIterator());
+		}
+	}
+
+	namespace Private
+	{
+		void FOverwrite_BoneTransformAttribute::TransformSetMap(const FAttributeSetMap* InputSetMap, const FAttributeSetMap* PerValueWeightMap, float DefaultWeight, FAttributeSetMap*& OutputSetMap, FReallocFun OutputAllocator)
+		{
+			if (OutputSetMap == nullptr)
+			{
+				const FAttributeSetMap::FConstructArgs Args(InputSetMap->GetTypedSet(), InputSetMap->GetValueType(), OutputAllocator);
+				OutputSetMap = MakeAttributeSetMap<FBoneTransformAnimationAttribute>(Args);
+			}
+
+			check(InputSetMap->GetTypedSet() == OutputSetMap->GetTypedSet());
+			check(InputSetMap->Num() == OutputSetMap->Num());
+
+			const TAttributeSetMap<FBoneTransformAnimationAttribute>* InputBoneTransforms = CastChecked<FBoneTransformAnimationAttribute>(InputSetMap);
+			TAttributeSetMap<FBoneTransformAnimationAttribute>* OutputBoneTransforms = CastChecked<FBoneTransformAnimationAttribute>(OutputSetMap);
+
+			const FBoneTransformAnimationAttribute* InputTransformPtr = InputBoneTransforms->GetData();
+			FBoneTransformAnimationAttribute* OutputTransformPtr = OutputBoneTransforms->GetData();
+			const FBoneTransformAnimationAttribute* InputTransformPtrEnd = InputTransformPtr + InputBoneTransforms->Num();
+
+			if (const TAttributeSetMap<FFloatAnimationAttribute>* PerValueWeights = Cast<FFloatAnimationAttribute>(PerValueWeightMap))
+			{
+				check(InputSetMap->GetTypedSet() == PerValueWeights->GetTypedSet());
+				check(InputSetMap->Num() == PerValueWeights->Num());
+
+				const FFloatAnimationAttribute* PerValueWeightsPtr = PerValueWeights->GetData();
+
+				while (InputTransformPtr < InputTransformPtrEnd)
+				{
+					OutputTransformPtr->Value = InputTransformPtr->Value * ScalarRegister(PerValueWeightsPtr->Value);
+
+					InputTransformPtr++;
+					PerValueWeightsPtr++;
+					OutputTransformPtr++;
+				}
+			}
+			else
+			{
+				const ScalarRegister VDefaultWeight(DefaultWeight);
+
+				while (InputTransformPtr < InputTransformPtrEnd)
+				{
+					OutputTransformPtr->Value = InputTransformPtr->Value * VDefaultWeight;
+
+					InputTransformPtr++;
+					OutputTransformPtr++;
+				}
+			}
+		}
+
+		void FOverwrite_BoneTransformAttribute::TransformMap(const FAttributeMap* InputMap, float Weight, FAttributeMap*& OutputMap, FReallocFun OutputAllocator)
+		{
+			const bool bAllocateOutput = OutputMap == nullptr;
+			if (bAllocateOutput)
+			{
+				OutputMap = MakeAttributeMap<FBoneTransformAnimationAttribute>(OutputAllocator);
+			}
+
+			const TAttributeMap<FBoneTransformAnimationAttribute>* InputBoneTransforms = CastChecked<FBoneTransformAnimationAttribute>(InputMap);
+			TAttributeMap<FBoneTransformAnimationAttribute>* OutputBoneTransforms = CastChecked<FBoneTransformAnimationAttribute>(OutputMap);
+
+			const ScalarRegister VWeight(Weight);
+
+			if (bAllocateOutput)
+			{
+				// We allocated a new output, append to it
+				for (auto It = InputBoneTransforms->CreateConstIterator(); It; ++It)
+				{
+					FBoneTransformAnimationAttribute Output{ It.GetValue()->Value * VWeight };
+					OutputBoneTransforms->Append(It.GetName(), Output);
+				}
+			}
+			else
+			{
+				// We overwrite in place
+				check(InputBoneTransforms == OutputBoneTransforms);
+
+				for (auto It = OutputBoneTransforms->CreateIterator(); It; ++It)
+				{
+					It.GetValue()->Value *= VWeight;
+				}
+			}
+		}
+
+		void FOverwrite_FloatAttribute::TransformSetMap(const FAttributeSetMap* InputSetMap, const FAttributeSetMap* PerValueWeightMap, float DefaultWeight, FAttributeSetMap*& OutputSetMap, FReallocFun OutputAllocator)
+		{
+			if (OutputSetMap == nullptr)
+			{
+				const FAttributeSetMap::FConstructArgs Args(InputSetMap->GetTypedSet(), InputSetMap->GetValueType(), OutputAllocator);
+				OutputSetMap = MakeAttributeSetMap<FFloatAnimationAttribute>(Args);
+			}
+
+			check(InputSetMap->GetTypedSet() == OutputSetMap->GetTypedSet());
+			check(InputSetMap->Num() == OutputSetMap->Num());
+
+			const TAttributeSetMap<FFloatAnimationAttribute>* InputFloats = CastChecked<FFloatAnimationAttribute>(InputSetMap);
+			TAttributeSetMap<FFloatAnimationAttribute>* OutputFloats = CastChecked<FFloatAnimationAttribute>(OutputSetMap);
+
+			const FFloatAnimationAttribute* FloatsPtr = InputFloats->GetData();
+			FFloatAnimationAttribute* OutputFloatsPtr = OutputFloats->GetData();
+			const FFloatAnimationAttribute* FloatsPtrEnd = FloatsPtr + InputFloats->Num();
+
+			if (const TAttributeSetMap<FFloatAnimationAttribute>* PerValueWeights = Cast<FFloatAnimationAttribute>(PerValueWeightMap))
+			{
+				check(InputSetMap->GetTypedSet() == PerValueWeights->GetTypedSet());
+				check(InputSetMap->Num() == PerValueWeights->Num());
+
+				const FFloatAnimationAttribute* PerValueWeightsPtr = PerValueWeights->GetData();
+
+				while (FloatsPtr < FloatsPtrEnd)
+				{
+					OutputFloatsPtr->Value = FloatsPtr->Value * PerValueWeightsPtr->Value;
+
+					FloatsPtr++;
+					PerValueWeightsPtr++;
+					OutputFloatsPtr++;
+				}
+			}
+			else
+			{
+				while (FloatsPtr < FloatsPtrEnd)
+				{
+					OutputFloatsPtr->Value = FloatsPtr->Value * DefaultWeight;
+
+					FloatsPtr++;
+					OutputFloatsPtr++;
+				}
+			}
+		}
+
+		void FOverwrite_FloatAttribute::TransformMap(const FAttributeMap* InputMap, float Weight, FAttributeMap*& OutputMap, FReallocFun OutputAllocator)
+		{
+			const bool bAllocateOutput = OutputMap == nullptr;
+			if (bAllocateOutput)
+			{
+				OutputMap = MakeAttributeMap<FFloatAnimationAttribute>(OutputAllocator);
+			}
+
+			const TAttributeMap<FFloatAnimationAttribute>* InputFloats = CastChecked<FFloatAnimationAttribute>(InputMap);
+			TAttributeMap<FFloatAnimationAttribute>* OutputFloats = CastChecked<FFloatAnimationAttribute>(OutputMap);
+
+			if (bAllocateOutput)
+			{
+				// We allocated a new output, append to it
+				for (auto It = InputFloats->CreateConstIterator(); It; ++It)
+				{
+					FFloatAnimationAttribute Output{ It.GetValue()->Value * Weight };
+					OutputFloats->Append(It.GetName(), Output);
+				}
+			}
+			else
+			{
+				// We overwrite in place
+				check(InputFloats == OutputFloats);
+
+				for (auto It = OutputFloats->CreateIterator(); It; ++It)
+				{
+					It.GetValue()->Value *= Weight;
+				}
+			}
+		}
+	}
+}
diff --git a/Engine/Plugins/Experimental/UAF/UAF/Source/UAF/Private/UAF/ValueRuntime/Transformers/Sanitize.cpp b/Engine/Plugins/Experimental/UAF/UAF/Source/UAF/Private/UAF/ValueRuntime/Transformers/Sanitize.cpp
new file mode 100644
--- /dev/null
+++ b/Engine/Plugins/Experimental/UAF/UAF/Source/UAF/Private/UAF/ValueRuntime/Transformers/Sanitize.cpp
@@ -0,0 +1,156 @@
+// Copyright Epic Games, Inc. All Rights Reserved.
+
+#include "UAF/ValueRuntime/Transformers/Sanitize.h"
+
+#include "UAF/ValueRuntime/IteratorUtils.h"
+
+namespace UE::UAF::Transformers
+{
+	FName FSanitize::TransformerName = TEXT("UAF::Transformers::Sanitize");
+
+	void FSanitize::Apply(const FAttributeValueTransformerMapPtr& TransformerMap, FAttributeCollection& Collection)
+	{
+		const FAttributeValueTransformerList* TransformerList = TransformerMap->Find(TransformerName);
+		if (TransformerList == nullptr)
+		{
+			// No transformer implementations have been registered
+			return;
+		}
+
+		// Set maps
+		InnerJoinBy(
+			FAttributeValueTransformerListWithAttributeCollectionSetMapJoinOp(),
+			[](FRawTransformerFunc Transformer, FAttributeSetMap* SetMap)
+			{
+				FTransformSetMapFunc TypedTransformer = reinterpret_cast<FTransformSetMapFunc>(Transformer);
+				(*TypedTransformer)(SetMap);
+			},
+			TransformerList->CreateSetMapTransformerIterator(),
+			Collection.CreateSetMapIterator());
+
+		// Name/value maps
+		InnerJoinBy(
+			FAttributeValueTransformerListWithAttributeCollectionMapJoinOp(),
+			[](FRawTransformerFunc Transformer, FAttributeMap* Map)
+			{
+				FTransformMapFunc TypedTransformer = reinterpret_cast<FTransformMapFunc>(Transformer);
+				(*TypedTransformer)(Map);
+			},
+			TransformerList->CreateMapTransformerIterator(),
+			Collection.CreateMapIterator());
+	}
+
+	namespace Private
+	{
+#if UE_UAF_VALIDATE_MAPPED_VALUES
+		// Don't inline this function to keep the stack usage down
+		static FORCENOINLINE void IsBoneTransformSetMapValid(const TAttributeSetMap<FBoneTransformAnimationAttribute>* BoneTransforms)
+		{
+			const FBoneTransformAnimationAttribute* TransformPtr = BoneTransforms->GetData();
+			const FBoneTransformAnimationAttribute* EndTransformPtr = TransformPtr + BoneTransforms->Num();
+
+			while (TransformPtr < EndTransformPtr)
+			{
+				if (TransformPtr->Value.ContainsNaN())
+				{
+					ensureMsgf(false, TEXT("Bone (%s) contains NaN with Value:[%s]"),
+						*BoneTransforms->GetTypedSet()->GetName(FAttributeSetIndex(TransformPtr - BoneTransforms->GetData())).ToString(),
+						*TransformPtr->Value.ToString());
+				}
+				else if (!TransformPtr->Value.IsRotationNormalized())
+				{
+					ensureMsgf(false, TEXT("Bone (%s) rotation not normalized with Rotation:[%s]"),
+						*BoneTransforms->GetTypedSet()->GetName(FAttributeSetIndex(TransformPtr - BoneTransforms->GetData())).ToString(),
+						*TransformPtr->Value.GetRotation().ToString());
+				}
+
+				TransformPtr++;
+			}
+		}
+
+		// Don't inline this function to keep the stack usage down
+		static FORCENOINLINE void IsBoneTransformMapValid(const TAttributeMap<FBoneTransformAnimationAttribute>* BoneTransforms)
+		{
+			for (auto It = BoneTransforms->CreateConstIterator(); It; ++It)
+			{
+				if (It.GetValue()->Value.ContainsNaN())
+				{
+					ensureMsgf(false, TEXT("Transform attribute (%s) contains NaN with Value:[%s]"),
+						*It.GetName().ToString(),
+						*It.GetValue()->Value.ToString());
+				}
+				else if (!It.GetValue()->Value.IsRotationNormalized())
+				{
+					ensureMsgf(false, TEXT("Transform attribute (%s) rotation not normalized with Rotation:[%s]"),
+						*It.GetName().ToString(),
+						*It.GetValue()->Value.GetRotation().ToString());
+				}
+			}
+		}
+#endif
+
+		void FSanitize_BoneTransformAttribute::TransformSetMap(FAttributeSetMap* SetMap)
+		{
+			TAttributeSetMap<FBoneTransformAnimationAttribute>* BoneTransforms = CastChecked<FBoneTransformAnimationAttribute>(SetMap);
+
+			FBoneTransformAnimationAttribute* TransformPtr = BoneTransforms->GetData();
+			const FBoneTransformAnimationAttribute* TransformPtrEnd = TransformPtr + BoneTransforms->Num();
+
+			while (TransformPtr < TransformPtrEnd)
+			{
+				TransformPtr->Value.NormalizeRotation();
+
+				TransformPtr++;
+			}
+
+#if UE_UAF_VALIDATE_MAPPED_VALUES
+			IsBoneTransformSetMapValid(BoneTransforms);
+#endif
+		}
+
+		void FSanitize_BoneTransformAttribute::TransformMap(FAttributeMap* Map)
+		{
+			TAttributeMap<FBoneTransformAnimationAttribute>* BoneTransforms = CastChecked<FBoneTransformAnimationAttribute>(Map);
+
+			for (auto It = BoneTransforms->CreateIterator(); It; ++It)
+			{
+				It.GetValue()->Value.NormalizeRotation();
+			}
+
+#if UE_UAF_VALIDATE_MAPPED_VALUES
+			IsBoneTransformMapValid(BoneTransforms);
+#endif
+		}
+
+#if UE_UAF_VALIDATE_MAPPED_VALUES
+		void FSanitize_FloatAttribute::TransformSetMap(FAttributeSetMap* SetMap)
+		{
+			const TAttributeSetMap<FFloatAnimationAttribute>* Floats = CastChecked<FFloatAnimationAttribute>(SetMap);
+
+			const FFloatAnimationAttribute* FloatsPtr = Floats->GetData();
+			const FFloatAnimationAttribute* FloatsPtrEnd = FloatsPtr + Floats->Num();
+
+			while (FloatsPtr < FloatsPtrEnd)
+			{
+				ensureMsgf(!FMath::IsNaN(FloatsPtr->Value), TEXT("Float attribute (%s) contains NaN with Value:[%3.3f]"),
+					*Floats->GetTypedSet()->GetName(FAttributeSetIndex(FloatsPtr - Floats->GetData())).ToString(),
+					FloatsPtr->Value);
+
+				FloatsPtr++;
+			}
+		}
+
+		void FSanitize_FloatAttribute::TransformMap(FAttributeMap* Map)
+		{
+			const TAttributeMap<FFloatAnimationAttribute>* Floats = CastChecked<FFloatAnimationAttribute>(Map);
+
+			for (auto It = Floats->CreateConstIterator(); It; ++It)
+			{
+				ensureMsgf(!FMath::IsNaN(It.GetValue()->Value), TEXT("Float attribute (%s) contains NaN with Value:[%3.3f]"),
+					*It.GetName().ToString(),
+					It.GetValue()->Value);
+			}
+		}
+#endif
+	}
+}
diff --git a/Engine/Plugins/Experimental/UAF/UAF/Source/UAF/Private/UAF/ValueRuntime/Transformers/Transformer.cpp b/Engine/Plugins/Experimental/UAF/UAF/Source/UAF/Private/UAF/ValueRuntime/Transformers/Transformer.cpp
new file mode 100644
--- /dev/null
+++ b/Engine/Plugins/Experimental/UAF/UAF/Source/UAF/Private/UAF/ValueRuntime/Transformers/Transformer.cpp
@@ -0,0 +1,56 @@
+// Copyright Epic Games, Inc. All Rights Reserved.
+
+#include "UAF/ValueRuntime/Transformers/Transformer.h"
+
+namespace UE::UAF
+{
+	FAttributeValueTransformerList::FAttributeValueTransformerList() = default;
+
+	FAttributeValueTransformerList::FAttributeValueTransformerList(FName InTransformerName)
+		: TransformerName(InTransformerName)
+	{
+	}
+
+	FName FAttributeValueTransformerList::GetTransformerName() const
+	{
+		return TransformerName;
+	}
+
+	bool FAttributeValueTransformerList::AddSetMapTransformer(UScriptStruct* ValueType, FRawTransformerFunc TransformerFunc)
+	{
+		const int32 EntryIndex = Algo::LowerBound(MakeArrayView(ValueTypes.GetData(), NumSetMapTransformers()), ValueType);
+		if (EntryIndex < ValueTypes.Num() && ValueTypes[EntryIndex] == ValueType)
+		{
+			// A transformer for this value type has already been registered
+			return false;
+		}
+
+		ValueTypes.Insert(ValueType, EntryIndex);
+		Transformers.Insert(TransformerFunc, EntryIndex);
+		SetMapTransformerCount++;
+
+		return true;
+	}
+
+	bool FAttributeValueTransformerList::AddMapTransformer(UScriptStruct* ValueType, FRawTransformerFunc TransformerFunc)
+	{
+		const int32 NumSetMapEntries = NumSetMapTransformers();
+		const int32 NumMapEntries = NumMapTransformers();
+
+		const int32 PerContainerEntryIndex = Algo::LowerBound(MakeArrayView(ValueTypes.GetData() + NumSetMapEntries, NumMapEntries), ValueType);
+
+		// Fixup our index
+		const int32 EntryIndex = PerContainerEntryIndex + NumSetMapEntries;
+
+		if (EntryIndex < ValueTypes.Num() && ValueTypes[EntryIndex] == ValueType)
+		{
+			// A transformer for this value type has already been registered
+			return false;
+		}
+
+		ValueTypes.Insert(ValueType, EntryIndex);
+		Transformers.Insert(TransformerFunc, EntryIndex);
+
+		return true;
+	}
+}
diff --git a/Engine/Plugins/Experimental/UAF/UAF/Source/UAF/Private/UAF/ValueRuntime/ValueRuntimeRegistry.cpp b/Engine/Plugins/Experimental/UAF/UAF/Source/UAF/Private/UAF/ValueRuntime/ValueRuntimeRegistry.cpp
new file mode 100644
--- /dev/null
+++ b/Engine/Plugins/Experimental/UAF/UAF/Source/UAF/Private/UAF/ValueRuntime/ValueRuntimeRegistry.cpp
@@ -0,0 +1,146 @@
+// Copyright Epic Games, Inc. All Rights Reserved.
+
+#include "UAF/ValueRuntime/ValueRuntimeRegistry.h"
+
+#include "Algo/BinarySearch.h"
+
+namespace UE::UAF
+{
+	namespace Private
+	{
+		static FValueRuntimeRegistry* GValueRuntimeRegistry = nullptr;
+	}
+
+	void FValueRuntimeRegistry::Init()
+	{
+		if (Private::GValueRuntimeRegistry == nullptr)
+		{
+			Private::GValueRuntimeRegistry = new FValueRuntimeRegistry();
+		}
+	}
+
+	void FValueRuntimeRegistry::Destroy()
+	{
+		if (Private::GValueRuntimeRegistry != nullptr)
+		{
+			delete Private::GValueRuntimeRegistry;
+			Private::GValueRuntimeRegistry = nullptr;
+		}
+	}
+
+	FValueRuntimeRegistry& FValueRuntimeRegistry::Get()
+	{
+		checkf(Private::GValueRuntimeRegistry != nullptr, TEXT("Attribute Runtime Registry is not instanced. It is only valid to access this while the UAF module is loaded."));
+		return *Private::GValueRuntimeRegistry;
+	}
+
+	FValueRuntimeRegistry::FValueRuntimeRegistry() = default;
+
+	FAttributeSetMap* FValueRuntimeRegistry::ConstructSetMap(const FAttributeSetMap::FConstructArgs& Args) const
+	{
+		FReadScopeLock Lock(SetMapInitializerLock);
+		if (const SetMapConstructorFun* ConstructorFun = SetMapInitializers.Find(Args.ValueType))
+		{
+			return (*ConstructorFun)(Args);
+		}
+
+		// Value type has not been registered
+		return nullptr;
+	}
+
+	FAttributeMap* FValueRuntimeRegistry::ConstructMap(UScriptStruct* ValueType, FReallocFun ReallocFun) const
+	{
+		FReadScopeLock Lock(SetMapInitializerLock);
+		if (const MapConstructorFun* ConstructorFun = MapInitializers.Find(ValueType))
+		{
+			return (*ConstructorFun)(ReallocFun);
+		}
+
+		// Value type has not been registered
+		return nullptr;
+	}
+
+	FAttributeValueTransformerMapPtr FValueRuntimeRegistry::GetTransformerMap() const
+	{
+		FReadScopeLock Lock(TransformerMapLock);
+		return TransformerMap;
+	}
+
+	void FValueRuntimeRegistry::RegisterSetMapInitializer(UScriptStruct* ValueType, SetMapConstructorFun Initializer)
+	{
+		FWriteScopeLock Lock(SetMapInitializerLock);
+		SetMapInitializers.Add(ValueType, Initializer);
+	}
+
+	void FValueRuntimeRegistry::RegisterMapInitializer(UScriptStruct* ValueType, MapConstructorFun Initializer)
+	{
+		FWriteScopeLock Lock(MapInitializerLock);
+		MapInitializers.Add(ValueType, Initializer);
+	}
+
+	bool FValueRuntimeRegistry::RegisterAttributeValueSetMapTransformer(FName TransformerName, UScriptStruct* ValueType, FRawTransformerFunc TransformerFunc)
+	{
+		FWriteScopeLock Lock(TransformerMapLock);
+
+		// If someone is already holding a reference to our named operators, we cannot mutate
+		// When this occurs, we make a copy and swap it, keeping the old one alive until it is no longer needed
+		FAttributeValueTransformerMapPtr MutableTransformerMap = TransformerMap;
+		if (!MutableTransformerMap || MutableTransformerMap.GetSharedReferenceCount() != 2)
+		{
+			MutableTransformerMap = MakeShared<typename FAttributeValueTransformerMapPtr::ElementType>();
+		}
+
+		const uint32 TransformerNameHash = GetTypeHash(TransformerName);
+		FAttributeValueTransformerList* TransformerList = MutableTransformerMap->FindByHash(TransformerNameHash, TransformerName);
+		if (TransformerList == nullptr)
+		{
+			TransformerList = &MutableTransformerMap->AddByHash(TransformerNameHash, TransformerName, FAttributeValueTransformerList(TransformerName));
+		}
+
+		if (!TransformerList->AddSetMapTransformer(ValueType, TransformerFunc))
+		{
+			// We failed to add this transformer
+			return false;
+		}
+
+		if (TransformerMap != MutableTransformerMap)
+		{
+			TransformerMap = MutableTransformerMap;
+		}
+
+		return true;
+	}
+
+	bool FValueRuntimeRegistry::RegisterAttributeValueMapTransformer(FName TransformerName, UScriptStruct* ValueType, FRawTransformerFunc TransformerFunc)
+	{
+		FWriteScopeLock Lock(TransformerMapLock);
+
+		// If someone is already holding a reference to our named operators, we cannot mutate
+		// When this occurs, we make a copy and swap it, keeping the old one alive until it is no longer needed
+		FAttributeValueTransformerMapPtr MutableTransformerMap = TransformerMap;
+		if (!MutableTransformerMap || MutableTransformerMap.GetSharedReferenceCount() != 2)
+		{
+			MutableTransformerMap = MakeShared<typename FAttributeValueTransformerMapPtr::ElementType>();
+		}
+
+		const uint32 TransformerNameHash = GetTypeHash(TransformerName);
+		FAttributeValueTransformerList* TransformerList = MutableTransformerMap->FindByHash(TransformerNameHash, TransformerName);
+		if (TransformerList == nullptr)
+		{
+			TransformerList = &MutableTransformerMap->AddByHash(TransformerNameHash, TransformerName, FAttributeValueTransformerList(TransformerName));
+		}
+
+		if (!TransformerList->AddMapTransformer(ValueType, TransformerFunc))
+		{
+			// We failed to add this transformer
+			return false;
+		}
+
+		if (TransformerMap != MutableTransformerMap)
+		{
+			TransformerMap = MutableTransformerMap;
+		}
+
+		return true;
+	}
+}
