Move and rename UAF ValueRuntime related files

diff --git a/Engine/Plugins/Experimental/UAF/UAF/Source/UAF/Public/UAF/AbstractSkeleton/AbstractSkeletonSetBinding.h b/Engine/Plugins/Experimental/UAF/UAF/Source/UAF/Public/UAF/AbstractSkeleton/AbstractSkeletonSetBinding.h
--- a/Engine/Plugins/Experimental/UAF/UAF/Source/UAF/Public/UAF/AbstractSkeleton/AbstractSkeletonSetBinding.h
+++ b/Engine/Plugins/Experimental/UAF/UAF/Source/UAF/Public/UAF/AbstractSkeleton/AbstractSkeletonSetBinding.h
@@ -6,8 +6,8 @@
 #include "Animation/AnimData/AttributeIdentifier.h"
 #include "Animation/BuiltInAttributeTypes.h"
 #include "UAF/Attributes/AttributeBindingIndex.h"
-#include "UAF/Attributes/AttributeCollection.h"
 #include "UAF/Attributes/AttributeBindingData.h"
+#include "UAF/ValueRuntime/AttributeCollection.h"
 
 #include "AbstractSkeletonSetBinding.generated.h"
 
diff --git a/Engine/Plugins/Experimental/UAF/UAF/Source/UAF/Public/UAF/ValueRuntime/AttributeCollection.h b/Engine/Plugins/Experimental/UAF/UAF/Source/UAF/Public/UAF/ValueRuntime/AttributeCollection.h
new file mode 100644
--- /dev/null
+++ b/Engine/Plugins/Experimental/UAF/UAF/Source/UAF/Public/UAF/ValueRuntime/AttributeCollection.h
@@ -0,0 +1,652 @@
+// Copyright Epic Games, Inc. All Rights Reserved.
+
+#pragma once
+
+#include "CoreMinimal.h"
+#include "Misc/MemStack.h"
+#include "UAF/Attributes/AttributeNamedSet.h"
+#include "UAF/ValueRuntime/AttributeMap.h"
+#include "UAF/ValueRuntime/AttributeMappingKey.h"
+#include "UAF/ValueRuntime/AttributeSetMap.h"
+#include "UAF/ValueRuntime/IndirectAllocator.h"
+#include "UAF/ValueRuntime/ValueSpace.h"
+
+#include "AttributeCollection.generated.h"
+
+#define UE_API UAF_API
+
+namespace UE::UAF
+{
+	// A iterator that iterates over all set maps within an attribute collection
+	template<class ContainerType>
+	class TAttributeCollectionSetMapIterator
+	{
+	public:
+		static constexpr bool bIsContainerConst = std::is_const_v<ContainerType>;
+
+		using EntryPtrType = typename std::conditional<bIsContainerConst, const typename ContainerType::FSetMapEntry*, typename ContainerType::FSetMapEntry*>::type;
+		using SetMapType = typename std::conditional<bIsContainerConst, const FAttributeSetMap*, FAttributeSetMap*>::type;
+
+		// Creates an empty iterator
+		[[nodiscard]] TAttributeCollectionSetMapIterator() = default;
+
+		// Creates an iterator over the set map entries within the specified collection
+		[[nodiscard]] TAttributeCollectionSetMapIterator(ContainerType& Collection);
+
+		// Increments and moves the iterator to the next set map entry
+		TAttributeCollectionSetMapIterator& operator++();
+
+		// Returns whether or not the iterator still contains values
+		[[nodiscard]] explicit operator bool() const;
+
+		// Returns the mapping key of the current set map entry
+		[[nodiscard]] const FAttributeMappingKey& GetKey() const;
+
+		// Returns the attribute type of the current set map entry
+		[[nodiscard]] UScriptStruct* GetAttributeType() const;
+
+		// Returns the value type of the current set map entry
+		[[nodiscard]] UScriptStruct* GetValueType() const;
+
+		// Returns the set map from the current entry
+		[[nodiscard]] SetMapType GetMap() const;
+
+	private:
+		// Current set map entry
+		EntryPtrType EntryPtr = nullptr;
+
+		// One past the last set map entry
+		EntryPtrType EntryEndPtr = nullptr;
+	};
+
+	// A iterator that iterates over all name/value maps within an attribute collection
+	template<class ContainerType>
+	class TAttributeCollectionMapIterator
+	{
+	public:
+		static constexpr bool bIsContainerConst = std::is_const_v<ContainerType>;
+
+		using EntryPtrType = typename std::conditional<bIsContainerConst, const typename ContainerType::FMapEntry*, typename ContainerType::FMapEntry*>::type;
+		using MapType = typename std::conditional<bIsContainerConst, const FAttributeMap*, FAttributeMap*>::type;
+
+		// Creates an empty iterator
+		[[nodiscard]] TAttributeCollectionMapIterator() = default;
+
+		// Creates an iterator over the name/value map entries within the specified collection
+		[[nodiscard]] TAttributeCollectionMapIterator(ContainerType& Collection);
+
+		// Increments and moves the iterator to the next name/value map entry
+		TAttributeCollectionMapIterator& operator++();
+
+		// Returns whether or not the iterator still contains values
+		[[nodiscard]] explicit operator bool() const;
+
+		// Returns the value type of the current name/value map entry
+		[[nodiscard]] UScriptStruct* GetValueType() const;
+
+		// Returns the name/value map from the current entry
+		[[nodiscard]] MapType GetMap() const;
+
+	private:
+		// Current name/value map entry
+		EntryPtrType EntryPtr = nullptr;
+
+		// One past the last name/value map entry
+		EntryPtrType EntryEndPtr = nullptr;
+	};
+
+	/*
+	 * Attribute Collection
+	 *
+	 * This struct holds attributes produced by UAF. They come in two flavors: static (set map) and dynamic (map).
+	 * Static attributes are fast and defined up front in the skeleton set binding asset.
+	 * Dynamic attributes are slower but produced on demand during evaluation.
+	 */
+	class FAttributeCollection
+	{
+	private:
+		struct FSetMapEntry
+		{
+			// The mapping key that identifies this entry
+			FAttributeMappingKey Key;
+
+			// The set map containing our values
+			FAttributeSetMap* Map = nullptr;
+		};
+
+		struct FMapEntry
+		{
+			// The value type (key) that identifies this entry
+			UScriptStruct* ValueType = nullptr;
+
+			// The map containing our values
+			FAttributeMap* Map = nullptr;
+		};
+
+	public:
+		UE_API FAttributeCollection(const FAttributeNamedSetPtr& NamedSet, FReallocFun InReallocFun);
+		UE_API FAttributeCollection(const FAttributeCollection& Other);
+		UE_API FAttributeCollection(FAttributeCollection&& Other);
+		UE_API ~FAttributeCollection();
+
+		UE_API FAttributeCollection& operator=(const FAttributeCollection& Other);
+		UE_API FAttributeCollection& operator=(FAttributeCollection&& Other);
+
+		// Returns whether or not we are empty
+		[[nodiscard]] UE_API bool IsEmpty() const;
+
+		// Returns the named set this collection maps to
+		[[nodiscard]] UE_API const FAttributeNamedSetPtr& GetNamedSet() const;
+
+		// Returns the space our values live in
+		[[nodiscard]] UE_API FValueSpace GetValueSpace() const;
+
+		// Sets the space our values live in
+		// Note that care must be taken when changing the value space as all values within the collection must be migrated accordingly
+		UE_API void SetValueSpace(FValueSpace ValueSpace);
+
+		// Sets the space our values line in and default initializes this collection
+		// For local/component space, we initialize with the reference pose
+		// For additive space, we initialize with the additive identity
+		// Other spaces initialize to their identity
+		// Note that only bound values are default initialized
+		UE_API void InitWithValueSpace(FValueSpace ValueSpace);
+
+		// Returns whether or not the space our values live in is additive
+		[[nodiscard]] UE_API bool IsAdditive() const;
+
+		// Resets the collection with the specified named set
+		UE_API void Reset(const FAttributeNamedSetPtr& NamedSet);
+
+		// Clears the collection of any content
+		UE_API void Empty();
+
+		// Returns the allocator function this collection was initialized with
+		[[nodiscard]] UE_API FReallocFun GetAllocator() const;
+
+		// Returns the number of set maps contained within
+		[[nodiscard]] UE_API int32 NumSetMaps() const;
+
+		// Constructs and adds a new mapping of ValueType to a set of AttributeType to this collection.
+		// Returns nullptr if we fail to add an entry (e.g. duplicate mapping)
+		template <class ValueType, typename... TArgs> [[nodiscard]] TAttributeSetMap<ValueType>* AddSetMap(const FAttributeMappingKey& MappingKey, TArgs... Args);
+
+		// Constructs and adds a new set mapping to this collection.
+		// Returns nullptr if we fail to add an entry (e.g. duplicate mapping)
+		[[nodiscard]] UE_API FAttributeSetMap* AddSetMap(const FAttributeMappingKey& MappingKey);
+
+		// Adds a new set mapping to this collection.
+		// Returns false if we fail to add an entry (e.g. duplicate mapping, mismatched allocator)
+		UE_API bool AddSetMap(FAttributeSetMap* Map);
+
+		// Appends a new set mapping to this collection (the input must belong at the end of the sorted set maps).
+		// Returns false if we fail to add an entry (e.g. duplicate mapping, mismatched allocator, not sorted)
+		UE_API bool AppendSetMap(FAttributeSetMap* Map);
+
+		// Removes the specified value set mapping from this collection
+		// Returns whether or not the entry was successfully removed
+		UE_API bool RemoveSetMap(FAttributeSetMap* Map);
+
+		// Removes the specified set mapping from this collection
+		// Returns whether or not the entry was successfully removed
+		UE_API bool RemoveSetMap(const FAttributeMappingKey& MappingKey);
+
+		// Finds the mapping of ValueType to a set of AttributeType in this collection
+		// Returns nullptr if this collection does not contain such a mapping
+		template <class ValueType> [[nodiscard]] TAttributeSetMap<ValueType>* FindSetMap(const FAttributeMappingKey& MappingKey);
+
+		// Finds the mapping of ValueType to a set of AttributeType in this collection
+		// Returns nullptr if this collection does not contain such a mapping
+		template <class ValueType> [[nodiscard]] const TAttributeSetMap<ValueType>* FindSetMap(const FAttributeMappingKey& MappingKey) const;
+
+		// Finds the mapping of ValueType to a set of AttributeType in this collection
+		// Returns nullptr if this collection does not contain such a mapping
+		[[nodiscard]] UE_API FAttributeSetMap* FindSetMap(const FAttributeMappingKey& MappingKey);
+
+		// Finds the mapping of ValueType to a set of AttributeType in this collection
+		// Returns nullptr if this collection does not contain such a mapping
+		[[nodiscard]] UE_API const FAttributeSetMap* FindSetMap(const FAttributeMappingKey& MappingKey) const;
+
+		// Finds or adds a mapping of ValueType to a set of AttributeType in this collection
+		// Returns nullptr if we fail to add an entry (e.g. duplicate mapping)
+		template <class ValueType, typename... TArgs> [[nodiscard]] TAttributeSetMap<ValueType>* FindOrAddSetMap(const FAttributeMappingKey& MappingKey, TArgs... Args);
+
+		// Finds or adds a mapping of ValueType to a set of AttributeType in this collection
+		// Returns nullptr if we fail to add an entry (e.g. duplicate mapping)
+		[[nodiscard]] UE_API FAttributeSetMap* FindOrAddSetMap(const FAttributeMappingKey& MappingKey);
+
+		// Iterators
+		using FSetMapIterator = TAttributeCollectionSetMapIterator<FAttributeCollection>;
+		using FConstSetMapIterator = TAttributeCollectionSetMapIterator<const FAttributeCollection>;
+
+		// Returns an iterator over all set maps contained within, sorted by their key (value type first, then attribute type)
+		[[nodiscard]] FSetMapIterator CreateSetMapIterator();
+		[[nodiscard]] FConstSetMapIterator CreateConstSetMapIterator() const;
+
+		// Returns the number of name/value maps contained within
+		[[nodiscard]] UE_API int32 NumMaps() const;
+
+		// Constructs and adds a new name/value map to this collection.
+		// Returns nullptr if we fail to add an entry (e.g. duplicate mapping)
+		template <class ValueType> [[nodiscard]] TAttributeMap<ValueType>* AddMap();
+
+		// Constructs and adds a new name/value mapping to this collection.
+		// Returns nullptr if we fail to add an entry (e.g. duplicate mapping)
+		[[nodiscard]] UE_API FAttributeMap* AddMap(UScriptStruct* ValueType);
+
+		// Adds a new name/value mapping to this collection.
+		// Returns false if we fail to add an entry (e.g. duplicate mapping, mismatched allocator)
+		UE_API bool AddMap(FAttributeMap* Map);
+
+		// Appends a new name/value mapping to this collection (the input must belong at the end of the sorted maps).
+		// Returns false if we fail to add an entry (e.g. duplicate mapping, mismatched allocator, not sorted)
+		UE_API bool AppendMap(FAttributeMap* Map);
+
+		// Removes the specified name/value mapping from this collection
+		// Returns whether or not the entry was successfully removed
+		UE_API bool RemoveMap(FAttributeMap* Map);
+
+		// Removes the specified name/value mapping from this collection
+		// Returns whether or not the entry was successfully removed
+		UE_API bool RemoveMap(UScriptStruct* ValueType);
+
+		// Finds the name/value mapping of the specified ValueType in this collection
+		// Returns nullptr if this collection does not contain such a mapping
+		template <class ValueType> [[nodiscard]] TAttributeMap<ValueType>* FindMap();
+
+		// Finds the name/value mapping of the specified ValueType in this collection
+		// Returns nullptr if this collection does not contain such a mapping
+		template <class ValueType> [[nodiscard]] const TAttributeMap<ValueType>* FindMap() const;
+
+		// Finds the name/value mapping of the specified ValueType in this collection
+		// Returns nullptr if this collection does not contain such a mapping
+		[[nodiscard]] UE_API FAttributeMap* FindMap(UScriptStruct* ValueType);
+
+		// Finds the name/value mapping of the specified ValueType in this collection
+		// Returns nullptr if this collection does not contain such a mapping
+		[[nodiscard]] UE_API const FAttributeMap* FindMap(UScriptStruct* ValueType) const;
+
+		// Finds or adds the name/value mapping of the specified ValueType in this collection
+		// Returns nullptr if we fail to add an entry (e.g. duplicate mapping)
+		template <class ValueType> [[nodiscard]] TAttributeMap<ValueType>* FindOrAddMap();
+
+		// Finds or adds the name/value mapping of the specified ValueType in this collection
+		// Returns nullptr if we fail to add an entry (e.g. duplicate mapping)
+		[[nodiscard]] UE_API FAttributeMap* FindOrAddMap(UScriptStruct* ValueType);
+
+		// Iterators
+		using FMapIterator = TAttributeCollectionMapIterator<FAttributeCollection>;
+		using FConstMapIterator = TAttributeCollectionMapIterator<const FAttributeCollection>;
+
+		// Returns an iterator over all name/value maps contained within, sorted by their key (value type)
+		[[nodiscard]] FMapIterator CreateMapIterator();
+		[[nodiscard]] FConstMapIterator CreateConstMapIterator() const;
+
+	private:
+		UE_API FAttributeSetMap** AddSetMapImpl(const FAttributeSetMap::FConstructArgs& Args, bool bConstructSetMap);
+		UE_API FAttributeMap** AddMapImpl(UScriptStruct* ValueType, bool bConstructMap);
+
+		// The named set we are based on, every set mapping uses a typed set within it
+		FAttributeNamedSetPtr NamedSet;
+
+		// The space in which our values live
+		FValueSpace ValueSpace;
+
+		// An array of set mappings (sorted by mapping key)
+		FSetMapEntry* SetMaps = nullptr;
+
+		// How many set mappings are contained within this collection
+		int32 SetMapCount = 0;
+
+		// How many set mappings we can contain within this collection
+		int32 SetMapCapacity = 0;
+
+		// An array of name/value mappings (sorted by value type)
+		FMapEntry* Maps = nullptr;
+
+		// How many mappings are contained within this collection
+		int32 MapCount = 0;
+
+		// How many mappings we can contain within this collection
+		int32 MapCapacity = 0;
+
+		// The indirect allocator function pointer
+		FReallocFun ReallocFun = nullptr;
+
+		friend FSetMapIterator;
+		friend FConstSetMapIterator;
+		friend FMapIterator;
+		friend FConstMapIterator;
+	};
+
+	// A heap allocated attribute collection
+	class FAttributeCollectionHeap : public FAttributeCollection
+	{
+	public:
+		// Constructs an empty collection
+		FAttributeCollectionHeap();
+
+		// Constructs an empty collection based on the specified named set
+		explicit FAttributeCollectionHeap(const FAttributeNamedSetPtr& NamedSet);
+
+		using FAttributeCollection::operator=;
+	};
+
+	// A Mem-Stack allocated attribute collection
+	class FAttributeCollectionStack : public FAttributeCollection
+	{
+	public:
+		// Constructs an empty collection
+		FAttributeCollectionStack();
+
+		// Constructs an empty collection based on the specified named set
+		explicit FAttributeCollectionStack(const FAttributeNamedSetPtr& NamedSet);
+
+		using FAttributeCollection::operator=;
+	};
+
+	using FAttributeCollectionPtr = TSharedPtr<FAttributeCollection, ESPMode::ThreadSafe>;
+
+	// An operator structure to join attribute collections using InnerJoin/OuterJoin/etc
+	struct FAttributeCollectionSetMapJoinOp
+	{
+		const FAttributeMappingKey& GetLargestKey() const { return FAttributeMappingKey::LARGEST_VALUE; }
+
+		const FAttributeMappingKey& GetKey(const FAttributeCollection::FSetMapIterator& It) const { return It.GetKey(); }
+		const FAttributeMappingKey& GetKey(const FAttributeCollection::FConstSetMapIterator& It) const { return It.GetKey(); }
+
+		FAttributeSetMap* GetValue(const FAttributeCollection::FSetMapIterator& It) const { return It.GetMap(); }
+		const FAttributeSetMap* GetValue(const FAttributeCollection::FConstSetMapIterator& It) const { return It.GetMap(); }
+	};
+
+	// An operator structure to join attribute collections using InnerJoin/OuterJoin/etc
+	struct FAttributeCollectionMapJoinOp
+	{
+		// Key sorting uses signed comparison, use the largest signed value regardless of pointer size
+		static constexpr intptr_t LARGEST_KEY_VALUE = static_cast<intptr_t>(static_cast<uintptr_t>(~0ULL) >> 1);
+
+		UScriptStruct* GetLargestKey() const { return reinterpret_cast<UScriptStruct*>(LARGEST_KEY_VALUE); }
+
+		UScriptStruct* GetKey(const FAttributeCollection::FMapIterator& It) const { return It.GetValueType(); }
+		UScriptStruct* GetKey(const FAttributeCollection::FConstMapIterator& It) const { return It.GetValueType(); }
+
+		FAttributeMap* GetValue(const FAttributeCollection::FMapIterator& It) const { return It.GetMap(); }
+		const FAttributeMap* GetValue(const FAttributeCollection::FConstMapIterator& It) const { return It.GetMap(); }
+	};
+}
+
+/*
+ * Attribute Collection
+ *
+ * This struct holds attributes produced by UAF. They come in two flavors: static (set map) and dynamic (map).
+ * Static attributes are fast and defined up front in the skeleton set binding asset.
+ * Dynamic attributes are slower but produced on demand during evaluation.
+ */
+USTRUCT(BlueprintType, meta = (DisplayName = "Attribute Collection"))
+struct FUAFAttributeCollection
+#if CPP
+	: public UE::UAF::FAttributeCollectionHeap
+#endif
+{
+	GENERATED_BODY()
+
+	// Equivalent to FAnimNextGraphLODPose
+};
+
+//////////////////////////////////////////////////////////////////////////
+// Implementation
+
+namespace UE::UAF
+{
+	template <class ValueType, typename... TArgs>
+	inline TAttributeSetMap<ValueType>* FAttributeCollection::AddSetMap(const FAttributeMappingKey& MappingKey, TArgs... Args)
+	{
+		if (!NamedSet)
+		{
+			// We are not initialized
+			return nullptr;
+		}
+
+		UScriptStruct* ValueTypeStruct = ValueType::StaticStruct();
+		if (!ensureMsgf(ValueTypeStruct == MappingKey.GetValueType(), TEXT("Mismatch between the return value type and mapping key value type")))
+		{
+			return nullptr;
+		}
+
+		FAttributeTypedSetPtr TypedSet = NamedSet->FindTypedSet(MappingKey.GetAttributeType());
+		if (!TypedSet)
+		{
+			// Unknown attribute type, it isn't contained within our named set
+			return nullptr;
+		}
+
+		// Forward the extra arguments
+		typename TAttributeSetMap<ValueType>::FConstructArgs ConstructArgs(TypedSet, ValueTypeStruct, ReallocFun, Args...);
+
+		// Don't construct the set map internally, it'll have to use indirection to do so and here we know the concrete type
+		constexpr bool bConstructSetMap = false;
+
+		FAttributeSetMap** SetMapPtr = AddSetMapImpl(ConstructArgs, bConstructSetMap);
+		if (SetMapPtr == nullptr)
+		{
+			return nullptr;
+		}
+
+		TAttributeSetMap<ValueType>* SetMap = MakeAttributeSetMap<ValueType>(ConstructArgs);
+
+		check(*SetMapPtr == nullptr);
+		*SetMapPtr = SetMap;
+
+		return SetMap;
+	}
+
+	template <class ValueType>
+	inline TAttributeSetMap<ValueType>* FAttributeCollection::FindSetMap(const FAttributeMappingKey& MappingKey)
+	{
+		UScriptStruct* ValueTypeStruct = ValueType::StaticStruct();
+		if (!ensureMsgf(ValueTypeStruct == MappingKey.GetValueType(), TEXT("Mismatch between the return value type and mapping key value type")))
+		{
+			return nullptr;
+		}
+
+		return Cast<ValueType>(FindSetMap(MappingKey));
+	}
+
+	template <class ValueType>
+	inline const TAttributeSetMap<ValueType>* FAttributeCollection::FindSetMap(const FAttributeMappingKey& MappingKey) const
+	{
+		UScriptStruct* ValueTypeStruct = ValueType::StaticStruct();
+		if (!ensureMsgf(ValueTypeStruct == MappingKey.GetValueType(), TEXT("Mismatch between the return value type and mapping key value type")))
+		{
+			return nullptr;
+		}
+
+		return Cast<ValueType>(FindSetMap(MappingKey));
+	}
+
+	template <class ValueType, typename... TArgs>
+	inline TAttributeSetMap<ValueType>* FAttributeCollection::FindOrAddSetMap(const FAttributeMappingKey& MappingKey, TArgs... Args)
+	{
+		UScriptStruct* ValueTypeStruct = ValueType::StaticStruct();
+		if (!ensureMsgf(ValueTypeStruct == MappingKey.GetValueType(), TEXT("Mismatch between the return value type and mapping key value type")))
+		{
+			return nullptr;
+		}
+
+		if (FAttributeSetMap* SetMap = FindSetMap(MappingKey))
+		{
+			return Cast<ValueType>(SetMap);
+		}
+
+		return AddSetMap<ValueType>(MappingKey, Args...);
+	}
+
+	inline typename FAttributeCollection::FSetMapIterator FAttributeCollection::CreateSetMapIterator()
+	{
+		return FSetMapIterator(*this);
+	}
+
+	inline typename FAttributeCollection::FConstSetMapIterator FAttributeCollection::CreateConstSetMapIterator() const
+	{
+		return FConstSetMapIterator(*this);
+	}
+
+	template <class ValueType>
+	inline TAttributeMap<ValueType>* FAttributeCollection::AddMap()
+	{
+		if (!NamedSet)
+		{
+			// We are not initialized
+			return nullptr;
+		}
+
+		UScriptStruct* ValueTypeStruct = ValueType::StaticStruct();
+
+		// Don't construct the map internally, it'll have to use indirection to do so and here we know the concrete type
+		constexpr bool bConstructMap = false;
+
+		FAttributeMap** MapPtr = AddMapImpl(ValueTypeStruct, bConstructMap);
+		if (MapPtr == nullptr)
+		{
+			return nullptr;
+		}
+
+		TAttributeMap<ValueType>* Map = MakeAttributeMap<ValueType>(ReallocFun);
+
+		check(*MapPtr == nullptr);
+		*MapPtr = Map;
+
+		return Map;
+	}
+
+	template <class ValueType>
+	inline TAttributeMap<ValueType>* FAttributeCollection::FindMap()
+	{
+		return Cast<ValueType>(FindMap(ValueType::StaticStruct()));
+	}
+
+	template <class ValueType>
+	inline const TAttributeMap<ValueType>* FAttributeCollection::FindMap() const
+	{
+		return Cast<ValueType>(FindMap(ValueType::StaticStruct()));
+	}
+
+	template <class ValueType>
+	inline TAttributeMap<ValueType>* FAttributeCollection::FindOrAddMap()
+	{
+		if (TAttributeMap<ValueType>* Map = FindMap<ValueType>())
+		{
+			return Map;
+		}
+
+		return AddMap<ValueType>();
+	}
+
+	inline typename FAttributeCollection::FMapIterator FAttributeCollection::CreateMapIterator()
+	{
+		return FMapIterator(*this);
+	}
+
+	inline typename FAttributeCollection::FConstMapIterator FAttributeCollection::CreateConstMapIterator() const
+	{
+		return FConstMapIterator(*this);
+	}
+
+	inline FAttributeCollectionHeap::FAttributeCollectionHeap()
+		: FAttributeCollection(nullptr, &FAllocatorTypeTrait<FDefaultAllocator>::Realloc)
+	{
+	}
+
+	inline FAttributeCollectionHeap::FAttributeCollectionHeap(const FAttributeNamedSetPtr& InNamedSet)
+		: FAttributeCollection(InNamedSet, &FAllocatorTypeTrait<FDefaultAllocator>::Realloc)
+	{
+	}
+
+	inline FAttributeCollectionStack::FAttributeCollectionStack()
+		: FAttributeCollection(nullptr, &FAllocatorTypeTrait<TMemStackAllocator<>>::Realloc)
+	{
+	}
+
+	inline FAttributeCollectionStack::FAttributeCollectionStack(const FAttributeNamedSetPtr& InNamedSet)
+		: FAttributeCollection(InNamedSet, &FAllocatorTypeTrait<TMemStackAllocator<>>::Realloc)
+	{
+	}
+
+	template<class ContainerType>
+	inline TAttributeCollectionSetMapIterator<ContainerType>::TAttributeCollectionSetMapIterator(ContainerType& Collection)
+		: EntryPtr(Collection.SetMaps)
+		, EntryEndPtr(Collection.SetMaps + Collection.SetMapCount)
+	{
+	}
+
+	template<class ContainerType>
+	inline TAttributeCollectionSetMapIterator<ContainerType>& TAttributeCollectionSetMapIterator<ContainerType>::operator++()
+	{
+		++EntryPtr;
+		return *this;
+	}
+
+	template<class ContainerType>
+	inline TAttributeCollectionSetMapIterator<ContainerType>::operator bool() const
+	{
+		return EntryPtr != EntryEndPtr;
+	}
+
+	template<class ContainerType>
+	inline const FAttributeMappingKey& TAttributeCollectionSetMapIterator<ContainerType>::GetKey() const
+	{
+		return EntryPtr->Key;
+	}
+
+	template<class ContainerType>
+	inline UScriptStruct* TAttributeCollectionSetMapIterator<ContainerType>::GetAttributeType() const
+	{
+		return EntryPtr->Key.GetAttributeType();
+	}
+
+	template<class ContainerType>
+	inline UScriptStruct* TAttributeCollectionSetMapIterator<ContainerType>::GetValueType() const
+	{
+		return EntryPtr->Key.GetValueType();
+	}
+
+	template<class ContainerType>
+	inline typename TAttributeCollectionSetMapIterator<ContainerType>::SetMapType TAttributeCollectionSetMapIterator<ContainerType>::GetMap() const
+	{
+		return EntryPtr->Map;
+	}
+
+	template<class ContainerType>
+	inline TAttributeCollectionMapIterator<ContainerType>::TAttributeCollectionMapIterator(ContainerType& Collection)
+		: EntryPtr(Collection.Maps)
+		, EntryEndPtr(Collection.Maps + Collection.MapCount)
+	{
+	}
+
+	template<class ContainerType>
+	inline TAttributeCollectionMapIterator<ContainerType>& TAttributeCollectionMapIterator<ContainerType>::operator++()
+	{
+		++EntryPtr;
+		return *this;
+	}
+
+	template<class ContainerType>
+	inline TAttributeCollectionMapIterator<ContainerType>::operator bool() const
+	{
+		return EntryPtr != EntryEndPtr;
+	}
+
+	template<class ContainerType>
+	inline UScriptStruct* TAttributeCollectionMapIterator<ContainerType>::GetValueType() const
+	{
+		return EntryPtr->ValueType;
+	}
+
+	template<class ContainerType>
+	inline typename TAttributeCollectionMapIterator<ContainerType>::MapType TAttributeCollectionMapIterator<ContainerType>::GetMap() const
+	{
+		return EntryPtr->Map;
+	}
+}
+
+#undef UE_API
diff --git a/Engine/Plugins/Experimental/UAF/UAF/Source/UAF/Public/UAF/ValueRuntime/AttributeMap.h b/Engine/Plugins/Experimental/UAF/UAF/Source/UAF/Public/UAF/ValueRuntime/AttributeMap.h
new file mode 100644
--- /dev/null
+++ b/Engine/Plugins/Experimental/UAF/UAF/Source/UAF/Public/UAF/ValueRuntime/AttributeMap.h
@@ -0,0 +1,815 @@
+// Copyright Epic Games, Inc. All Rights Reserved.
+
+#pragma once
+
+#include "CoreMinimal.h"
+#include "Algo/BinarySearch.h"
+#include "Algo/Sort.h"
+#include "Templates/MemoryOps.h"
+#include "UAF/ValueRuntime/IndirectAllocator.h"
+
+#define UE_API UAF_API
+
+namespace UE::UAF
+{
+	// Untyped base class for TAttributeMap
+	class FAttributeMap
+	{
+	public:
+		UE_NONCOPYABLE(FAttributeMap);	// We disallow copy/move semantics
+
+		// Returns whether or not this map has any values
+		[[nodiscard]] UE_API bool IsEmpty() const;
+
+		// Returns the number of mapped values
+		[[nodiscard]] UE_API int32 Num() const;
+
+		// Returns the maximum number of values we can map
+		[[nodiscard]] UE_API int32 Max() const;
+
+		// Returns the mapped value type
+		[[nodiscard]] UE_API UScriptStruct* GetValueType() const;
+
+		// Returns the total allocated size for this map
+		[[nodiscard]] UE_API int32 GetAllocatedSize() const;
+
+		// Returns the allocator function this map was initialized with
+		[[nodiscard]] UE_API FReallocFun GetAllocator() const;
+
+		// Returns a copy of this map using the specified allocator
+		[[nodiscard]] virtual FAttributeMap* Duplicate(FReallocFun InReallocFun) const = 0;
+
+		// Copies the content of this map into another
+		// Both must have matching value types
+		virtual void CopyTo(FAttributeMap& Other) const = 0;
+
+		// Moves the content of this map to another
+		// Both must have matching value types
+		// The content of this map after the move is undefined, just like with move assignment
+		UE_API virtual void MoveTo(FAttributeMap& Other);
+
+		// Finds a mapping and returns its index or INDEX_NONE if not found, O(logN)
+		[[nodiscard]] int32 IndexOf(FName Name) const;
+
+		// Returns whether or not this map contains the specified mapped value, O(logN)
+		[[nodiscard]] bool Contains(FName Name) const;
+
+		// Retrieves a mapping name by index, O(1)
+		[[nodiscard]] FName GetName(int32 Index) const;
+
+		// Adds a new mapping, sets the value using the provided setter, O(logN)
+		// Returns whether or not we inserted (we can fail if a value with that name already exists)
+		// The value setter (and/or caller) is responsible for validating the value type
+		virtual bool AddWithSetter(FName Name, const TFunctionRef<void(UScriptStruct* ValueType, uint8* OutValuePtr)>& ValueSetter) = 0;
+
+		// Gets a value at the specified index using a value getter function, O(1)
+		// The value getter (and/or caller) is responsible for validating the value type
+		virtual void GetValueWithGetter(int32 Index, const TFunctionRef<void(UScriptStruct* ValueType, const uint8* ValuePtr)>& ValueGetter) const = 0;
+
+	protected:
+		// Default construction for an empty map
+		UE_API FAttributeMap();
+
+		// Constructs a map but does not allocate its memory, derived types are responsible for it
+		UE_API FAttributeMap(UScriptStruct* ValueType, FReallocFun ReallocFun, int32 ContainerSize);
+
+		// Destroys the map and frees its memory
+		UE_API virtual ~FAttributeMap();
+
+		// Returns a pointer to the start of the names
+		[[nodiscard]] FName* GetNames();
+		[[nodiscard]] const FName* GetNames() const;
+
+		// The type of values this map contains
+		UScriptStruct* ValueTypeStruct = nullptr;
+
+		// The number of values mapped in our buffer
+		int32 NumValues = 0;
+
+		// The maximum number of values that we can map in our buffer
+		int32 MaxValues = 0;
+
+		// The base untyped pointer for our names/values
+		// Names and values are allocated together within the same contiguous memory buffer
+		// Names are first within the buffer followed by value types using default alignment
+		uint8* DataPtr = nullptr;
+
+		// The indirect allocator function pointer
+		FReallocFun ReallocFun = nullptr;
+
+		// The size in bytes of this container (excluding any externally allocated data)
+		int32 ContainerSize = 0;
+
+		friend class FAttributeCollection;
+		friend class FValueRuntimeRegistry;
+		friend void ReleaseAttributeMap(FAttributeMap* Map);
+	};
+
+	// A iterator that iterates over all name/value pairs within an attribute map
+	template<class ContainerType>
+	class TAttributeMapIterator
+	{
+	public:
+		static constexpr bool bIsContainerConst = std::is_const_v<ContainerType>;
+
+		using NamePtrType = typename std::conditional<bIsContainerConst, const FName*, FName*>::type;
+		using ValuePtrType = typename std::conditional<bIsContainerConst, const typename ContainerType::ValueType*, typename ContainerType::ValueType*>::type;
+
+		// Creates an empty iterator
+		[[nodiscard]] TAttributeMapIterator() = default;
+
+		// Creates an iterator over the name/value map entries within the specified attribute map
+		[[nodiscard]] TAttributeMapIterator(ContainerType& Map);
+
+		// Increments and moves the iterator to the next name/value pair
+		TAttributeMapIterator& operator++();
+
+		// Returns whether or not the iterator still contains values
+		[[nodiscard]] explicit operator bool() const;
+
+		// Returns the name of the current name/value pair
+		[[nodiscard]] FName GetName() const;
+
+		// Returns the value of the current name/value pair
+		[[nodiscard]] ValuePtrType GetValue() const;
+
+	private:
+		// Current name/value pair
+		NamePtrType NamePtr = nullptr;
+		ValuePtrType ValuePtr = nullptr;
+
+		// One past the last name/value pair
+		NamePtrType NameEndPtr = nullptr;
+	};
+
+	// Encapsulates a mapping of FName to value type
+	// Value types must be a UStruct
+	template<class InValueType>
+	class TAttributeMap : public FAttributeMap
+	{
+	public:
+		// The type of values held in this map
+		using ValueType = InValueType;
+
+		// Returns a copy of this map using the specified allocator
+		[[nodiscard]] virtual FAttributeMap* Duplicate(FReallocFun InReallocFun) const override;
+
+		// Copies the content of this map into another
+		// Both must have matching value types
+		virtual void CopyTo(FAttributeMap& Other) const override;
+
+		// Adds a new mapping, sets the value using the provided setter, O(logN)
+		// Returns whether or not we inserted (we can fail if a value with that name already exists)
+		// The value setter (and/or caller) is responsible for validating the value type
+		virtual bool AddWithSetter(FName Name, const TFunctionRef<void(UScriptStruct* ValueType, uint8* OutValuePtr)>& ValueSetter) override;
+
+		// Gets a value at the specified index using a value getter function, O(1)
+		// The value getter (and/or caller) is responsible for validating the value type
+		virtual void GetValueWithGetter(int32 Index, const TFunctionRef<void(UScriptStruct* ValueType, const uint8* ValuePtr)>& ValueGetter) const override;
+
+		// Adds a new mapping, O(logN)
+		// Returns whether or not we inserted (we can fail if a value with that name already exists)
+		bool Add(FName Name, const InValueType& Value);
+
+		// Adds a new mapping at the end of the map, must be sorted, O(1)
+		// Returns whether or not we inserted (we can fail if the value isn't sorted to land at the end of the map)
+		bool Append(FName Name, const InValueType& Value);
+
+		// Removes a mapping, O(logN)
+		// Returns whether or not we removed the value (we can fail if the name isn't found)
+		bool Remove(FName Name);
+
+		// Removes a mapping, O(1)
+		void RemoveAt(int32 Index);
+
+		// Finds a mapping and returns a pointer to it or nullptr if not found, O(logN)
+		[[nodiscard]] InValueType* Find(FName Name);
+		[[nodiscard]] const InValueType* Find(FName Name) const;
+
+		// Retrieves a mapping value by index, O(1)
+		[[nodiscard]] InValueType& GetValue(int32 Index);
+		[[nodiscard]] const InValueType& GetValue(int32 Index) const;
+
+		// Initializes the mapping using sorted inputs, O(N)
+		template<class GetNamePredicateType, class GetValuePredicateType>
+		void SetSorted(int32 NumValues, GetNamePredicateType GetNamePredicate, GetValuePredicateType GetValuePredicate);
+
+		// Initializes the mapping using unsorted inputs, O(NlogN)
+		template<class GetNamePredicateType, class GetValuePredicateType>
+		void SetUnsorted(int32 NumValues, GetNamePredicateType GetNamePredicate, GetValuePredicateType GetValuePredicate);
+
+		// Iterators
+		using FIterator = TAttributeMapIterator<TAttributeMap<InValueType>>;
+		using FConstIterator = TAttributeMapIterator<const TAttributeMap<InValueType>>;
+
+		// Returns an iterator over all name/value pairs contained within, sorted by their name
+		[[nodiscard]] FIterator CreateIterator();
+		[[nodiscard]] FConstIterator CreateConstIterator() const;
+
+	protected:
+		static const int32 ValueTypeAlignment = sizeof(InValueType) >= 16 ? 16 : 8;
+
+		// Constructs a map
+		TAttributeMap(UScriptStruct* ValueType, FReallocFun ReallocFun, int32 ContainerSize);
+
+		// Destroys the map and frees its memory
+		virtual ~TAttributeMap() override;
+
+		void Grow();
+
+		// Returns a pointer to the start of the values
+		InValueType* GetValues();
+		const InValueType* GetValues() const;
+
+		template<class T>
+		friend TAttributeMap<T>* MakeAttributeMap(FReallocFun ReallocFun);
+		friend FIterator;
+		friend FConstIterator;
+	};
+
+	// An operator structure to join attribute maps using InnerJoin/OuterJoin/etc
+	struct FAttributeMapJoinOp
+	{
+		// We generally want the key in our predicate since the values by themselves have little meaning
+		static constexpr bool HasPredicateWithKey() { return true; }
+
+		bool IsLessThan(FName LHS, FName RHS) const
+		{
+			return LHS.FastLess(RHS);
+		}
+
+		// Special dummy FName that is larger than all other FNames in fast sorted order
+		FName GetLargestKey() const { return FName(FNameEntryId::FromUnstableInt(MAX_int32), FNameEntryId::FromUnstableInt(MAX_int32), MAX_int32); }
+
+		template<typename IteratorType>
+		FName GetKey(const IteratorType& It) const { return It.GetName(); }
+
+		template<typename IteratorType>
+		typename IteratorType::ValuePtrType GetValue(const IteratorType& It) const { return It.GetValue(); }
+	};
+
+	// Allocates and constructs a new attribute map for the specified value type
+	template<class ValueType>
+	TAttributeMap<ValueType>* MakeAttributeMap(FReallocFun ReallocFun);
+
+	// Releases and destroys the specified attribute map
+	void ReleaseAttributeMap(FAttributeMap* Map);
+
+	// Cast from an untyped name/value map to a typed name/value map
+	template<class ValueType>
+	UE_FORCEINLINE_HINT TAttributeMap<ValueType>* Cast(FAttributeMap* Src)
+	{
+		return Src != nullptr && ValueType::StaticStruct() == Src->GetValueType() ? static_cast<TAttributeMap<ValueType>*>(Src) : nullptr;
+	}
+
+	// Cast from an untyped name/value map to a typed name/value map
+	template<class ValueType>
+	UE_FORCEINLINE_HINT const TAttributeMap<ValueType>* Cast(const FAttributeMap* Src)
+	{
+		return Src != nullptr && ValueType::StaticStruct() == Src->GetValueType() ? static_cast<const TAttributeMap<ValueType>*>(Src) : nullptr;
+	}
+
+	// Cast from an untyped name/value map to a typed name/value map with an assert if the types mismatch
+	template<class ValueType>
+	inline TAttributeMap<ValueType>* CastChecked(FAttributeMap* Src)
+	{
+		TAttributeMap<ValueType>* CastResult = Src != nullptr && ValueType::StaticStruct() == Src->GetValueType() ? static_cast<TAttributeMap<ValueType>*>(Src) : nullptr;
+		check(CastResult);
+		return CastResult;
+	}
+
+	// Cast from an untyped name/value map to a typed name/value map with an assert if the types mismatch
+	template<class ValueType>
+	inline const TAttributeMap<ValueType>* CastChecked(const FAttributeMap* Src)
+	{
+		const TAttributeMap<ValueType>* CastResult = Src != nullptr && ValueType::StaticStruct() == Src->GetValueType() ? static_cast<const TAttributeMap<ValueType>*>(Src) : nullptr;
+		check(CastResult);
+		return CastResult;
+	}
+
+	//////////////////////////////////////////////////////////////////////////
+	// Implementation
+
+	inline int32 FAttributeMap::IndexOf(FName Name) const
+	{
+		const FName* Names = GetNames();
+
+		const int32 Index = Algo::LowerBound(TArrayView<const FName>(Names, NumValues), Name, FNameFastLess());
+		if (Index >= NumValues || Names[Index] != Name)
+		{
+			// No value is mapped to that name
+			return INDEX_NONE;
+		}
+
+		return Index;
+	}
+
+	inline bool FAttributeMap::Contains(FName Name) const
+	{
+		return IndexOf(Name) != INDEX_NONE;
+	}
+
+	inline FName FAttributeMap::GetName(int32 Index) const
+	{
+		checkf((Index >= 0) & (Index < NumValues), TEXT("Map index out of bounds: %lld into an map of size %lld"), (long long)Index, (long long)NumValues); // & for one branch
+		return GetNames()[Index];
+	}
+
+	inline FName* FAttributeMap::GetNames()
+	{
+		return reinterpret_cast<FName*>(DataPtr);
+	}
+
+	inline const FName* FAttributeMap::GetNames() const
+	{
+		return reinterpret_cast<const FName*>(DataPtr);
+	}
+
+	template<class InValueType>
+	inline TAttributeMap<InValueType>::TAttributeMap(UScriptStruct* ValueTypeStruct, FReallocFun InReallocFun, int32 InContainerSize)
+		: FAttributeMap(ValueTypeStruct, InReallocFun, InContainerSize)
+	{
+		check(ValueTypeStruct == InValueType::StaticStruct());
+		check(InContainerSize == sizeof(TAttributeMap<InValueType>));
+	}
+
+	template<class InValueType>
+	inline TAttributeMap<InValueType>::~TAttributeMap()
+	{
+		if (DataPtr != nullptr)
+		{
+			// FNames don't need to be destroyed, only destroy the values
+			DestructItems<InValueType>(GetValues(), NumValues);
+
+			DataPtr = nullptr;
+			NumValues = MaxValues = 0;
+		}
+	}
+
+	template<class InValueType>
+	inline FAttributeMap* TAttributeMap<InValueType>::Duplicate(FReallocFun InReallocFun) const
+	{
+		TAttributeMap<InValueType>* Copy = MakeAttributeMap<InValueType>(InReallocFun);
+
+		if (NumValues != 0)
+		{
+			check(Copy->MaxValues == 0 && Copy->DataPtr == nullptr);
+
+			const int32 NewSize = Align(sizeof(FName) * NumValues, ValueTypeAlignment) + (sizeof(InValueType) * NumValues);
+
+			Copy->DataPtr = (*InReallocFun)(nullptr, 0, NewSize);
+			Copy->NumValues = NumValues;
+			Copy->MaxValues = NumValues;
+
+			ConstructItems<FName>(Copy->GetNames(), GetNames(), NumValues);
+			ConstructItems<InValueType>(Copy->GetValues(), GetValues(), NumValues);
+		}
+
+		return Copy;
+	}
+
+	template<class InValueType>
+	inline void TAttributeMap<InValueType>::CopyTo(FAttributeMap& Other) const
+	{
+		checkf(ValueTypeStruct == Other.GetValueType(), TEXT("Value types must match to allow copying"));
+
+		TAttributeMap<InValueType>& Copy = reinterpret_cast<TAttributeMap<InValueType>&>(Other);
+
+		// Destroy old values
+		if (Copy.NumValues != 0)
+		{
+			// FNames don't need to be destroyed, only destroy the values
+			DestructItems<InValueType>(Copy.GetValues(), Copy.NumValues);
+		}
+
+		// Resize
+		if (Copy.MaxValues != NumValues)
+		{
+			const int32 OldSize = Align(sizeof(FName) * Copy.MaxValues, ValueTypeAlignment) + (sizeof(InValueType) * Copy.MaxValues);
+			const int32 NewSize = Align(sizeof(FName) * NumValues, ValueTypeAlignment) + (sizeof(InValueType) * NumValues);
+
+			Copy.DataPtr = (*Copy.ReallocFun)(Copy.DataPtr, OldSize, NewSize);
+			Copy.MaxValues = NumValues;
+		}
+
+		// Copy construct our new values
+		if (NumValues != 0)
+		{
+			ConstructItems<FName>(Copy.GetNames(), GetNames(), NumValues);
+			ConstructItems<InValueType>(Copy.GetValues(), GetValues(), NumValues);
+		}
+
+		Copy.NumValues = NumValues;
+	}
+
+	template<class InValueType>
+	inline bool TAttributeMap<InValueType>::AddWithSetter(FName Name, const TFunctionRef<void(UScriptStruct* ValueType, uint8* OutValuePtr)>& ValueSetter)
+	{
+		FName* Names = GetNames();
+
+		const int32 Index = Algo::LowerBound(TArrayView<const FName>(Names, NumValues), Name, FNameFastLess());
+		if (Index < NumValues && Names[Index] == Name)
+		{
+			// A mapping already exists that matches the specified name
+			return false;
+		}
+
+		if (NumValues >= MaxValues)
+		{
+			// Not enough space, grow our buffer
+			Grow();
+
+			Names = GetNames();
+		}
+
+		InValueType* Values = GetValues();
+
+		// Insert our value at the right position, making space for it
+		const int32 NumToMove = NumValues - Index;
+		RelocateConstructItems<FName>(Names + Index + 1, Names + Index, NumToMove);
+		RelocateConstructItems<InValueType>(Values + Index + 1, Values + Index, NumToMove);
+		new(Names + Index) FName(Name);
+		ValueSetter(ValueTypeStruct, reinterpret_cast<uint8*>(Values + Index));
+		NumValues++;
+
+		return true;
+	}
+
+	template<class InValueType>
+	inline void TAttributeMap<InValueType>::GetValueWithGetter(int32 Index, const TFunctionRef<void(UScriptStruct* ValueType, const uint8* ValuePtr)>& ValueGetter) const
+	{
+		checkf((Index >= 0) & (Index < NumValues), TEXT("Map index out of bounds: %lld into an map of size %lld"), (long long)Index, (long long)NumValues); // & for one branch
+		const InValueType* ValuePtr = GetValues() + Index;
+		ValueGetter(ValueTypeStruct, reinterpret_cast<const uint8*>(ValuePtr));
+	}
+
+	template<class InValueType>
+	inline bool TAttributeMap<InValueType>::Add(FName Name, const InValueType& Value)
+	{
+		FName* Names = GetNames();
+
+		const int32 Index = Algo::LowerBound(TArrayView<const FName>(Names, NumValues), Name, FNameFastLess());
+		if (Index < NumValues && Names[Index] == Name)
+		{
+			// A mapping already exists that matches the specified name
+			return false;
+		}
+
+		if (NumValues >= MaxValues)
+		{
+			// Not enough space, grow our buffer
+			Grow();
+			
+			Names = GetNames();
+		}
+
+		InValueType* Values = GetValues();
+
+		// Insert our value at the right position, making space for it
+		const int32 NumToMove = NumValues - Index;
+		RelocateConstructItems<FName>(Names + Index + 1, Names + Index, NumToMove);
+		RelocateConstructItems<InValueType>(Values + Index + 1, Values + Index, NumToMove);
+		new(Names + Index) FName(Name);
+		new(Values + Index) InValueType(Value);
+		NumValues++;
+
+		return true;
+	}
+
+	template<class InValueType>
+	inline bool TAttributeMap<InValueType>::Append(FName Name, const InValueType& Value)
+	{
+		FName* Names = GetNames();
+
+		if (NumValues != 0 && !Names[NumValues - 1].FastLess(Name))
+		{
+			// The name doesn't belong at the end, it isn't sorted
+			return false;
+		}
+
+		if (NumValues >= MaxValues)
+		{
+			// Not enough space, grow our buffer
+			Grow();
+
+			Names = GetNames();
+		}
+
+		InValueType* Values = GetValues();
+
+		// Append our new value at the end
+		new(Names + NumValues) FName(Name);
+		new(Values + NumValues) InValueType(Value);
+		NumValues++;
+
+		return true;
+	}
+
+	template<class InValueType>
+	inline bool TAttributeMap<InValueType>::Remove(FName Name)
+	{
+		const FName* Names = GetNames();
+
+		const int32 Index = Algo::LowerBound(TArrayView<const FName>(Names, NumValues), Name, FNameFastLess());
+		if (Index >= NumValues || Names[Index] != Name)
+		{
+			// No value is mapped to that name
+			return false;
+		}
+
+		RemoveAt(Index);
+
+		return true;
+	}
+
+	template<class InValueType>
+	inline void TAttributeMap<InValueType>::RemoveAt(int32 Index)
+	{
+		checkf((Index >= 0) & (Index < NumValues), TEXT("Map index out of bounds: %lld into an map of size %lld"), (long long)Index, (long long)NumValues); // & for one branch
+
+		FName* Names = GetNames();
+		InValueType* Values = GetValues();
+
+		// FNames don't need to be destroyed, only destroy the values
+		DestructItem<InValueType>(Values + Index);
+
+		const int32 NumToMove = NumValues - Index - 1;
+		RelocateConstructItems<FName>(Names + Index, Names + Index + 1, NumToMove);
+		RelocateConstructItems<InValueType>(Values + Index, Values + Index + 1, NumToMove);
+		NumValues--;
+	}
+
+	template<class InValueType>
+	inline InValueType* TAttributeMap<InValueType>::Find(FName Name)
+	{
+		const FName* Names = GetNames();
+
+		const int32 Index = Algo::LowerBound(TArrayView<const FName>(Names, NumValues), Name, FNameFastLess());
+		if (Index >= NumValues || Names[Index] != Name)
+		{
+			// No value is mapped to that name
+			return nullptr;
+		}
+
+		return GetValues() + Index;
+	}
+
+	template<class InValueType>
+	inline const InValueType* TAttributeMap<InValueType>::Find(FName Name) const
+	{
+		const FName* Names = GetNames();
+
+		const int32 Index = Algo::LowerBound(TArrayView<const FName>(Names, NumValues), Name, FNameFastLess());
+		if (Index >= NumValues || Names[Index] != Name)
+		{
+			// No value is mapped to that name
+			return nullptr;
+		}
+
+		return GetValues() + Index;
+	}
+
+	template<class InValueType>
+	inline InValueType& TAttributeMap<InValueType>::GetValue(int32 Index)
+	{
+		checkf((Index >= 0) & (Index < NumValues), TEXT("Map index out of bounds: %lld into an map of size %lld"), (long long)Index, (long long)NumValues); // & for one branch
+
+		return GetValues()[Index];
+	}
+
+	template<class InValueType>
+	inline const InValueType& TAttributeMap<InValueType>::GetValue(int32 Index) const
+	{
+		checkf((Index >= 0) & (Index < NumValues), TEXT("Map index out of bounds: %lld into an map of size %lld"), (long long)Index, (long long)NumValues); // & for one branch
+
+		return GetValues()[Index];
+	}
+
+	template<class InValueType>
+	template<class GetNamePredicateType, class GetValuePredicateType>
+	inline void TAttributeMap<InValueType>::SetSorted(int32 InNumValues, GetNamePredicateType GetNamePredicate, GetValuePredicateType GetValuePredicate)
+	{
+		// Destroy old values
+		if (NumValues != 0)
+		{
+			// FNames don't need to be destroyed, only destroy the values
+			DestructItems<InValueType>(GetValues(), NumValues);
+		}
+
+		// Resize
+		if (MaxValues != InNumValues)
+		{
+			const int32 OldSize = Align(sizeof(FName) * MaxValues, ValueTypeAlignment) + (sizeof(InValueType) * MaxValues);
+			const int32 NewSize = Align(sizeof(FName) * InNumValues, ValueTypeAlignment) + (sizeof(InValueType) * InNumValues);
+
+			DataPtr = (*ReallocFun)(DataPtr, OldSize, NewSize);
+			MaxValues = InNumValues;
+		}
+
+		FName* Names = GetNames();
+		InValueType* Values = GetValues();
+
+		// Set now to make debugger easier
+		NumValues = InNumValues;
+
+		// Set our values
+		for (int32 Index = 0; Index < InNumValues; ++Index)
+		{
+			const FName Name = GetNamePredicate(Index);
+			checkf(Index == 0 || Names[Index - 1].FastLess(Name), TEXT("Attempting to initialize an attribute map from unsorted data"));
+
+			new(Names + Index) FName(Name);
+			new(Values + Index) InValueType(GetValuePredicate(Index));
+		}
+	}
+
+	template<class InValueType>
+	template<class GetNamePredicateType, class GetValuePredicateType>
+	inline void TAttributeMap<InValueType>::SetUnsorted(int32 InNumValues, GetNamePredicateType GetNamePredicate, GetValuePredicateType GetValuePredicate)
+	{
+		// Destroy old values
+		if (NumValues != 0)
+		{
+			// FNames don't need to be destroyed, only destroy the values
+			DestructItems<InValueType>(GetValues(), NumValues);
+		}
+
+		// Resize
+		if (MaxValues != InNumValues)
+		{
+			const int32 OldSize = Align(sizeof(FName) * MaxValues, ValueTypeAlignment) + (sizeof(InValueType) * MaxValues);
+			const int32 NewSize = Align(sizeof(FName) * InNumValues, ValueTypeAlignment) + (sizeof(InValueType) * InNumValues);
+
+			DataPtr = (*ReallocFun)(DataPtr, OldSize, NewSize);
+			MaxValues = InNumValues;
+		}
+
+		FName* Names = GetNames();
+		InValueType* Values = GetValues();
+
+		// Set now to make debugger easier
+		NumValues = InNumValues;
+
+#if 1
+		// TODO: We don't appear to have a version of Algo::Sort that works on iterators
+		// and so we can't sort our two arrays in tandem, use a temporary array instead
+		using FTmpElementType = TPair<FName, InValueType>;
+		TArray<FTmpElementType> TmpElements;
+		TmpElements.Reserve(InNumValues);
+		for (int32 Index = 0; Index < InNumValues; ++Index)
+		{
+			TmpElements.Emplace(FTmpElementType{ FName(GetNamePredicate(Index)), InValueType(GetValuePredicate(Index)) });
+		}
+
+		// Sort our temporary elements
+		Algo::Sort(TmpElements, [](const FTmpElementType& LHS, const FTmpElementType& RHS) { return LHS.Key.FastLess(RHS.Key); });
+
+		// Set our values
+		for (int32 Index = 0; Index < InNumValues; ++Index)
+		{
+			const FTmpElementType& Element = TmpElements[Index];
+
+			new(Names + Index) FName(Element.Key);
+			new(Values + Index) InValueType(Element.Value);
+		}
+#else
+		// Set our values
+		for (int32 Index = 0; Index < InNumValues; ++Index)
+		{
+			new(Names + Index) FName(GetNamePredicate(Index));
+			new(Values + Index) InValueType(GetValuePredicate(Index));
+		}
+
+		// Sort our values
+		//Algo::Sort(Elements, FElementSortPredicate());	// TODO
+#endif
+	}
+
+	template<class InValueType>
+	inline typename TAttributeMap<InValueType>::FIterator TAttributeMap<InValueType>::CreateIterator()
+	{
+		return FIterator(*this);
+	}
+
+	template<class InValueType>
+	inline typename TAttributeMap<InValueType>::FConstIterator TAttributeMap<InValueType>::CreateConstIterator() const
+	{
+		return FConstIterator(*this);
+	}
+
+	template<class InValueType>
+	FORCENOINLINE void TAttributeMap<InValueType>::Grow()
+	{
+		const bool bAllowQuantize = false;
+		const int32 NewCapacity = DefaultCalculateSlackGrow(MaxValues + 1, MaxValues, 1, bAllowQuantize);
+
+		const int32 OldValuesOffset = Align(sizeof(FName) * MaxValues, ValueTypeAlignment);
+		const int32 NewValuesOffset = Align(sizeof(FName) * NewCapacity, ValueTypeAlignment);
+		const int32 OldSize = OldValuesOffset + (sizeof(InValueType) * MaxValues);
+		const int32 NewSize = NewValuesOffset + (sizeof(InValueType) * NewCapacity);
+
+		// Our names and values live within the same buffer but when we re-size, they need to end up at
+		// different offsets since both halves grow. To avoid reallocate moving our names/values and
+		// then having another pass to offset our values, we handle it manually.
+
+		// Cache our old pointers
+		FName* OldNames = GetNames();
+		InValueType* OldValues = GetValues();
+
+		// Allocate our new buffer
+		uint8* NewDataPtr = (*ReallocFun)(nullptr, 0, NewSize);
+
+		// Cache our new pointers
+		FName* NewNames = reinterpret_cast<FName*>(NewDataPtr);
+		InValueType* NewValues = reinterpret_cast<InValueType*>(NewDataPtr + NewValuesOffset);
+
+		// Relocate our name/value pairs
+		RelocateConstructItems<FName>(NewNames, OldNames, NumValues);
+		RelocateConstructItems<InValueType>(NewValues, OldValues, NumValues);
+
+		// Deallocate our old buffer
+		(*ReallocFun)(DataPtr, OldSize, 0);
+
+		// Update our state
+		DataPtr = NewDataPtr;
+
+		MaxValues = NewCapacity;
+	}
+
+	template<class InValueType>
+	inline InValueType* TAttributeMap<InValueType>::GetValues()
+	{
+		const int32 ValuesOffset = Align(MaxValues * sizeof(FName), ValueTypeAlignment);
+		return reinterpret_cast<InValueType*>(DataPtr + ValuesOffset);
+	}
+
+	template<class InValueType>
+	inline const InValueType* TAttributeMap<InValueType>::GetValues() const
+	{
+		const int32 ValuesOffset = Align(MaxValues * sizeof(FName), ValueTypeAlignment);
+		return reinterpret_cast<const InValueType*>(DataPtr + ValuesOffset);
+	}
+
+	template<class ContainerType>
+	inline TAttributeMapIterator<ContainerType>::TAttributeMapIterator(ContainerType& Map)
+		: NamePtr(Map.GetNames())
+		, ValuePtr(Map.GetValues())
+		, NameEndPtr(NamePtr + Map.Num())
+	{
+	}
+
+	template<class ContainerType>
+	inline TAttributeMapIterator<ContainerType>& TAttributeMapIterator<ContainerType>::operator++()
+	{
+		++NamePtr;
+		++ValuePtr;
+		return *this;
+	}
+
+	template<class ContainerType>
+	inline TAttributeMapIterator<ContainerType>::operator bool() const
+	{
+		return NamePtr != NameEndPtr;
+	}
+
+	template<class ContainerType>
+	inline FName TAttributeMapIterator<ContainerType>::GetName() const
+	{
+		return *NamePtr;
+	}
+
+	template<class ContainerType>
+	inline typename TAttributeMapIterator<ContainerType>::ValuePtrType TAttributeMapIterator<ContainerType>::GetValue() const
+	{
+		return ValuePtr;
+	}
+
+	template<class ValueType>
+	inline TAttributeMap<ValueType>* MakeAttributeMap(FReallocFun ReallocFun)
+	{
+		const int32 ContainerSize = sizeof(TAttributeMap<ValueType>);
+
+		// Allocate our map container buffer
+		uint8* Ptr = (*ReallocFun)(nullptr, 0, ContainerSize);
+
+		// Construct our map over it
+		return new(Ptr) TAttributeMap<ValueType>(ValueType::StaticStruct(), ReallocFun, ContainerSize);
+	}
+
+	inline void ReleaseAttributeMap(FAttributeMap* Map)
+	{
+		if (Map == nullptr)
+		{
+			return;
+		}
+
+		const int32 ContainerSize = Map->ContainerSize;
+		const FReallocFun ReallocFun = Map->ReallocFun;
+
+		Map->~FAttributeMap();
+
+		(*ReallocFun)(reinterpret_cast<uint8*>(Map), ContainerSize, 0);
+	}
+}
+
+#undef UE_API
diff --git a/Engine/Plugins/Experimental/UAF/UAF/Source/UAF/Public/UAF/ValueRuntime/AttributeSetMap.h b/Engine/Plugins/Experimental/UAF/UAF/Source/UAF/Public/UAF/ValueRuntime/AttributeSetMap.h
new file mode 100644
--- /dev/null
+++ b/Engine/Plugins/Experimental/UAF/UAF/Source/UAF/Public/UAF/ValueRuntime/AttributeSetMap.h
@@ -0,0 +1,435 @@
+// Copyright Epic Games, Inc. All Rights Reserved.
+
+#pragma once
+
+#include "CoreMinimal.h"
+#include "Templates/MemoryOps.h"
+#include "UAF/Attributes/AttributeTypedSet.h"
+#include "UAF/Attributes/EngineAttributes.h"
+#include "UAF/ValueRuntime/AttributeMappingKey.h"
+#include "UAF/ValueRuntime/IndirectAllocator.h"
+
+#define UE_API UAF_API
+
+namespace UE::UAF
+{
+	// Untyped base class for TAttributeSetMap
+	class FAttributeSetMap
+	{
+	public:
+		UE_NONCOPYABLE(FAttributeSetMap);	// We disallow copy/move semantics
+
+		struct FDefaultConstructArgs
+		{
+			const FAttributeTypedSetPtr& TypedSet;
+			UScriptStruct* ValueType = nullptr;
+			FReallocFun ReallocFun = nullptr;
+
+			FDefaultConstructArgs(const FAttributeTypedSetPtr& InTypedSet, UScriptStruct* InValueType, FReallocFun InReallocFun)
+				: TypedSet(InTypedSet)
+				, ValueType(InValueType)
+				, ReallocFun(InReallocFun)
+			{
+			}
+		};
+
+		using FConstructArgs = FDefaultConstructArgs;
+
+		// Returns whether or not this map has any values
+		[[nodiscard]] UE_API bool IsEmpty() const;
+
+		// Returns the number of mapped values
+		[[nodiscard]] UE_API int32 Num() const;
+
+		// Returns the attribute typed set we map to
+		[[nodiscard]] UE_API const FAttributeTypedSetPtr& GetTypedSet() const;
+
+		// Returns the attribute type we map to
+		[[nodiscard]] UE_API UScriptStruct* GetAttributeType() const;
+
+		// Returns the mapped value type
+		[[nodiscard]] UE_API UScriptStruct* GetValueType() const;
+
+		// Returns the mapping key represented by this set
+		[[nodiscard]] UE_API FAttributeMappingKey GetMappingKey() const;
+
+		// Returns the total allocated size for this set map
+		[[nodiscard]] UE_API int32 GetAllocatedSize() const;
+
+		// Returns the allocator function this set map was initialized with
+		[[nodiscard]] UE_API FReallocFun GetAllocator() const;
+
+		// Returns a copy of this set map using the specified allocator
+		[[nodiscard]] virtual FAttributeSetMap* Duplicate(FReallocFun InReallocFun) const = 0;
+
+		// Initializes every value to its identity, optionally additive
+		virtual void FillWithIdentity(bool bIsAdditive) = 0;
+
+		// Copies the content of this set map into another
+		// Both must have matching typed sets and value types
+		virtual void CopyTo(FAttributeSetMap& Other) const = 0;
+
+		// Sets a value at the specified index using a value setter function
+		// The value setter (and/or caller) is responsible for validating the value type
+		virtual void SetValueWithSetter(FAttributeSetIndex Index, const TFunctionRef<void(UScriptStruct* ValueType, uint8* OutValuePtr)>& ValueSetter) = 0;
+
+		// Gets a value at the specified index using a value getter function
+		// The value getter (and/or caller) is responsible for validating the value type
+		virtual void GetValueWithGetter(FAttributeSetIndex Index, const TFunctionRef<void(UScriptStruct* ValueType, const uint8* ValuePtr)>& ValueGetter) const = 0;
+
+		// Sets a value at the specified index
+		// The called is responsible for validating the value type matches the one contained within this set map
+		template <class ValueType>
+		void SetValue(FAttributeSetIndex Index, const ValueType& Value);
+
+		// Gets a value from the specified index
+		// The called is responsible for validating the value type matches the one contained within this set map
+		template <class ValueType>
+		[[nodiscard]] ValueType GetValue(FAttributeSetIndex Index) const;
+
+	protected:
+		// Default construction for an empty set map
+		UE_API FAttributeSetMap();
+
+		// Constructs a set map but does not allocate its memory, derived types are responsible for it
+		UE_API FAttributeSetMap(const FConstructArgs& Args, int32 AllocatedSize);
+
+		// Destroys the set map and frees its memory
+		UE_API virtual ~FAttributeSetMap();
+
+		// Sets a value at the specified index
+		virtual void SetValueImpl(FAttributeSetIndex Index, const void* ValuePtr) = 0;
+
+		// Gets a value from the specified index
+		virtual void GetValueImpl(FAttributeSetIndex Index, void* OutValuePtr) const = 0;
+
+		// The set whose attributes we map values to
+		FAttributeTypedSetPtr TypedSet;
+
+		// The type of values this set map contains
+		UScriptStruct* ValueTypeStruct = nullptr;
+
+		// The total allocated size of this set map
+		int32 AllocatedSize = 0;
+
+		// The number of values allocated in our buffer
+		int32 NumValues = 0;
+
+		// The base untyped pointer for our values
+		uint8* ValuesPtr = nullptr;
+
+		// The indirect allocator function pointer
+		FReallocFun ReallocFun = nullptr;
+
+		friend class FAttributeCollection;
+		friend class FValueRuntimeRegistry;
+		friend void ReleaseAttributeSetMap(FAttributeSetMap* SetMap);
+	};
+
+	// Encapsulates a mapping of a value type to a specific attribute set (of any attribute type)
+	// Our keys are held by the attribute set we point to (keys are shared) while we hold the
+	// mapped values internally. Set maps and their values are allocated together in a contiguous buffer.
+	// Value types must be a UStruct
+	template<class InValueType>
+	class TAttributeSetMap : public FAttributeSetMap
+	{
+	public:
+		// The type of values held in this map
+		using ValueType = InValueType;
+
+		// Returns a copy of this set map using the specified allocator
+		[[nodiscard]] virtual FAttributeSetMap* Duplicate(FReallocFun InReallocFun) const override;
+
+		// Initializes every value to the specified one
+		void FillWithValue(InValueType Value);
+
+		// Initializes every value to its identity, optionally additive
+		virtual void FillWithIdentity(bool bIsAdditive) override;
+
+		// Copies the content of this set map into another
+		// Both must have matching typed sets and value types
+		virtual void CopyTo(FAttributeSetMap& Other) const override;
+
+		// Value accessors by index
+		[[nodiscard]] InValueType& operator[](FAttributeSetIndex Index);
+
+		// Value accessors by index
+		[[nodiscard]] const InValueType& operator[](FAttributeSetIndex Index) const;
+
+		// Returns a pointer to the contiguous values held within
+		[[nodiscard]] InValueType* GetData();
+
+		// Returns a pointer to the contiguous values held within
+		[[nodiscard]] const InValueType* GetData() const;
+
+		// Returns the total size required for an instance of a set map of this type
+		[[nodiscard]] static int32 CalculateAllocationSize(int32 NumAttributes);
+
+	protected:
+		// Constructs a set map
+		TAttributeSetMap(const FConstructArgs& Args, int32 AllocatedSize);
+
+		// Destroys the set map and frees its memory
+		virtual ~TAttributeSetMap() override;
+
+		virtual void SetValueWithSetter(FAttributeSetIndex Index, const TFunctionRef<void(UScriptStruct* ValueType, uint8* OutValuePtr)>& ValueSetter) override;
+		virtual void GetValueWithGetter(FAttributeSetIndex Index, const TFunctionRef<void(UScriptStruct* ValueType, const uint8* ValuePtr)>& ValueGetter) const override;
+		virtual void SetValueImpl(FAttributeSetIndex Index, const void* ValuePtr) override;
+		virtual void GetValueImpl(FAttributeSetIndex Index, void* OutValuePtr) const override;
+
+		template<class T>
+		friend TAttributeSetMap<T>* MakeAttributeSetMap(const typename TAttributeSetMap<T>::FConstructArgs& Args);
+	};
+
+	// Allocates and constructs a new attribute set map for the specified value type
+	// Set maps and their values are contiguous in memory
+	template<class ValueType>
+	TAttributeSetMap<ValueType>* MakeAttributeSetMap(const typename TAttributeSetMap<ValueType>::FConstructArgs& Args);
+
+	// Releases and destroys the specified attribute set map
+	void ReleaseAttributeSetMap(FAttributeSetMap* SetMap);
+
+	// Cast from an untyped set map to a typed set map
+	template<class ValueType>
+	UE_FORCEINLINE_HINT TAttributeSetMap<ValueType>* Cast(FAttributeSetMap* Src)
+	{
+		return Src != nullptr && ValueType::StaticStruct() == Src->GetValueType() ? static_cast<TAttributeSetMap<ValueType>*>(Src) : nullptr;
+	}
+
+	// Cast from an untyped set map to a typed set map
+	template<class ValueType>
+	UE_FORCEINLINE_HINT const TAttributeSetMap<ValueType>* Cast(const FAttributeSetMap* Src)
+	{
+		return Src != nullptr && ValueType::StaticStruct() == Src->GetValueType() ? static_cast<const TAttributeSetMap<ValueType>*>(Src) : nullptr;
+	}
+
+	// Cast from an untyped set map to a typed set map with an assert if the types mismatch
+	template<class ValueType>
+	inline TAttributeSetMap<ValueType>* CastChecked(FAttributeSetMap* Src)
+	{
+		TAttributeSetMap<ValueType>* CastResult = Src != nullptr && ValueType::StaticStruct() == Src->GetValueType() ? static_cast<TAttributeSetMap<ValueType>*>(Src) : nullptr;
+		check(CastResult);
+		return CastResult;
+	}
+
+	// Cast from an untyped set map to a typed set map with an assert if the types mismatch
+	template<class ValueType>
+	inline const TAttributeSetMap<ValueType>* CastChecked(const FAttributeSetMap* Src)
+	{
+		const TAttributeSetMap<ValueType>* CastResult = Src != nullptr && ValueType::StaticStruct() == Src->GetValueType() ? static_cast<const TAttributeSetMap<ValueType>*>(Src) : nullptr;
+		check(CastResult);
+		return CastResult;
+	}
+
+	//////////////////////////////////////////////////////////////////////////
+	// Implementation
+
+	template<class ValueType>
+	inline void FAttributeSetMap::SetValue(FAttributeSetIndex Index, const ValueType& Value)
+	{
+		checkf(ValueType::StaticStruct() == ValueTypeStruct, TEXT("Value type does not match the set map"));
+		SetValueImpl(Index, &Value);
+	}
+
+	template<class ValueType>
+	inline ValueType FAttributeSetMap::GetValue(FAttributeSetIndex Index) const
+	{
+		checkf(ValueType::StaticStruct() == ValueTypeStruct, TEXT("Value type does not match the set map"));
+
+		ValueType Result;
+		GetValueImpl(Index, &Result);
+
+		return Result;
+	}
+
+	template<class InValueType>
+	inline TAttributeSetMap<InValueType>::TAttributeSetMap(const FConstructArgs& Args, int32 InAllocatedSize)
+		: FAttributeSetMap(Args, InAllocatedSize)
+	{
+		const int32 NumAttributes = Args.TypedSet->Num();
+
+		check(InAllocatedSize == TAttributeSetMap<InValueType>::CalculateAllocationSize(NumAttributes));
+
+		if (NumAttributes != 0)
+		{
+			ValuesPtr = Align(reinterpret_cast<uint8*>(this) + sizeof(TAttributeSetMap<InValueType>), alignof(InValueType));
+			NumValues = NumAttributes;
+
+			DefaultConstructItems<InValueType>(ValuesPtr, NumAttributes);
+		}
+	}
+
+	template<class InValueType>
+	inline TAttributeSetMap<InValueType>::~TAttributeSetMap()
+	{
+		if (ValuesPtr != nullptr)
+		{
+			DestructItems<InValueType>(GetData(), NumValues);
+
+			ValuesPtr = nullptr;
+			NumValues = 0;
+		}
+	}
+
+	template<class InValueType>
+	inline FAttributeSetMap* TAttributeSetMap<InValueType>::Duplicate(FReallocFun InReallocFun) const
+	{
+		TAttributeSetMap<InValueType>* Copy = MakeAttributeSetMap<InValueType>(FConstructArgs(TypedSet, ValueTypeStruct, InReallocFun));
+
+		if (NumValues != 0)
+		{
+			ConstructItems<InValueType>(Copy->GetData(), GetData(), NumValues);
+		}
+
+		return Copy;
+	}
+
+	template<class InValueType>
+	inline void TAttributeSetMap<InValueType>::FillWithValue(InValueType Value)
+	{
+		if (NumValues == 0)
+		{
+			return;
+		}
+
+		InValueType* ValuePtr = GetData();
+		InValueType* ValuePtrEnd = ValuePtr + NumValues;
+		while (ValuePtr < ValuePtrEnd)
+		{
+			*ValuePtr = Value;
+
+			ValuePtr++;
+		}
+	}
+
+	template<class InValueType>
+	inline void TAttributeSetMap<InValueType>::FillWithIdentity(bool bIsAdditive)
+	{
+		InValueType Identity{};
+
+		if constexpr (std::is_base_of_v<FTransformAnimationAttribute, InValueType>)
+		{
+			// Transforms have a special additive identity
+			if (bIsAdditive)
+			{
+				Identity = FTransformAnimationAttribute{ FTransform(FQuat::Identity, FVector::ZeroVector, FVector::ZeroVector) };
+			}
+			else
+			{
+				Identity = FTransformAnimationAttribute{ FTransform(FQuat::Identity, FVector::ZeroVector, FVector::OneVector) };
+			}
+		}
+		else
+		{
+			// Everything else uses the default value
+		}
+
+		FillWithValue(Identity);
+	}
+
+	template<class InValueType>
+	inline void TAttributeSetMap<InValueType>::CopyTo(FAttributeSetMap& Other) const
+	{
+		checkf(TypedSet == Other.GetTypedSet() && ValueTypeStruct == Other.GetValueType(), TEXT("Typed set and value type must match to allow copying"));
+		check(Other.Num() == Num());	// If the typed sets match, their sizes does too
+
+		TAttributeSetMap<InValueType>& Copy = reinterpret_cast<TAttributeSetMap<InValueType>&>(Other);
+
+		CopyAssignItems<InValueType>(Copy.GetData(), GetData(), NumValues);
+	}
+
+	template<class InValueType>
+	inline InValueType& TAttributeSetMap<InValueType>::operator[](FAttributeSetIndex Index)
+	{
+		checkf((Index.GetInt() >= 0) & (Index.GetInt() < NumValues), TEXT("Map index out of bounds: %lld into an map of size %lld"), (long long)Index.GetInt(), (long long)NumValues); // & for one branch
+		return GetData()[Index.GetInt()];
+	}
+
+	template<class InValueType>
+	inline const InValueType& TAttributeSetMap<InValueType>::operator[](FAttributeSetIndex Index) const
+	{
+		checkf((Index.GetInt() >= 0) & (Index.GetInt() < NumValues), TEXT("Map index out of bounds: %lld into an map of size %lld"), (long long)Index.GetInt(), (long long)NumValues); // & for one branch
+		return GetData()[Index.GetInt()];
+	}
+
+	template<class InValueType>
+	inline InValueType* TAttributeSetMap<InValueType>::GetData()
+	{
+		return reinterpret_cast<InValueType*>(ValuesPtr);
+	}
+
+	template<class InValueType>
+	inline const InValueType* TAttributeSetMap<InValueType>::GetData() const
+	{
+		return reinterpret_cast<const InValueType*>(ValuesPtr);
+	}
+
+	template<class InValueType>
+	inline void TAttributeSetMap<InValueType>::SetValueWithSetter(FAttributeSetIndex Index, const TFunctionRef<void(UScriptStruct* ValueType, uint8* OutValuePtr)>& ValueSetter)
+	{
+		checkf((Index.GetInt() >= 0) & (Index.GetInt() < NumValues), TEXT("Map index out of bounds: %lld into an map of size %lld"), (long long)Index.GetInt(), (long long)NumValues); // & for one branch
+		InValueType* ValuePtr = GetData() + Index.GetInt();
+		ValueSetter(ValueTypeStruct, reinterpret_cast<uint8*>(ValuePtr));
+	}
+
+	template<class InValueType>
+	inline void TAttributeSetMap<InValueType>::GetValueWithGetter(FAttributeSetIndex Index, const TFunctionRef<void(UScriptStruct* ValueType, const uint8* ValuePtr)>& ValueGetter) const
+	{
+		checkf((Index.GetInt() >= 0) & (Index.GetInt() < NumValues), TEXT("Map index out of bounds: %lld into an map of size %lld"), (long long)Index.GetInt(), (long long)NumValues); // & for one branch
+		const InValueType* ValuePtr = GetData() + Index.GetInt();
+		ValueGetter(ValueTypeStruct, reinterpret_cast<const uint8*>(ValuePtr));
+	}
+
+	template<class InValueType>
+	inline void TAttributeSetMap<InValueType>::SetValueImpl(FAttributeSetIndex Index, const void* ValuePtr)
+	{
+		checkf((Index.GetInt() >= 0) & (Index.GetInt() < NumValues), TEXT("Map index out of bounds: %lld into an map of size %lld"), (long long)Index.GetInt(), (long long)NumValues); // & for one branch
+		const InValueType* Value = reinterpret_cast<const InValueType*>(ValuePtr);
+		GetData()[Index.GetInt()] = *Value;
+	}
+
+	template<class InValueType>
+	inline void TAttributeSetMap<InValueType>::GetValueImpl(FAttributeSetIndex Index, void* OutValuePtr) const
+	{
+		checkf((Index.GetInt() >= 0) & (Index.GetInt() < NumValues), TEXT("Map index out of bounds: %lld into an map of size %lld"), (long long)Index.GetInt(), (long long)NumValues); // & for one branch
+		InValueType* OutValue = reinterpret_cast<InValueType*>(OutValuePtr);
+		*OutValue = GetData()[Index.GetInt()];
+	}
+
+	template<class InValueType>
+	inline int32 TAttributeSetMap<InValueType>::CalculateAllocationSize(int32 NumAttributes)
+	{
+		return Align(sizeof(TAttributeSetMap<InValueType>), alignof(InValueType)) + sizeof(InValueType) * NumAttributes;
+	}
+
+	template<class ValueType>
+	inline TAttributeSetMap<ValueType>* MakeAttributeSetMap(const typename TAttributeSetMap<ValueType>::FConstructArgs& Args)
+	{
+		// We allocate the set map and its values in a contiguous buffer part of a single allocation
+		// As such, the set map struct ends up being a small header
+		const int32 AllocationSize = TAttributeSetMap<ValueType>::CalculateAllocationSize(Args.TypedSet->Num());
+
+		// Allocate our single buffer
+		uint8* Ptr = (*Args.ReallocFun)(nullptr, 0, AllocationSize);
+
+		// Construct our set map over it, the specialized constructor is responsible for partitioning the rest of the buffer
+		return new(Ptr) TAttributeSetMap<ValueType>(Args, AllocationSize);
+	}
+
+	inline void ReleaseAttributeSetMap(FAttributeSetMap* SetMap)
+	{
+		if (SetMap == nullptr)
+		{
+			return;
+		}
+
+		const int32 AllocationSize = SetMap->AllocatedSize;
+		const FReallocFun ReallocFun = SetMap->ReallocFun;
+
+		SetMap->~FAttributeSetMap();
+
+		(*ReallocFun)(reinterpret_cast<uint8*>(SetMap), AllocationSize, 0);
+	}
+}
+
+#undef UE_API
diff --git a/Engine/Plugins/Experimental/UAF/UAF/Source/UAF/Public/UAF/ValueRuntime/IndirectAllocator.h b/Engine/Plugins/Experimental/UAF/UAF/Source/UAF/Public/UAF/ValueRuntime/IndirectAllocator.h
new file mode 100644
--- /dev/null
+++ b/Engine/Plugins/Experimental/UAF/UAF/Source/UAF/Public/UAF/ValueRuntime/IndirectAllocator.h
@@ -0,0 +1,68 @@
+// Copyright Epic Games, Inc. All Rights Reserved.
+
+#pragma once
+
+#include "CoreMinimal.h"
+#include "Containers/ContainerAllocationPolicies.h"
+#include "Misc/MemStack.h"
+
+namespace UE::UAF
+{
+	// Use allocator pointer instead of template argument
+	// The allocator should behave as follows:
+	//    - When Ptr is nullptr (and OldSize is 0), we are allocating memory
+	//    - When Ptr is valid but NewSize is 0, we are freeing memory
+	//    - When Ptr is valid and NewSize is not 0, we are re-allocating memory
+	//    - It assumes that the elements that live in said memory can be trivially copyable
+	//    - Allocations use default alignment
+	typedef uint8* (*FReallocFun)(uint8* OldPtr, uint32 OldSize, uint32 NewSize);
+
+	// This type can be specialized for template style allocators for easy interop
+	// with the engine allocator types
+	template<class AllocatorType>
+	struct FAllocatorTypeTrait
+	{
+		// Specialize this struct for your allocator by implementing this function
+		static uint8* Realloc(uint8* OldPtr, uint32 OldSize, uint32 NewSize) = delete;
+	};
+
+	template<>
+	struct FAllocatorTypeTrait<FDefaultAllocator>
+	{
+		static uint8* Realloc(uint8* OldPtr, uint32 OldSize, uint32 NewSize)
+		{
+			if (OldPtr == nullptr && NewSize == 0)
+			{
+				// Nothing to do if we weren't previously allocated and we aren't requesting a valid size
+				return nullptr;
+			}
+
+			return (uint8*)FMemory::Realloc(OldPtr, NewSize, DEFAULT_ALIGNMENT);
+		}
+	};
+
+	template<>
+	struct FAllocatorTypeTrait<TMemStackAllocator<>>
+	{
+		static uint8* Realloc(uint8* OldPtr, uint32 OldSize, uint32 NewSize)
+		{
+			if (NewSize == 0)
+			{
+				// Nothing to do if we are freeing memory
+				return nullptr;
+			}
+
+			// Allocate memory from the stack.
+			uint8* NewPtr = FMemStack::Get().PushBytes(NewSize, DEFAULT_ALIGNMENT);
+
+			// If we are growing an existing allocation, copy from the old
+			if (OldPtr != nullptr)
+			{
+				const uint32 NumToCopy = FMath::Min(NewSize, OldSize);
+				FMemory::Memcpy(NewPtr, OldPtr, NumToCopy);
+			}
+
+			return NewPtr;
+		}
+	};
+}
diff --git a/Engine/Plugins/Experimental/UAF/UAF/Source/UAF/Public/UAF/ValueRuntime/Transformers/Accumulate.h b/Engine/Plugins/Experimental/UAF/UAF/Source/UAF/Public/UAF/ValueRuntime/Transformers/Accumulate.h
new file mode 100644
--- /dev/null
+++ b/Engine/Plugins/Experimental/UAF/UAF/Source/UAF/Public/UAF/ValueRuntime/Transformers/Accumulate.h
@@ -0,0 +1,64 @@
+// Copyright Epic Games, Inc. All Rights Reserved.
+
+#pragma once
+
+#include "CoreMinimal.h"
+
+#include "UAF/Attributes/EngineAttributes.h"
+#include "UAF/ValueRuntime/Transformers/Transformer.h"
+
+#define UE_API UAF_API
+
+namespace UE::UAF::Transformers
+{
+	// Attribute value transformer to accumulate values
+	// Accumulate is the repeating step within an N-step interpolation: Overwrite, Accumulate0, Accumulate1, ...
+	// Make sure to also implement variations for Overwrite and Interpolate
+	struct FAccumulate final : public FAttributeValueTransformer
+	{
+		static UE_API FName TransformerName;
+
+		// If the output buffer is nullptr, an allocator is provided and the implementation is responsible for allocating the output buffer and returning it
+		// Per-value map is optional, when not present, the default weight should be used
+		// Unbound values do not support per-value weights
+		using FTransformSetMapFunc = void (*)(const FAttributeSetMap* SetMapA, const FAttributeSetMap* SetMapB, const FAttributeSetMap* PerValueWeightMapB, float DefaultWeightB, FAttributeSetMap*& OutputSetMap, FReallocFun OutputAllocator);
+		using FTransformMapFunc = void (*)(const FAttributeMap* MapA, const FAttributeMap* MapB, float WeightB, FAttributeMap*& OutputMap, FReallocFun OutputAllocator);
+
+		// Accumulate attribute values
+		// Output = A + (B * WeightB)
+		// This function supports outputting to either input
+		static UE_API void Apply(const FAttributeValueTransformerMapPtr& TransformerMap, const FAttributeCollection& InputA, const FAttributeCollection& InputB, float WeightB, FAttributeCollection& Output);
+
+		// Accumulate attribute values
+		// If a weight isn't specified, the default weight will be used
+		// Output = A + (B * (HasValueWeightB ? ValueWeightB : DefaultWeightB))
+		// This function supports outputting to either input
+		static UE_API void Apply(const FAttributeValueTransformerMapPtr& TransformerMap, const FAttributeCollection& InputA, const FAttributeCollection& InputB, const FAttributeCollection& PerValueWeightsB, float DefaultWeightB, FAttributeCollection& Output);
+	};
+
+	//////////////////////////////////////////////////////////////////////////
+	// Engine built-in value type specializations
+
+	namespace Private
+	{
+		struct FAccumulate_BoneTransformAttribute
+		{
+			using TransformerType = FAccumulate;
+			using ValueType = FBoneTransformAnimationAttribute;
+
+			static void TransformSetMap(const FAttributeSetMap* SetMapA, const FAttributeSetMap* SetMapB, const FAttributeSetMap* PerValueWeightMapB, float DefaultWeightB, FAttributeSetMap*& OutputSetMap, FReallocFun OutputAllocator);
+			static void TransformMap(const FAttributeMap* MapA, const FAttributeMap* MapB, float WeightB, FAttributeMap*& OutputMap, FReallocFun OutputAllocator);
+		};
+
+		struct FAccumulate_FloatAttribute
+		{
+			using TransformerType = FAccumulate;
+			using ValueType = FFloatAnimationAttribute;
+
+			static void TransformSetMap(const FAttributeSetMap* SetMapA, const FAttributeSetMap* SetMapB, const FAttributeSetMap* PerValueWeightMapB, float DefaultWeightB, FAttributeSetMap*& OutputSetMap, FReallocFun OutputAllocator);
+			static void TransformMap(const FAttributeMap* MapA, const FAttributeMap* MapB, float WeightB, FAttributeMap*& OutputMap, FReallocFun OutputAllocator);
+		};
+	}
+}
+
+#undef UE_API
diff --git a/Engine/Plugins/Experimental/UAF/UAF/Source/UAF/Public/UAF/ValueRuntime/Transformers/Interpolate.h b/Engine/Plugins/Experimental/UAF/UAF/Source/UAF/Public/UAF/ValueRuntime/Transformers/Interpolate.h
new file mode 100644
--- /dev/null
+++ b/Engine/Plugins/Experimental/UAF/UAF/Source/UAF/Public/UAF/ValueRuntime/Transformers/Interpolate.h
@@ -0,0 +1,65 @@
+// Copyright Epic Games, Inc. All Rights Reserved.
+
+#pragma once
+
+#include "CoreMinimal.h"
+
+#include "UAF/Attributes/EngineAttributes.h"
+#include "UAF/ValueRuntime/Transformers/Transformer.h"
+
+#define UE_API UAF_API
+
+namespace UE::UAF::Transformers
+{
+	// Attribute value transformer to interpolate values
+	// Make sure to also implement variations for Overwrite and Accumulate
+	struct FInterpolate final : public FAttributeValueTransformer
+	{
+		static UE_API FName TransformerName;
+
+		// If the output buffer is nullptr, an allocator is provided and the implementation is responsible for allocating the output buffer and returning it
+		// Note that the input/output buffers can be the same buffer
+		// Per-value map is optional, when not present, the default weight should be used
+		// Unbound values do not support per-value weights
+		using FTransformSetMapFunc = void (*)(const FAttributeSetMap* SetMapA, const FAttributeSetMap* SetMapB, const FAttributeSetMap* PerValueWeightMapB, float DefaultWeightB, FAttributeSetMap*& OutputSetMap, FReallocFun OutputAllocator);
+		using FTransformMapFunc = void (*)(const FAttributeMap* MapA, const FAttributeMap* MapB, float WeightB, FAttributeMap*& OutputMap, FReallocFun OutputAllocator);
+
+		// Interpolate attribute values
+		// Output = lerp(A, B, WeightB)
+		// Output = (A * (1.0 - WeightB)) + (B * WeightB)
+		// This function supports outputting to either input
+		static UE_API void Apply(const FAttributeValueTransformerMapPtr& TransformerMap, const FAttributeCollection& InputA, const FAttributeCollection& InputB, float WeightB, FAttributeCollection& Output);
+
+		// Interpolate attribute values
+		// If a weight isn't specified, the default weight will be used
+		// Output = lerp(A, B, HasValueWeightB ? ValueWeightB : DefaultWeightB)
+		// This function supports outputting to either input
+		static UE_API void Apply(const FAttributeValueTransformerMapPtr& TransformerMap, const FAttributeCollection& InputA, const FAttributeCollection& InputB, const FAttributeCollection& PerValueWeightsB, float DefaultWeightB, FAttributeCollection& Output);
+	};
+
+	//////////////////////////////////////////////////////////////////////////
+	// Engine built-in value type specializations
+
+	namespace Private
+	{
+		struct FInterpolate_BoneTransformAttribute
+		{
+			using TransformerType = FInterpolate;
+			using ValueType = FBoneTransformAnimationAttribute;
+
+			static void TransformSetMap(const FAttributeSetMap* SetMapA, const FAttributeSetMap* SetMapB, const FAttributeSetMap* PerValueWeightMapB, float DefaultWeightB, FAttributeSetMap*& OutputSetMap, FReallocFun OutputAllocator);
+			static void TransformMap(const FAttributeMap* MapA, const FAttributeMap* MapB, float WeightB, FAttributeMap*& OutputMap, FReallocFun OutputAllocator);
+		};
+
+		struct FInterpolate_FloatAttribute
+		{
+			using TransformerType = FInterpolate;
+			using ValueType = FFloatAnimationAttribute;
+
+			static void TransformSetMap(const FAttributeSetMap* SetMapA, const FAttributeSetMap* SetMapB, const FAttributeSetMap* PerValueWeightMapB, float DefaultWeightB, FAttributeSetMap*& OutputSetMap, FReallocFun OutputAllocator);
+			static void TransformMap(const FAttributeMap* MapA, const FAttributeMap* MapB, float WeightB, FAttributeMap*& OutputMap, FReallocFun OutputAllocator);
+		};
+	}
+}
+
+#undef UE_API
diff --git a/Engine/Plugins/Experimental/UAF/UAF/Source/UAF/Public/UAF/ValueRuntime/Transformers/Overwrite.h b/Engine/Plugins/Experimental/UAF/UAF/Source/UAF/Public/UAF/ValueRuntime/Transformers/Overwrite.h
new file mode 100644
--- /dev/null
+++ b/Engine/Plugins/Experimental/UAF/UAF/Source/UAF/Public/UAF/ValueRuntime/Transformers/Overwrite.h
@@ -0,0 +1,65 @@
+// Copyright Epic Games, Inc. All Rights Reserved.
+
+#pragma once
+
+#include "CoreMinimal.h"
+
+#include "UAF/Attributes/EngineAttributes.h"
+#include "UAF/ValueRuntime/Transformers/Transformer.h"
+
+#define UE_API UAF_API
+
+namespace UE::UAF::Transformers
+{
+	// Attribute value transformer to overwrite values
+	// Overwrite is the initial step within an N-step interpolation: Overwrite, Accumulate0, Accumulate1, ...
+	// Make sure to also implement variations for Accumulate and Interpolate
+	struct FOverwrite final : public FAttributeValueTransformer
+	{
+		static UE_API FName TransformerName;
+
+		// If the output buffer is nullptr, an allocator is provided and the implementation is responsible for allocating the output buffer and returning it
+		// Note that the input/output buffers can be the same buffer
+		// Per-value map is optional, when not present, the default weight should be used
+		// Unbound values do not support per-value weights
+		using FTransformSetMapFunc = void (*)(const FAttributeSetMap* InputSetMap, const FAttributeSetMap* PerValueWeightMap, float DefaultWeight, FAttributeSetMap*& OutputSetMap, FReallocFun OutputAllocator);
+		using FTransformMapFunc = void (*)(const FAttributeMap* InputMap, float Weight, FAttributeMap*& OutputMap, FReallocFun OutputAllocator);
+
+		// Overwrite attribute values
+		// Output = (A * Weight)
+		// This function supports outputting to its input
+		static UE_API void Apply(const FAttributeValueTransformerMapPtr& TransformerMap, const FAttributeCollection& Input, float Weight, FAttributeCollection& Output);
+
+		// Overwrite attribute values using per value weights (must be of type FFloatAnimationAttribute)
+		// If a weight isn't specified, the default weight will be used
+		// Output = A * (HasValueWeight ? ValueWeight : DefaultWeight)
+		// This function supports outputting to its input
+		static UE_API void Apply(const FAttributeValueTransformerMapPtr& TransformerMap, const FAttributeCollection& Input, const FAttributeCollection& PerValueWeights, float DefaultWeight, FAttributeCollection& Output);
+	};
+
+	//////////////////////////////////////////////////////////////////////////
+	// Engine built-in value type specializations
+
+	namespace Private
+	{
+		struct FOverwrite_BoneTransformAttribute
+		{
+			using TransformerType = FOverwrite;
+			using ValueType = FBoneTransformAnimationAttribute;
+
+			static void TransformSetMap(const FAttributeSetMap* InputSetMap, const FAttributeSetMap* PerValueWeightMap, float DefaultWeight, FAttributeSetMap*& OutputSetMap, FReallocFun OutputAllocator);
+			static void TransformMap(const FAttributeMap* InputMap, float Weight, FAttributeMap*& OutputMap, FReallocFun OutputAllocator);
+		};
+
+		struct FOverwrite_FloatAttribute
+		{
+			using TransformerType = FOverwrite;
+			using ValueType = FFloatAnimationAttribute;
+
+			static void TransformSetMap(const FAttributeSetMap* InputSetMap, const FAttributeSetMap* PerValueWeightMap, float DefaultWeight, FAttributeSetMap*& OutputSetMap, FReallocFun OutputAllocator);
+			static void TransformMap(const FAttributeMap* InputMap, float Weight, FAttributeMap*& OutputMap, FReallocFun OutputAllocator);
+		};
+	}
+}
+
+#undef UE_API
diff --git a/Engine/Plugins/Experimental/UAF/UAF/Source/UAF/Public/UAF/ValueRuntime/Transformers/Transformer.h b/Engine/Plugins/Experimental/UAF/UAF/Source/UAF/Public/UAF/ValueRuntime/Transformers/Transformer.h
new file mode 100644
--- /dev/null
+++ b/Engine/Plugins/Experimental/UAF/UAF/Source/UAF/Public/UAF/ValueRuntime/Transformers/Transformer.h
@@ -0,0 +1,210 @@
+// Copyright Epic Games, Inc. All Rights Reserved.
+
+#pragma once
+
+#include "CoreMinimal.h"
+#include "UAF/ValueRuntime/AttributeCollection.h"
+
+#define UE_API UAF_API
+
+namespace UE::UAF
+{
+	class FValueRuntimeRegistry;
+	class FAttributeValueTransformerList;
+
+	// A raw function pointer for all transformers (static functions)
+	// Only used for storage purposes, NEVER CALL THIS!
+	using FRawTransformerFunc = void(*)();
+
+	// Attribute value transformer base type
+	// All transformer types must derive from this
+	struct FAttributeValueTransformer
+	{
+		// The name of the transformer (exposed in derived types)
+		//static FName TransformerName;
+
+		// An optional function signature of the transformer for set maps (exposed in derived types)
+		//using FTransformSetMapFunc = void (*)(...);
+
+		// An optional function signature of the transformer for name/value maps (exposed in derived types)
+		//using FTransformMapFunc = void (*)(...);
+	};
+
+	// An iterator over a transformer list within the attribute runtime registry
+	class FAttributeValueTransformerListIterator
+	{
+	public:
+		[[nodiscard]] FAttributeValueTransformerListIterator(const FAttributeValueTransformerList& TransformerList, int32 StartOffset, int32 Count);
+
+		// Increments and moves the iterator to the next transformer entry
+		FAttributeValueTransformerListIterator& operator++();
+
+		// Returns whether or not the iterator still contains values
+		[[nodiscard]] explicit operator bool() const;
+
+		// Returns the value type that the transformer operates on for the current entry
+		[[nodiscard]] UScriptStruct* GetValueType() const;
+
+		// Returns the transformer for the current entry
+		[[nodiscard]] FRawTransformerFunc GetTransformerFunc() const;
+
+	private:
+		UScriptStruct* const* ValueTypePtr = nullptr;
+		const FRawTransformerFunc* TransformerFuncPtr = nullptr;
+		UScriptStruct* const* ValueTypeEndPtr = nullptr;
+	};
+
+	// A list of transformer specializations that operators over various value types.
+	// Transformers have a base type (e.g. interpolate, sanitize) and can be specialized to operate on arbitrary
+	// value types (e.g. FFloatAnimationAttribute). They can also transform two types of containers: set maps
+	// and name/value maps.
+	class FAttributeValueTransformerList
+	{
+	public:
+		// Creates an empty transformer list
+		UE_API FAttributeValueTransformerList();
+
+		// Creates a transformer list for the specified transformer type
+		[[nodiscard]] UE_API explicit FAttributeValueTransformerList(FName TransformerName);
+
+		// Returns the transformer name
+		[[nodiscard]] UE_API FName GetTransformerName() const;
+
+		// Returns the number of set map transformers in this list
+		[[nodiscard]] int32 NumSetMapTransformers() const { return SetMapTransformerCount; }
+
+		// Returns the number of name/value map transformers in this list
+		[[nodiscard]] int32 NumMapTransformers() const { return ValueTypes.Num() - SetMapTransformerCount; }
+
+		// Adds a new set map transformer entry
+		UE_API bool AddSetMapTransformer(UScriptStruct* ValueType, FRawTransformerFunc TransformerFunc);
+
+		// Adds a new name/value map transformer entry
+		UE_API bool AddMapTransformer(UScriptStruct* ValueType, FRawTransformerFunc TransformerFunc);
+
+		// Returns an iterator over all set map transformer specializations within
+		[[nodiscard]] FAttributeValueTransformerListIterator CreateSetMapTransformerIterator() const { return FAttributeValueTransformerListIterator(*this, 0, NumSetMapTransformers()); }
+
+		// Returns an iterator over all name/value map transformer specializations within
+		[[nodiscard]] FAttributeValueTransformerListIterator CreateMapTransformerIterator() const { return FAttributeValueTransformerListIterator(*this, NumSetMapTransformers(), NumMapTransformers()); }
+
+	private:
+		// The type of the transformer (e.g. interpolate, sanitize)
+		FName TransformerName;
+
+		// The number of set map transformers in this list
+		int32 SetMapTransformerCount = 0;
+
+		// List of attribute value types for which we have transformers registered (e.g. FTransformAnimationAttribute, FFloatAnimationAttribute) (sorted by container type, then by value type)
+		// This list is split in two, set map transformers come first, then map transformers
+		// Each sub-list is sorted by value type
+		TArray<UScriptStruct*> ValueTypes;
+
+		// List of transformer functions for each value type (sorted by container type, then by value type)
+		// This list is split in two, set map transformers come first, then map transformers
+		// Each sub-list is sorted by value type
+		TArray<FRawTransformerFunc> Transformers;
+
+		friend FValueRuntimeRegistry;
+		friend FAttributeValueTransformerListIterator;
+	};
+
+	using FAttributeValueTransformerMapPtr = TSharedPtr<TMap<FName, FAttributeValueTransformerList>>;
+
+	// An operator structure to join attribute value transformer lists and attribute collections using InnerJoin/OuterJoin/etc
+	// Ensures that the predicate is called for each set map with a value type that has a corresponding transformer
+	struct FAttributeValueTransformerListWithAttributeCollectionSetMapJoinOp
+	{
+		// Key sorting uses signed comparison, use the largest signed value regardless of pointer size
+		static constexpr intptr_t WILDCARD_ATTRIBUTE_TYPE_VALUE = static_cast<intptr_t>(static_cast<uintptr_t>(~0ULL) >> 1);
+
+		// We enable fuzzy matching by implementing its API
+		bool IsMatch(FAttributeMappingKey ReferenceKey, FAttributeMappingKey TestKey) const
+		{
+			// Our wildcard key ensures that it is larger than all other keys meaning it will only ever be the
+			// reference key if we've matched all other smaller values (all valid ones)
+			if (reinterpret_cast<intptr_t>(TestKey.GetAttributeType()) == WILDCARD_ATTRIBUTE_TYPE_VALUE)
+			{
+				// If our test value is our wildcard value, we match if our value types match
+				return ReferenceKey.GetValueType() == TestKey.GetValueType();
+			}
+
+			// Otherwise we test normally
+			return ReferenceKey == TestKey;
+		}
+
+		FAttributeMappingKey GetLargestKey() const { return FAttributeMappingKey::LARGEST_VALUE; }
+
+		// FAttributeValueTransformerListIterator entries
+		FAttributeMappingKey GetKey(const FAttributeValueTransformerListIterator& It) const
+		{
+			// We use a wildcard entry for the attribute type within our key to ensure it is always larger than all other valid values
+			return FAttributeMappingKey(reinterpret_cast<UScriptStruct*>(WILDCARD_ATTRIBUTE_TYPE_VALUE), It.GetValueType());
+		}
+
+		FRawTransformerFunc GetValue(const FAttributeValueTransformerListIterator& It) const { return It.GetTransformerFunc(); }
+
+		// Normal FAttributeCollectionSetMapJoinOp
+		FAttributeMappingKey GetKey(const FAttributeCollection::FSetMapIterator& It) const { return It.GetKey(); }
+		FAttributeMappingKey GetKey(const FAttributeCollection::FConstSetMapIterator& It) const { return It.GetKey(); }
+
+		FAttributeSetMap* GetValue(const FAttributeCollection::FSetMapIterator& It) const { return It.GetMap(); }
+		const FAttributeSetMap* GetValue(const FAttributeCollection::FConstSetMapIterator& It) const { return It.GetMap(); }
+	};
+
+	// An operator structure to join attribute value transformer lists and attribute collections using InnerJoin/OuterJoin/etc
+	// Ensures that the predicate is called for each set map with a value type that has a corresponding transformer
+	struct FAttributeValueTransformerListWithAttributeCollectionMapJoinOp
+	{
+		// Key sorting uses signed comparison, use the largest signed value regardless of pointer size
+		static constexpr intptr_t KEY_LARGEST_VALUE = static_cast<intptr_t>(static_cast<uintptr_t>(~0ULL) >> 1);
+
+		UScriptStruct* GetLargestKey() const { return reinterpret_cast<UScriptStruct*>(KEY_LARGEST_VALUE); }
+
+		// FAttributeValueTransformerListIterator entries
+		UScriptStruct* GetKey(const FAttributeValueTransformerListIterator& It) const { return It.GetValueType(); }
+		FRawTransformerFunc GetValue(const FAttributeValueTransformerListIterator& It) const { return It.GetTransformerFunc(); }
+
+		// Normal FAttributeCollectionMapJoinOp
+		UScriptStruct* GetKey(const FAttributeCollection::FMapIterator& It) const { return It.GetValueType(); }
+		UScriptStruct* GetKey(const FAttributeCollection::FConstMapIterator& It) const { return It.GetValueType(); }
+
+		FAttributeMap* GetValue(const FAttributeCollection::FMapIterator& It) const { return It.GetMap(); }
+		const FAttributeMap* GetValue(const FAttributeCollection::FConstMapIterator& It) const { return It.GetMap(); }
+	};
+
+	//////////////////////////////////////////////////////////////////////////
+	// Implementation
+
+	inline FAttributeValueTransformerListIterator::FAttributeValueTransformerListIterator(const FAttributeValueTransformerList& TransformerList, int32 StartOffset, int32 Count)
+		: ValueTypePtr(TransformerList.ValueTypes.GetData() + StartOffset)
+		, TransformerFuncPtr(TransformerList.Transformers.GetData() + StartOffset)
+		, ValueTypeEndPtr(ValueTypePtr + Count)
+	{
+	}
+
+	inline FAttributeValueTransformerListIterator& FAttributeValueTransformerListIterator::operator++()
+	{
+		++ValueTypePtr;
+		++TransformerFuncPtr;
+
+		return *this;
+	}
+
+	inline FAttributeValueTransformerListIterator::operator bool() const
+	{
+		return ValueTypePtr < ValueTypeEndPtr;
+	}
+
+	inline UScriptStruct* FAttributeValueTransformerListIterator::GetValueType() const
+	{
+		return *ValueTypePtr;
+	}
+
+	inline FRawTransformerFunc FAttributeValueTransformerListIterator::GetTransformerFunc() const
+	{
+		return *TransformerFuncPtr;
+	}
+}
+
+#undef UE_API
