Rename UAF ValueRuntime types

#Synced-CL 49164946

diff --git a/Engine/Plugins/Experimental/UAF/UAF/Source/UAF/Private/UAF/Attributes/AttributeBindingData.cpp b/Engine/Plugins/Experimental/UAF/UAF/Source/UAF/Private/UAF/Attributes/AttributeBindingData.cpp
--- a/Engine/Plugins/Experimental/UAF/UAF/Source/UAF/Private/UAF/Attributes/AttributeBindingData.cpp
+++ b/Engine/Plugins/Experimental/UAF/UAF/Source/UAF/Private/UAF/Attributes/AttributeBindingData.cpp
@@ -348,7 +348,7 @@ namespace UE::UAF
 					}
 
 					FAttributeMappingKey MappingKey = FAttributeMappingKey::MakeFrom(TypedSetIt->GetType()).To<FBoneTransformAnimationAttribute>();
-					TAttributeSetMap<FBoneTransformAnimationAttribute>* DefaultValues = DefaultAttributeValues->AddSetMap<FBoneTransformAnimationAttribute>(MappingKey);
+					TBoundValueMap<FBoneTransformAnimationAttribute>* DefaultValues = DefaultAttributeValues->AddSetMap<FBoneTransformAnimationAttribute>(MappingKey);
 
 					// Set every value within our map
 					const int32 NumValues = DefaultValues->Num();
diff --git a/Engine/Plugins/Experimental/UAF/UAF/Source/UAF/Private/UAF/Attributes/EngineAttributes.cpp b/Engine/Plugins/Experimental/UAF/UAF/Source/UAF/Private/UAF/Attributes/EngineAttributes.cpp
--- a/Engine/Plugins/Experimental/UAF/UAF/Source/UAF/Private/UAF/Attributes/EngineAttributes.cpp
+++ b/Engine/Plugins/Experimental/UAF/UAF/Source/UAF/Private/UAF/Attributes/EngineAttributes.cpp
@@ -16,32 +16,32 @@ namespace UE::UAF
 	{
 		FValueRuntimeRegistry& RuntimeRegistry = FValueRuntimeRegistry::Get();
 
-		RuntimeRegistry.RegisterSetMapInitializer<FFloatAnimationAttribute>();
-		RuntimeRegistry.RegisterSetMapInitializer<FIntegerAnimationAttribute>();
-		RuntimeRegistry.RegisterSetMapInitializer<FStringAnimationAttribute>();
-		RuntimeRegistry.RegisterSetMapInitializer<FBoneTransformAnimationAttribute>();
+		RuntimeRegistry.RegisterBoundValueMapInitializer<FFloatAnimationAttribute>();
+		RuntimeRegistry.RegisterBoundValueMapInitializer<FIntegerAnimationAttribute>();
+		RuntimeRegistry.RegisterBoundValueMapInitializer<FStringAnimationAttribute>();
+		RuntimeRegistry.RegisterBoundValueMapInitializer<FBoneTransformAnimationAttribute>();
 
-		RuntimeRegistry.RegisterMapInitializer<FFloatAnimationAttribute>();
-		RuntimeRegistry.RegisterMapInitializer<FIntegerAnimationAttribute>();
-		RuntimeRegistry.RegisterMapInitializer<FStringAnimationAttribute>();
-		RuntimeRegistry.RegisterMapInitializer<FBoneTransformAnimationAttribute>();
+		RuntimeRegistry.RegisterUnboundValueMapInitializer<FFloatAnimationAttribute>();
+		RuntimeRegistry.RegisterUnboundValueMapInitializer<FIntegerAnimationAttribute>();
+		RuntimeRegistry.RegisterUnboundValueMapInitializer<FStringAnimationAttribute>();
+		RuntimeRegistry.RegisterUnboundValueMapInitializer<FBoneTransformAnimationAttribute>();
 
-		RuntimeRegistry.RegisterAttributeValueTransformer<Transformers::Private::FAccumulate_BoneTransformAttribute>();
-		RuntimeRegistry.RegisterAttributeValueTransformer<Transformers::Private::FAccumulate_FloatAttribute>();
-		RuntimeRegistry.RegisterAttributeValueTransformer<Transformers::Private::FApplyAdditiveSpace_BoneTransformAttribute>();
-		RuntimeRegistry.RegisterAttributeValueTransformer<Transformers::Private::FApplyAdditiveSpace_FloatAttribute>();
-		RuntimeRegistry.RegisterAttributeValueTransformer<Transformers::Private::FInterpolate_BoneTransformAttribute>();
-		RuntimeRegistry.RegisterAttributeValueTransformer<Transformers::Private::FInterpolate_FloatAttribute>();
-		RuntimeRegistry.RegisterAttributeValueTransformer<Transformers::Private::FLayer_BoneTransformAttribute>();
-		RuntimeRegistry.RegisterAttributeValueTransformer<Transformers::Private::FLayer_FloatAttribute>();
-		RuntimeRegistry.RegisterAttributeValueTransformer<Transformers::Private::FMakeAdditiveSpace_BoneTransformAttribute>();
-		RuntimeRegistry.RegisterAttributeValueTransformer<Transformers::Private::FMakeAdditiveSpace_FloatAttribute>();
-		RuntimeRegistry.RegisterAttributeValueTransformer<Transformers::Private::FOverwrite_BoneTransformAttribute>();
-		RuntimeRegistry.RegisterAttributeValueTransformer<Transformers::Private::FOverwrite_FloatAttribute>();
-		RuntimeRegistry.RegisterAttributeValueTransformer<Transformers::Private::FSanitize_BoneTransformAttribute>();
+		RuntimeRegistry.RegisterValueTransformer<Transformers::Private::FAccumulate_BoneTransformAttribute>();
+		RuntimeRegistry.RegisterValueTransformer<Transformers::Private::FAccumulate_FloatAttribute>();
+		RuntimeRegistry.RegisterValueTransformer<Transformers::Private::FApplyAdditiveSpace_BoneTransformAttribute>();
+		RuntimeRegistry.RegisterValueTransformer<Transformers::Private::FApplyAdditiveSpace_FloatAttribute>();
+		RuntimeRegistry.RegisterValueTransformer<Transformers::Private::FInterpolate_BoneTransformAttribute>();
+		RuntimeRegistry.RegisterValueTransformer<Transformers::Private::FInterpolate_FloatAttribute>();
+		RuntimeRegistry.RegisterValueTransformer<Transformers::Private::FLayer_BoneTransformAttribute>();
+		RuntimeRegistry.RegisterValueTransformer<Transformers::Private::FLayer_FloatAttribute>();
+		RuntimeRegistry.RegisterValueTransformer<Transformers::Private::FMakeAdditiveSpace_BoneTransformAttribute>();
+		RuntimeRegistry.RegisterValueTransformer<Transformers::Private::FMakeAdditiveSpace_FloatAttribute>();
+		RuntimeRegistry.RegisterValueTransformer<Transformers::Private::FOverwrite_BoneTransformAttribute>();
+		RuntimeRegistry.RegisterValueTransformer<Transformers::Private::FOverwrite_FloatAttribute>();
+		RuntimeRegistry.RegisterValueTransformer<Transformers::Private::FSanitize_BoneTransformAttribute>();
 
 #if UE_UAF_VALIDATE_MAPPED_VALUES
-		RuntimeRegistry.RegisterAttributeValueTransformer<Transformers::Private::FSanitize_FloatAttribute>();
+		RuntimeRegistry.RegisterValueTransformer<Transformers::Private::FSanitize_FloatAttribute>();
 #endif
 	}
 }
diff --git a/Engine/Plugins/Experimental/UAF/UAF/Source/UAF/Private/UAF/ValueRuntime/AttributeCollection.cpp b/Engine/Plugins/Experimental/UAF/UAF/Source/UAF/Private/UAF/ValueRuntime/AttributeCollection.cpp
--- a/Engine/Plugins/Experimental/UAF/UAF/Source/UAF/Private/UAF/ValueRuntime/AttributeCollection.cpp
+++ b/Engine/Plugins/Experimental/UAF/UAF/Source/UAF/Private/UAF/ValueRuntime/AttributeCollection.cpp
@@ -222,7 +222,7 @@ namespace UE::UAF
 				// By default, attributes map to values of their type
 				const FAttributeMappingKey MappingKey = FAttributeMappingKey::MakeFromTo(TypedSet->GetType());
 
-				FAttributeSetMap* Copy = AddSetMap(MappingKey);
+				FBoundValueMap* Copy = AddSetMap(MappingKey);
 				Copy->FillWithIdentity(bIsAdditive);
 			}
 		}
@@ -240,13 +240,13 @@ namespace UE::UAF
 		// Free our maps but retain the arrays that contains them
 		for (int32 MapIndex = 0; MapIndex < SetMapCount; ++MapIndex)
 		{
-			ReleaseAttributeSetMap(SetMaps[MapIndex].Map);
+			ReleaseBoundValueMap(SetMaps[MapIndex].Map);
 		}
 		SetMapCount = 0;
 
 		for (int32 MapIndex = 0; MapIndex < MapCount; ++MapIndex)
 		{
-			ReleaseAttributeMap(Maps[MapIndex].Map);
+			ReleaseUnboundValueMap(Maps[MapIndex].Map);
 		}
 		MapCount = 0;
 
@@ -260,7 +260,7 @@ namespace UE::UAF
 		{
 			for (int32 MapIndex = 0; MapIndex < SetMapCount; ++MapIndex)
 			{
-				ReleaseAttributeSetMap(SetMaps[MapIndex].Map);
+				ReleaseBoundValueMap(SetMaps[MapIndex].Map);
 			}
 
 			(*ReallocFun)(reinterpret_cast<uint8*>(SetMaps), sizeof(FSetMapEntry) * SetMapCapacity, 0);
@@ -273,7 +273,7 @@ namespace UE::UAF
 		{
 			for (int32 MapIndex = 0; MapIndex < MapCount; ++MapIndex)
 			{
-				ReleaseAttributeMap(Maps[MapIndex].Map);
+				ReleaseUnboundValueMap(Maps[MapIndex].Map);
 			}
 
 			(*ReallocFun)(reinterpret_cast<uint8*>(Maps), sizeof(FMapEntry) * MapCapacity, 0);
@@ -296,7 +296,7 @@ namespace UE::UAF
 		return SetMapCount;
 	}
 
-	FAttributeSetMap* FAttributeCollection::AddSetMap(const FAttributeMappingKey& MappingKey)
+	FBoundValueMap* FAttributeCollection::AddSetMap(const FAttributeMappingKey& MappingKey)
 	{
 		if (!NamedSet)
 		{
@@ -312,15 +312,15 @@ namespace UE::UAF
 		}
 
 		// Use default construction arguments
-		const FAttributeSetMap::FConstructArgs Args(TypedSet, MappingKey.GetValueType(), ReallocFun);
+		const FBoundValueMap::FConstructArgs Args(TypedSet, MappingKey.GetValueType(), ReallocFun);
 
 		// We don't know the concrete type of the set map, use indirection to construct it
 		constexpr bool bConstructSetMap = true;
-		FAttributeSetMap** SetMap = AddSetMapImpl(Args, bConstructSetMap);
-		return SetMap != nullptr ? *SetMap : nullptr;
+		FBoundValueMap** Map = AddSetMapImpl(Args, bConstructSetMap);
+		return Map != nullptr ? *Map : nullptr;
 	}
 
-	FAttributeSetMap** FAttributeCollection::AddSetMapImpl(const FAttributeSetMap::FConstructArgs& Args, bool bConstructSetMap)
+	FBoundValueMap** FAttributeCollection::AddSetMapImpl(const FBoundValueMap::FConstructArgs& Args, bool bConstructSetMap)
 	{
 		UScriptStruct* AttributeType = Args.TypedSet->GetType();
 		UScriptStruct* ValueType = Args.ValueType;
@@ -334,12 +334,12 @@ namespace UE::UAF
 			return nullptr;
 		}
 
-		FAttributeSetMap* SetMap = nullptr;
+		FBoundValueMap* Map = nullptr;
 		if (bConstructSetMap)
 		{
-			SetMap = FValueRuntimeRegistry::Get().ConstructSetMap(Args);
+			Map = FValueRuntimeRegistry::Get().ConstructBoundValueMap(Args);
 
-			if (SetMap == nullptr)
+			if (Map == nullptr)
 			{
 				// Failed to construct a set map, the type has probably not been registered
 				return nullptr;
@@ -364,13 +364,13 @@ namespace UE::UAF
 		const int32 NumToMove = SetMapCount - MapIndex;
 		FMemory::Memmove(SetMaps + MapIndex + 1, SetMaps + MapIndex, sizeof(FSetMapEntry) * NumToMove);
 		SetMaps[MapIndex].Key = Key;
-		SetMaps[MapIndex].Map = SetMap;
+		SetMaps[MapIndex].Map = Map;
 		SetMapCount++;
 
 		return &SetMaps[MapIndex].Map;
 	}
 
-	bool FAttributeCollection::AddSetMap(FAttributeSetMap* Map)
+	bool FAttributeCollection::AddSetMap(FBoundValueMap* Map)
 	{
 		if (Map == nullptr)
 		{
@@ -416,7 +416,7 @@ namespace UE::UAF
 		return true;
 	}
 
-	bool FAttributeCollection::AppendSetMap(FAttributeSetMap* Map)
+	bool FAttributeCollection::AppendSetMap(FBoundValueMap* Map)
 	{
 		if (Map == nullptr)
 		{
@@ -459,7 +459,7 @@ namespace UE::UAF
 		return true;
 	}
 
-	bool FAttributeCollection::RemoveSetMap(FAttributeSetMap* Map)
+	bool FAttributeCollection::RemoveSetMap(FBoundValueMap* Map)
 	{
 		if (Map == nullptr)
 		{
@@ -480,7 +480,7 @@ namespace UE::UAF
 		}
 
 		// Destroy and free our set map
-		ReleaseAttributeSetMap(SetMaps[MapIndex].Map);
+		ReleaseBoundValueMap(SetMaps[MapIndex].Map);
 
 		// Remove our entry
 		const int32 NumToMove = SetMapCount - MapIndex - 1;
@@ -490,7 +490,7 @@ namespace UE::UAF
 		return true;
 	}
 
-	FAttributeSetMap* FAttributeCollection::FindSetMap(const FAttributeMappingKey& MappingKey)
+	FBoundValueMap* FAttributeCollection::FindSetMap(const FAttributeMappingKey& MappingKey)
 	{
 		const int32 MapIndex = Algo::LowerBoundBy(TArrayView<const FSetMapEntry>(SetMaps, SetMapCount), MappingKey, [](const FSetMapEntry& Entry) { return Entry.Key; });
 		if (MapIndex >= SetMapCount || SetMaps[MapIndex].Key != MappingKey)
@@ -502,7 +502,7 @@ namespace UE::UAF
 		return SetMaps[MapIndex].Map;
 	}
 
-	const FAttributeSetMap* FAttributeCollection::FindSetMap(const FAttributeMappingKey& MappingKey) const
+	const FBoundValueMap* FAttributeCollection::FindSetMap(const FAttributeMappingKey& MappingKey) const
 	{
 		const int32 MapIndex = Algo::LowerBoundBy(TArrayView<const FSetMapEntry>(SetMaps, SetMapCount), MappingKey, [](const FSetMapEntry& Entry) { return Entry.Key; });
 		if (MapIndex >= SetMapCount || SetMaps[MapIndex].Key != MappingKey)
@@ -514,9 +514,9 @@ namespace UE::UAF
 		return SetMaps[MapIndex].Map;
 	}
 
-	FAttributeSetMap* FAttributeCollection::FindOrAddSetMap(const FAttributeMappingKey& MappingKey)
+	FBoundValueMap* FAttributeCollection::FindOrAddSetMap(const FAttributeMappingKey& MappingKey)
 	{
-		if (FAttributeSetMap* SetMap = FindSetMap(MappingKey))
+		if (FBoundValueMap* SetMap = FindSetMap(MappingKey))
 		{
 			return SetMap;
 		}
@@ -529,7 +529,7 @@ namespace UE::UAF
 		return MapCount;
 	}
 
-	FAttributeMap* FAttributeCollection::AddMap(UScriptStruct* ValueType)
+	FUnboundValueMap* FAttributeCollection::AddMap(UScriptStruct* ValueType)
 	{
 		if (!NamedSet)
 		{
@@ -539,11 +539,11 @@ namespace UE::UAF
 
 		// We don't know the concrete type of the attribute map, use indirection to construct it
 		constexpr bool bConstructMap = true;
-		FAttributeMap** Map = AddMapImpl(ValueType, bConstructMap);
+		FUnboundValueMap** Map = AddMapImpl(ValueType, bConstructMap);
 		return Map != nullptr ? *Map : nullptr;
 	}
 
-	FAttributeMap** FAttributeCollection::AddMapImpl(UScriptStruct* ValueType, bool bConstructMap)
+	FUnboundValueMap** FAttributeCollection::AddMapImpl(UScriptStruct* ValueType, bool bConstructMap)
 	{
 		const int32 MapIndex = Algo::LowerBoundBy(TArrayView<const FMapEntry>(Maps, MapCount), ValueType, [](const FMapEntry& Entry) { return Entry.ValueType; });
 		if (MapIndex < MapCount && Maps[MapIndex].ValueType == ValueType)
@@ -552,10 +552,10 @@ namespace UE::UAF
 			return nullptr;
 		}
 
-		FAttributeMap* Map = nullptr;
+		FUnboundValueMap* Map = nullptr;
 		if (bConstructMap)
 		{
-			Map = FValueRuntimeRegistry::Get().ConstructMap(ValueType, ReallocFun);
+			Map = FValueRuntimeRegistry::Get().ConstructUnboundValueMap(ValueType, ReallocFun);
 
 			if (Map == nullptr)
 			{
@@ -587,7 +587,7 @@ namespace UE::UAF
 		return &Maps[MapIndex].Map;
 	}
 
-	bool FAttributeCollection::AddMap(FAttributeMap* Map)
+	bool FAttributeCollection::AddMap(FUnboundValueMap* Map)
 	{
 		if (Map == nullptr)
 		{
@@ -632,7 +632,7 @@ namespace UE::UAF
 		return true;
 	}
 
-	bool FAttributeCollection::AppendMap(FAttributeMap* Map)
+	bool FAttributeCollection::AppendMap(FUnboundValueMap* Map)
 	{
 		if (Map == nullptr)
 		{
@@ -674,7 +674,7 @@ namespace UE::UAF
 		return true;
 	}
 
-	bool FAttributeCollection::RemoveMap(FAttributeMap* Map)
+	bool FAttributeCollection::RemoveMap(FUnboundValueMap* Map)
 	{
 		if (Map == nullptr)
 		{
@@ -695,7 +695,7 @@ namespace UE::UAF
 		}
 
 		// Destroy and free our map
-		ReleaseAttributeMap(Maps[MapIndex].Map);
+		ReleaseUnboundValueMap(Maps[MapIndex].Map);
 
 		// Remove our entry
 		const int32 NumToMove = MapCount - MapIndex - 1;
@@ -705,7 +705,7 @@ namespace UE::UAF
 		return true;
 	}
 
-	FAttributeMap* FAttributeCollection::FindMap(UScriptStruct* ValueType)
+	FUnboundValueMap* FAttributeCollection::FindMap(UScriptStruct* ValueType)
 	{
 		const int32 MapIndex = Algo::LowerBoundBy(TArrayView<const FMapEntry>(Maps, MapCount), ValueType, [](const FMapEntry& Entry) { return Entry.ValueType; });
 		if (MapIndex >= MapCount || Maps[MapIndex].ValueType != ValueType)
@@ -717,7 +717,7 @@ namespace UE::UAF
 		return Maps[MapIndex].Map;
 	}
 
-	const FAttributeMap* FAttributeCollection::FindMap(UScriptStruct* ValueType) const
+	const FUnboundValueMap* FAttributeCollection::FindMap(UScriptStruct* ValueType) const
 	{
 		const int32 MapIndex = Algo::LowerBoundBy(TArrayView<const FMapEntry>(Maps, MapCount), ValueType, [](const FMapEntry& Entry) { return Entry.ValueType; });
 		if (MapIndex >= MapCount || Maps[MapIndex].ValueType != ValueType)
@@ -729,9 +729,9 @@ namespace UE::UAF
 		return Maps[MapIndex].Map;
 	}
 
-	FAttributeMap* FAttributeCollection::FindOrAddMap(UScriptStruct* ValueType)
+	FUnboundValueMap* FAttributeCollection::FindOrAddMap(UScriptStruct* ValueType)
 	{
-		if (FAttributeMap* Map = FindMap(ValueType))
+		if (FUnboundValueMap* Map = FindMap(ValueType))
 		{
 			return Map;
 		}
diff --git a/Engine/Plugins/Experimental/UAF/UAF/Source/UAF/Private/UAF/ValueRuntime/Transformers/Accumulate.cpp b/Engine/Plugins/Experimental/UAF/UAF/Source/UAF/Private/UAF/ValueRuntime/Transformers/Accumulate.cpp
--- a/Engine/Plugins/Experimental/UAF/UAF/Source/UAF/Private/UAF/ValueRuntime/Transformers/Accumulate.cpp
+++ b/Engine/Plugins/Experimental/UAF/UAF/Source/UAF/Private/UAF/ValueRuntime/Transformers/Accumulate.cpp
@@ -8,7 +8,7 @@ namespace UE::UAF::Transformers
 {
 	FName FAccumulate::TransformerName = TEXT("UAF::Transformers::Accumulate");
 
-	void FAccumulate::Apply(const FAttributeValueTransformerMapPtr& TransformerMap, const FAttributeCollection& InputA, const FAttributeCollection& InputB, float WeightB, FAttributeCollection& Output)
+	void FAccumulate::Apply(const FValueTransformerMapPtr& TransformerMap, const FAttributeCollection& InputA, const FAttributeCollection& InputB, float WeightB, FAttributeCollection& Output)
 	{
 		// Named sets must match as it ensures that our inputs have the same sizes/shapes
 		checkf(InputA.GetNamedSet() == InputB.GetNamedSet(), TEXT("Both input collections must use the same named set"));
@@ -26,14 +26,14 @@ namespace UE::UAF::Transformers
 		FAccumulate::Apply(TransformerMap, InputA, InputB, EmptyCollection, WeightB, Output);
 	}
 
-	void FAccumulate::Apply(const FAttributeValueTransformerMapPtr& TransformerMap, const FAttributeCollection& InputA, const FAttributeCollection& InputB, const FAttributeCollection& PerValueWeightsB, float DefaultWeightB, FAttributeCollection& Output)
+	void FAccumulate::Apply(const FValueTransformerMapPtr& TransformerMap, const FAttributeCollection& InputA, const FAttributeCollection& InputB, const FAttributeCollection& PerValueWeightsB, float DefaultWeightB, FAttributeCollection& Output)
 	{
 		// Named sets must match as it ensures that our inputs have the same sizes/shapes
 		checkf(InputA.GetNamedSet() == InputB.GetNamedSet(), TEXT("Both input collections must use the same named set"));
 		checkf(InputA.GetNamedSet() == PerValueWeightsB.GetNamedSet(), TEXT("Input and per-value weight collections must use the same named set"));
 		checkf(InputA.GetValueSpace() == InputB.GetValueSpace(), TEXT("Both input collections must have matching value spaces: [%s] != [%s]"), *InputA.GetValueSpace().ToString(), *InputB.GetValueSpace().ToString());
 
-		const FAttributeValueTransformerList* TransformerList = TransformerMap->Find(TransformerName);
+		const FValueTransformerList* TransformerList = TransformerMap->Find(TransformerName);
 		if (TransformerList == nullptr)
 		{
 			// No transformer implementations have been registered
@@ -49,18 +49,18 @@ namespace UE::UAF::Transformers
 			// Our output collection matches one of our inputs, we'll modify it in place
 			FReallocFun OutputAllocator = Output.GetAllocator();
 
-			// Set maps
+			// Bound value maps
 			{
-				// Keep a list of set maps that needs to be appended to our output
-				TArray<FAttributeSetMap*, TInlineAllocator<8>> PendingSetMaps;
+				// Keep a list of bound value maps that needs to be appended to our output
+				TArray<FBoundValueMap*, TInlineAllocator<8>> PendingMaps;
 
 				OuterJoinBy(
-					FAttributeValueTransformerListWithAttributeCollectionSetMapJoinOp(),
-					[&PerValueWeightsB, DefaultWeightB, bIsOutputA, FloatValueType, OutputAllocator, &PendingSetMaps](FRawTransformerFunc Transformer, FAttributeSetMap* SetMapA, FAttributeSetMap* SetMapB)
+					FValueTransformerListWithBoundMapCollectionJoinOp(),
+					[&PerValueWeightsB, DefaultWeightB, bIsOutputA, FloatValueType, OutputAllocator, &PendingMaps](FRawTransformerFunc Transformer, FBoundValueMap* MapA, FBoundValueMap* MapB)
 					{
-						if (SetMapA == nullptr)
+						if (MapA == nullptr)
 						{
-							if (SetMapB == nullptr)
+							if (MapB == nullptr)
 							{
 								// No maps found for the transformer type
 							}
@@ -70,7 +70,7 @@ namespace UE::UAF::Transformers
 								if (bIsOutputA)
 								{
 									// A is our output, duplicate it
-									PendingSetMaps.Add(SetMapB->Duplicate(OutputAllocator));
+									PendingMaps.Add(MapB->Duplicate(OutputAllocator));
 								}
 								else
 								{
@@ -80,7 +80,7 @@ namespace UE::UAF::Transformers
 						}
 						else
 						{
-							if (SetMapB == nullptr)
+							if (MapB == nullptr)
 							{
 								// B doesn't have a matching pair
 								if (bIsOutputA)
@@ -90,7 +90,7 @@ namespace UE::UAF::Transformers
 								else
 								{
 									// B is our output, duplicate it
-									PendingSetMaps.Add(SetMapA->Duplicate(OutputAllocator));
+									PendingMaps.Add(MapA->Duplicate(OutputAllocator));
 								}
 							}
 							else
@@ -104,7 +104,7 @@ namespace UE::UAF::Transformers
 										if (bIsOutputA)
 										{
 											// A is our output, copy it over
-											SetMapB->CopyTo(*SetMapA);
+											MapB->CopyTo(*MapA);
 										}
 										else
 										{
@@ -121,45 +121,45 @@ namespace UE::UAF::Transformers
 										else
 										{
 											// B is our output, copy it over
-											SetMapA->CopyTo(*SetMapB);
+											MapA->CopyTo(*MapB);
 										}
 									}
 								}
 								else
 								{
 									// Transform our matching pair in place
-									FAttributeSetMap* OutputSetMap = bIsOutputA ? SetMapA : SetMapB;
+									FBoundValueMap* OutputMap = bIsOutputA ? MapA : MapB;
 
 									// We cannot join our weights because their value type differs
-									const FAttributeMappingKey MappingKey = SetMapA->GetMappingKey().To(FloatValueType);
-									const FAttributeSetMap* PerValueWeightMapB = PerValueWeightsB.FindSetMap(MappingKey);
+									const FAttributeMappingKey MappingKey = MapA->GetMappingKey().To(FloatValueType);
+									const FBoundValueMap* PerValueWeightMapB = PerValueWeightsB.FindSetMap(MappingKey);
 
-									FTransformSetMapFunc TypedTransformer = reinterpret_cast<FTransformSetMapFunc>(Transformer);
-									(*TypedTransformer)(SetMapA, SetMapB, PerValueWeightMapB, DefaultWeightB, OutputSetMap, nullptr);
+									FTransformBoundValueMapFunc TypedTransformer = reinterpret_cast<FTransformBoundValueMapFunc>(Transformer);
+									(*TypedTransformer)(MapA, MapB, PerValueWeightMapB, DefaultWeightB, OutputMap, nullptr);
 								}
 							}
 						}
 					},
-					TransformerList->CreateSetMapTransformerIterator(),
+					TransformerList->CreateBoundValueMapTransformerIterator(),
 					// Remove the const qualifier since we'll mutate one in place as our output
 					const_cast<FAttributeCollection&>(InputA).CreateSetMapIterator(),
 					const_cast<FAttributeCollection&>(InputB).CreateSetMapIterator());
 
 				// Append our pending entries
-				for (FAttributeSetMap* SetMap : PendingSetMaps)
+				for (FBoundValueMap* Map : PendingMaps)
 				{
-					Output.AddSetMap(SetMap);
+					Output.AddSetMap(Map);
 				}
 			}
 
 			// Name/value maps
 			{
 				// Keep a list of maps that needs to be appended to our output
-				TArray<FAttributeMap*, TInlineAllocator<8>> PendingMaps;
+				TArray<FUnboundValueMap*, TInlineAllocator<8>> PendingMaps;
 
 				OuterJoinBy(
-					FAttributeValueTransformerListWithAttributeCollectionMapJoinOp(),
-					[DefaultWeightB, bIsOutputA, OutputAllocator, &PendingMaps](FRawTransformerFunc Transformer, FAttributeMap* MapA, FAttributeMap* MapB)
+					FValueTransformerListWithUnboundMapCollectionJoinOp(),
+					[DefaultWeightB, bIsOutputA, OutputAllocator, &PendingMaps](FRawTransformerFunc Transformer, FUnboundValueMap* MapA, FUnboundValueMap* MapB)
 					{
 						if (MapA == nullptr)
 						{
@@ -231,9 +231,9 @@ namespace UE::UAF::Transformers
 								else
 								{
 									// Transform our matching pair in a new map
-									FAttributeMap* OutputMap = nullptr;
+									FUnboundValueMap* OutputMap = nullptr;
 
-									FTransformMapFunc TypedTransformer = reinterpret_cast<FTransformMapFunc>(Transformer);
+									FTransformUnboundValueMapFunc TypedTransformer = reinterpret_cast<FTransformUnboundValueMapFunc>(Transformer);
 									(*TypedTransformer)(MapA, MapB, DefaultWeightB, OutputMap, OutputAllocator);
 
 									checkf(OutputMap != nullptr, TEXT("Transformer function is responsible for allocating an output buffer"));
@@ -243,18 +243,18 @@ namespace UE::UAF::Transformers
 									OutputMap->MoveTo(bIsOutputA ? *MapA : *MapB);
 
 									// Destroy our temporary output map
-									ReleaseAttributeMap(OutputMap);
+									ReleaseUnboundValueMap(OutputMap);
 								}
 							}
 						}
 					},
-					TransformerList->CreateMapTransformerIterator(),
+					TransformerList->CreateUnboundValueMapTransformerIterator(),
 					// Remove the const qualifier since we'll mutate one in place as our output
 					const_cast<FAttributeCollection&>(InputA).CreateMapIterator(),
 					const_cast<FAttributeCollection&>(InputB).CreateMapIterator());
 
 				// Append our pending entries
-				for (FAttributeMap* Map : PendingMaps)
+				for (FUnboundValueMap* Map : PendingMaps)
 				{
 					Output.AddMap(Map);
 				}
@@ -266,74 +266,74 @@ namespace UE::UAF::Transformers
 			Output.Reset(InputA.GetNamedSet());
 			Output.SetValueSpace(InputA.GetValueSpace());
 
-			// Set maps
+			// Bound value maps
 			OuterJoinBy(
-				FAttributeValueTransformerListWithAttributeCollectionSetMapJoinOp(),
-				[&PerValueWeightsB, DefaultWeightB, FloatValueType, &Output](FRawTransformerFunc Transformer, const FAttributeSetMap* SetMapA, const FAttributeSetMap* SetMapB)
+				FValueTransformerListWithBoundMapCollectionJoinOp(),
+				[&PerValueWeightsB, DefaultWeightB, FloatValueType, &Output](FRawTransformerFunc Transformer, const FBoundValueMap* MapA, const FBoundValueMap* MapB)
 				{
-					if (SetMapA == nullptr)
+					if (MapA == nullptr)
 					{
-						if (SetMapB == nullptr)
+						if (MapB == nullptr)
 						{
 							// No maps found for the transformer type
 						}
 						else
 						{
 							// This set map doesn't have a matching pair, duplicate it
-							Output.AppendSetMap(SetMapB->Duplicate(Output.GetAllocator()));
+							Output.AppendSetMap(MapB->Duplicate(Output.GetAllocator()));
 						}
 					}
 					else
 					{
-						if (SetMapB == nullptr)
+						if (MapB == nullptr)
 						{
 							// This set map doesn't have a matching pair, duplicate it
-							Output.AppendSetMap(SetMapA->Duplicate(Output.GetAllocator()));
+							Output.AppendSetMap(MapA->Duplicate(Output.GetAllocator()));
 						}
 						else
 						{
 							if (Transformer == nullptr)
 							{
 								// No transformer was found for this value type, retain the one with the largest weight
-								FAttributeSetMap* SetMapCopy;
+								FBoundValueMap* MapCopy;
 								if (DefaultWeightB >= 0.5f)
 								{
-									SetMapCopy = SetMapB->Duplicate(Output.GetAllocator());
+									MapCopy = MapB->Duplicate(Output.GetAllocator());
 								}
 								else
 								{
-									SetMapCopy = SetMapA->Duplicate(Output.GetAllocator());
+									MapCopy = MapA->Duplicate(Output.GetAllocator());
 								}
 
-								Output.AppendSetMap(SetMapCopy);
+								Output.AppendSetMap(MapCopy);
 							}
 							else
 							{
 								// Transform our matching pair
-								FAttributeSetMap* OutputSetMap = nullptr;
+								FBoundValueMap* OutputMap = nullptr;
 
 								// We cannot join our weights because their value type differs
-								const FAttributeMappingKey MappingKey = SetMapA->GetMappingKey().To(FloatValueType);
-								const FAttributeSetMap* PerValueWeightMapB = PerValueWeightsB.FindSetMap(MappingKey);
+								const FAttributeMappingKey MappingKey = MapA->GetMappingKey().To(FloatValueType);
+								const FBoundValueMap* PerValueWeightMapB = PerValueWeightsB.FindSetMap(MappingKey);
 
-								FTransformSetMapFunc TypedTransformer = reinterpret_cast<FTransformSetMapFunc>(Transformer);
-								(*TypedTransformer)(SetMapA, SetMapB, PerValueWeightMapB, DefaultWeightB, OutputSetMap, Output.GetAllocator());
+								FTransformBoundValueMapFunc TypedTransformer = reinterpret_cast<FTransformBoundValueMapFunc>(Transformer);
+								(*TypedTransformer)(MapA, MapB, PerValueWeightMapB, DefaultWeightB, OutputMap, Output.GetAllocator());
 
-								checkf(OutputSetMap != nullptr, TEXT("Transformer function is responsible for allocating an output buffer"));
-								checkf(OutputSetMap->GetAllocator() == Output.GetAllocator(), TEXT("Transformer function is responsible for allocating using the output allocator"));
-								Output.AppendSetMap(OutputSetMap);
+								checkf(OutputMap != nullptr, TEXT("Transformer function is responsible for allocating an output buffer"));
+								checkf(OutputMap->GetAllocator() == Output.GetAllocator(), TEXT("Transformer function is responsible for allocating using the output allocator"));
+								Output.AppendSetMap(OutputMap);
 							}
 						}
 					}
 				},
-				TransformerList->CreateSetMapTransformerIterator(),
+				TransformerList->CreateBoundValueMapTransformerIterator(),
 				InputA.CreateConstSetMapIterator(),
 				InputB.CreateConstSetMapIterator());
 
 			// Name/value maps
 			OuterJoinBy(
-				FAttributeValueTransformerListWithAttributeCollectionMapJoinOp(),
-				[DefaultWeightB, &Output](FRawTransformerFunc Transformer, const FAttributeMap* MapA, const FAttributeMap* MapB)
+				FValueTransformerListWithUnboundMapCollectionJoinOp(),
+				[DefaultWeightB, &Output](FRawTransformerFunc Transformer, const FUnboundValueMap* MapA, const FUnboundValueMap* MapB)
 				{
 					if (MapA == nullptr)
 					{
@@ -359,7 +359,7 @@ namespace UE::UAF::Transformers
 							if (Transformer == nullptr)
 							{
 								// No transformer was found for this value type, retain the one with the largest weight
-								FAttributeMap* MapCopy;
+								FUnboundValueMap* MapCopy;
 								if (DefaultWeightB >= 0.5f)
 								{
 									MapCopy = MapB->Duplicate(Output.GetAllocator());
@@ -374,9 +374,9 @@ namespace UE::UAF::Transformers
 							else
 							{
 								// Transform our matching pair in place
-								FAttributeMap* OutputMap = nullptr;
+								FUnboundValueMap* OutputMap = nullptr;
 
-								FTransformMapFunc TypedTransformer = reinterpret_cast<FTransformMapFunc>(Transformer);
+								FTransformUnboundValueMapFunc TypedTransformer = reinterpret_cast<FTransformUnboundValueMapFunc>(Transformer);
 								(*TypedTransformer)(MapA, MapB, DefaultWeightB, OutputMap, Output.GetAllocator());
 
 								checkf(OutputMap != nullptr, TEXT("Transformer function is responsible for allocating an output buffer"));
@@ -386,7 +386,7 @@ namespace UE::UAF::Transformers
 						}
 					}
 				},
-				TransformerList->CreateMapTransformerIterator(),
+				TransformerList->CreateUnboundValueMapTransformerIterator(),
 				InputA.CreateConstMapIterator(),
 				InputB.CreateConstMapIterator());
 		}
@@ -394,33 +394,33 @@ namespace UE::UAF::Transformers
 
 	namespace Private
 	{
-		void FAccumulate_BoneTransformAttribute::TransformSetMap(const FAttributeSetMap* SetMapA, const FAttributeSetMap* SetMapB, const FAttributeSetMap* PerValueWeightMapB, float DefaultWeightB, FAttributeSetMap*& OutputSetMap, FReallocFun OutputAllocator)
+		void FAccumulate_BoneTransformAttribute::TransformBoundValueMap(const FBoundValueMap* MapA, const FBoundValueMap* MapB, const FBoundValueMap* PerValueWeightMapB, float DefaultWeightB, FBoundValueMap*& OutputMap, FReallocFun OutputAllocator)
 		{
-			check(SetMapA->GetTypedSet() == SetMapB->GetTypedSet());
-			check(SetMapA->Num() == SetMapB->Num());
+			check(MapA->GetTypedSet() == MapB->GetTypedSet());
+			check(MapA->Num() == MapB->Num());
 
-			if (OutputSetMap == nullptr)
+			if (OutputMap == nullptr)
 			{
-				const FAttributeSetMap::FConstructArgs Args(SetMapA->GetTypedSet(), SetMapA->GetValueType(), OutputAllocator);
-				OutputSetMap = MakeAttributeSetMap<FBoneTransformAnimationAttribute>(Args);
+				const FBoundValueMap::FConstructArgs Args(MapA->GetTypedSet(), MapA->GetValueType(), OutputAllocator);
+				OutputMap = MakeBoundValueMap<FBoneTransformAnimationAttribute>(Args);
 			}
 
-			check(SetMapA->GetTypedSet() == OutputSetMap->GetTypedSet());
-			check(SetMapA->Num() == OutputSetMap->Num());
+			check(MapA->GetTypedSet() == OutputMap->GetTypedSet());
+			check(MapA->Num() == OutputMap->Num());
 
-			const TAttributeSetMap<FBoneTransformAnimationAttribute>* BoneTransformsA = CastChecked<FBoneTransformAnimationAttribute>(SetMapA);
-			const TAttributeSetMap<FBoneTransformAnimationAttribute>* BoneTransformsB = CastChecked<FBoneTransformAnimationAttribute>(SetMapB);
-			TAttributeSetMap<FBoneTransformAnimationAttribute>* OutputBoneTransforms = CastChecked<FBoneTransformAnimationAttribute>(OutputSetMap);
+			const TBoundValueMap<FBoneTransformAnimationAttribute>* BoneTransformsA = CastChecked<FBoneTransformAnimationAttribute>(MapA);
+			const TBoundValueMap<FBoneTransformAnimationAttribute>* BoneTransformsB = CastChecked<FBoneTransformAnimationAttribute>(MapB);
+			TBoundValueMap<FBoneTransformAnimationAttribute>* OutputBoneTransforms = CastChecked<FBoneTransformAnimationAttribute>(OutputMap);
 
 			const FBoneTransformAnimationAttribute* TransformPtrA = BoneTransformsA->GetData();
 			const FBoneTransformAnimationAttribute* TransformPtrB = BoneTransformsB->GetData();
 			FBoneTransformAnimationAttribute* OutputTransformPtr = OutputBoneTransforms->GetData();
 			const FBoneTransformAnimationAttribute* TransformPtrAEnd = TransformPtrA + BoneTransformsA->Num();
 
-			if (const TAttributeSetMap<FFloatAnimationAttribute>* PerValueWeightsB = Cast<FFloatAnimationAttribute>(PerValueWeightMapB))
+			if (const TBoundValueMap<FFloatAnimationAttribute>* PerValueWeightsB = Cast<FFloatAnimationAttribute>(PerValueWeightMapB))
 			{
-				check(SetMapA->GetTypedSet() == PerValueWeightsB->GetTypedSet());
-				check(SetMapA->Num() == PerValueWeightsB->Num());
+				check(MapA->GetTypedSet() == PerValueWeightsB->GetTypedSet());
+				check(MapA->Num() == PerValueWeightsB->Num());
 
 				const FFloatAnimationAttribute* PerValueWeightsBPtr = PerValueWeightsB->GetData();
 
@@ -455,18 +455,18 @@ namespace UE::UAF::Transformers
 			}
 		}
 
-		void FAccumulate_BoneTransformAttribute::TransformMap(const FAttributeMap* MapA, const FAttributeMap* MapB, float WeightB, FAttributeMap*& OutputMap, FReallocFun OutputAllocator)
+		void FAccumulate_BoneTransformAttribute::TransformUnboundValueMap(const FUnboundValueMap* MapA, const FUnboundValueMap* MapB, float WeightB, FUnboundValueMap*& OutputMap, FReallocFun OutputAllocator)
 		{
 			check(OutputMap == nullptr);
 
-			const TAttributeMap<FBoneTransformAnimationAttribute>* BoneTransformsA = CastChecked<FBoneTransformAnimationAttribute>(MapA);
-			const TAttributeMap<FBoneTransformAnimationAttribute>* BoneTransformsB = CastChecked<FBoneTransformAnimationAttribute>(MapB);
-			TAttributeMap<FBoneTransformAnimationAttribute>* OutputBoneTransforms = MakeAttributeMap<FBoneTransformAnimationAttribute>(OutputAllocator);
+			const TUnboundValueMap<FBoneTransformAnimationAttribute>* BoneTransformsA = CastChecked<FBoneTransformAnimationAttribute>(MapA);
+			const TUnboundValueMap<FBoneTransformAnimationAttribute>* BoneTransformsB = CastChecked<FBoneTransformAnimationAttribute>(MapB);
+			TUnboundValueMap<FBoneTransformAnimationAttribute>* OutputBoneTransforms = MakeUnboundValueMap<FBoneTransformAnimationAttribute>(OutputAllocator);
 
 			const ScalarRegister VWeightB(WeightB);
 
 			OuterJoinBy(
-				FAttributeMapJoinOp(),
+				FUnboundValueMapJoinOp(),
 				[VWeightB, OutputBoneTransforms](FName Name, const FBoneTransformAnimationAttribute* InputA, const FBoneTransformAnimationAttribute* InputB)
 				{
 					// Missing values are set to the identity
@@ -484,33 +484,33 @@ namespace UE::UAF::Transformers
 			OutputMap = OutputBoneTransforms;
 		}
 
-		void FAccumulate_FloatAttribute::TransformSetMap(const FAttributeSetMap* SetMapA, const FAttributeSetMap* SetMapB, const FAttributeSetMap* PerValueWeightMapB, float DefaultWeightB, FAttributeSetMap*& OutputSetMap, FReallocFun OutputAllocator)
+		void FAccumulate_FloatAttribute::TransformBoundValueMap(const FBoundValueMap* MapA, const FBoundValueMap* MapB, const FBoundValueMap* PerValueWeightMapB, float DefaultWeightB, FBoundValueMap*& OutputMap, FReallocFun OutputAllocator)
 		{
-			check(SetMapA->GetTypedSet() == SetMapB->GetTypedSet());
-			check(SetMapA->Num() == SetMapB->Num());
+			check(MapA->GetTypedSet() == MapB->GetTypedSet());
+			check(MapA->Num() == MapB->Num());
 
-			if (OutputSetMap == nullptr)
+			if (OutputMap == nullptr)
 			{
-				const FAttributeSetMap::FConstructArgs Args(SetMapA->GetTypedSet(), SetMapA->GetValueType(), OutputAllocator);
-				OutputSetMap = MakeAttributeSetMap<FFloatAnimationAttribute>(Args);
+				const FBoundValueMap::FConstructArgs Args(MapA->GetTypedSet(), MapA->GetValueType(), OutputAllocator);
+				OutputMap = MakeBoundValueMap<FFloatAnimationAttribute>(Args);
 			}
 
-			check(SetMapA->GetTypedSet() == OutputSetMap->GetTypedSet());
-			check(SetMapA->Num() == OutputSetMap->Num());
+			check(MapA->GetTypedSet() == OutputMap->GetTypedSet());
+			check(MapA->Num() == OutputMap->Num());
 
-			const TAttributeSetMap<FFloatAnimationAttribute>* FloatsA = CastChecked<FFloatAnimationAttribute>(SetMapA);
-			const TAttributeSetMap<FFloatAnimationAttribute>* FloatsB = CastChecked<FFloatAnimationAttribute>(SetMapB);
-			TAttributeSetMap<FFloatAnimationAttribute>* OutputFloats = CastChecked<FFloatAnimationAttribute>(OutputSetMap);
+			const TBoundValueMap<FFloatAnimationAttribute>* FloatsA = CastChecked<FFloatAnimationAttribute>(MapA);
+			const TBoundValueMap<FFloatAnimationAttribute>* FloatsB = CastChecked<FFloatAnimationAttribute>(MapB);
+			TBoundValueMap<FFloatAnimationAttribute>* OutputFloats = CastChecked<FFloatAnimationAttribute>(OutputMap);
 
 			const FFloatAnimationAttribute* FloatsPtrA = FloatsA->GetData();
 			const FFloatAnimationAttribute* FloatsPtrB = FloatsB->GetData();
 			FFloatAnimationAttribute* OutputFloatsPtr = OutputFloats->GetData();
 			const FFloatAnimationAttribute* FloatsPtrAEnd = FloatsPtrA + FloatsA->Num();
 
-			if (const TAttributeSetMap<FFloatAnimationAttribute>* PerValueWeightsB = Cast<FFloatAnimationAttribute>(PerValueWeightMapB))
+			if (const TBoundValueMap<FFloatAnimationAttribute>* PerValueWeightsB = Cast<FFloatAnimationAttribute>(PerValueWeightMapB))
 			{
-				check(SetMapA->GetTypedSet() == PerValueWeightsB->GetTypedSet());
-				check(SetMapA->Num() == PerValueWeightsB->Num());
+				check(MapA->GetTypedSet() == PerValueWeightsB->GetTypedSet());
+				check(MapA->Num() == PerValueWeightsB->Num());
 
 				const FFloatAnimationAttribute* PerValueWeightsBPtr = PerValueWeightsB->GetData();
 
@@ -537,16 +537,16 @@ namespace UE::UAF::Transformers
 			}
 		}
 
-		void FAccumulate_FloatAttribute::TransformMap(const FAttributeMap* MapA, const FAttributeMap* MapB, float WeightB, FAttributeMap*& OutputMap, FReallocFun OutputAllocator)
+		void FAccumulate_FloatAttribute::TransformUnboundValueMap(const FUnboundValueMap* MapA, const FUnboundValueMap* MapB, float WeightB, FUnboundValueMap*& OutputMap, FReallocFun OutputAllocator)
 		{
 			check(OutputMap == nullptr);
 
-			const TAttributeMap<FFloatAnimationAttribute>* FloatsA = CastChecked<FFloatAnimationAttribute>(MapA);
-			const TAttributeMap<FFloatAnimationAttribute>* FloatsB = CastChecked<FFloatAnimationAttribute>(MapB);
-			TAttributeMap<FFloatAnimationAttribute>* OutputFloats = MakeAttributeMap<FFloatAnimationAttribute>(OutputAllocator);
+			const TUnboundValueMap<FFloatAnimationAttribute>* FloatsA = CastChecked<FFloatAnimationAttribute>(MapA);
+			const TUnboundValueMap<FFloatAnimationAttribute>* FloatsB = CastChecked<FFloatAnimationAttribute>(MapB);
+			TUnboundValueMap<FFloatAnimationAttribute>* OutputFloats = MakeUnboundValueMap<FFloatAnimationAttribute>(OutputAllocator);
 
 			OuterJoinBy(
-				FAttributeMapJoinOp(),
+				FUnboundValueMapJoinOp(),
 				[&WeightB, OutputFloats](FName Name, const FFloatAnimationAttribute* InputA, const FFloatAnimationAttribute* InputB)
 				{
 					// Missing values are set to 0.0
diff --git a/Engine/Plugins/Experimental/UAF/UAF/Source/UAF/Private/UAF/ValueRuntime/Transformers/AdditiveSpace.cpp b/Engine/Plugins/Experimental/UAF/UAF/Source/UAF/Private/UAF/ValueRuntime/Transformers/AdditiveSpace.cpp
--- a/Engine/Plugins/Experimental/UAF/UAF/Source/UAF/Private/UAF/ValueRuntime/Transformers/AdditiveSpace.cpp
+++ b/Engine/Plugins/Experimental/UAF/UAF/Source/UAF/Private/UAF/ValueRuntime/Transformers/AdditiveSpace.cpp
@@ -11,7 +11,7 @@ namespace UE::UAF::Transformers
 	FName FApplyAdditiveSpace::TransformerName = TEXT("UAF::Transformers::ApplyAdditive");
 
 	void FMakeAdditiveSpace::Apply(
-		const FAttributeValueTransformerMapPtr& TransformerMap,
+		const FValueTransformerMapPtr& TransformerMap,
 		const FAttributeCollection& Base,
 		const FAttributeCollection& Input,
 		FAttributeCollection& AdditiveOutput)
@@ -21,7 +21,7 @@ namespace UE::UAF::Transformers
 		checkf(&Base != &AdditiveOutput, TEXT("Cannot output to the Base collection when making an additive space collection"));
 		checkf(&Input != &AdditiveOutput, TEXT("Cannot output to the Input collection when making an additive space collection"));
 
-		const FAttributeValueTransformerList* TransformerList = TransformerMap->Find(TransformerName);
+		const FValueTransformerList* TransformerList = TransformerMap->Find(TransformerName);
 		if (TransformerList == nullptr)
 		{
 			// No transformer implementations have been registered
@@ -39,10 +39,10 @@ namespace UE::UAF::Transformers
 		AdditiveOutput.Reset(Input.GetNamedSet());
 		AdditiveOutput.SetValueSpace(OutputSpace);
 
-		// Set maps
+		// Bound value maps
 		OuterJoinBy(
-			FAttributeValueTransformerListWithAttributeCollectionSetMapJoinOp(),
-			[&AdditiveOutput](FRawTransformerFunc Transformer, const FAttributeSetMap* BaseMap, const FAttributeSetMap* InputMap)
+			FValueTransformerListWithBoundMapCollectionJoinOp(),
+			[&AdditiveOutput](FRawTransformerFunc Transformer, const FBoundValueMap* BaseMap, const FBoundValueMap* InputMap)
 			{
 				if (InputMap == nullptr)
 				{
@@ -53,30 +53,29 @@ namespace UE::UAF::Transformers
 				if (Transformer == nullptr)
 				{
 					// No transformer was found for this value type, retain the input as-is
-					FAttributeSetMap* SetMapCopy = InputMap->Duplicate(AdditiveOutput.GetAllocator());
-					AdditiveOutput.AppendSetMap(SetMapCopy);
+					AdditiveOutput.AppendSetMap(InputMap->Duplicate(AdditiveOutput.GetAllocator()));
 					return;
 				}
 
 				// Transform our matching pair, note that the base map might be missing
 				// If it is missing, transformer specializations should use a default value
-				FAttributeSetMap* OutputSetMap = nullptr;
+				FBoundValueMap* OutputMap = nullptr;
 
-				FTransformSetMapFunc TypedTransformer = reinterpret_cast<FTransformSetMapFunc>(Transformer);
-				(*TypedTransformer)(BaseMap, InputMap, OutputSetMap, AdditiveOutput.GetAllocator());
+				FTransformBoundValueMapFunc TypedTransformer = reinterpret_cast<FTransformBoundValueMapFunc>(Transformer);
+				(*TypedTransformer)(BaseMap, InputMap, OutputMap, AdditiveOutput.GetAllocator());
 
-				checkf(OutputSetMap != nullptr, TEXT("Transformer function is responsible for allocating an output buffer"));
-				checkf(OutputSetMap->GetAllocator() == AdditiveOutput.GetAllocator(), TEXT("Transformer function is responsible for allocating using the output allocator"));
-				AdditiveOutput.AppendSetMap(OutputSetMap);
+				checkf(OutputMap != nullptr, TEXT("Transformer function is responsible for allocating an output buffer"));
+				checkf(OutputMap->GetAllocator() == AdditiveOutput.GetAllocator(), TEXT("Transformer function is responsible for allocating using the output allocator"));
+				AdditiveOutput.AppendSetMap(OutputMap);
 			},
-			TransformerList->CreateSetMapTransformerIterator(),
+			TransformerList->CreateBoundValueMapTransformerIterator(),
 			Base.CreateConstSetMapIterator(),
 			Input.CreateConstSetMapIterator());
 
 		// Name/value maps
 		OuterJoinBy(
-			FAttributeValueTransformerListWithAttributeCollectionMapJoinOp(),
-			[&AdditiveOutput](FRawTransformerFunc Transformer, const FAttributeMap* BaseMap, const FAttributeMap* InputMap)
+			FValueTransformerListWithUnboundMapCollectionJoinOp(),
+			[&AdditiveOutput](FRawTransformerFunc Transformer, const FUnboundValueMap* BaseMap, const FUnboundValueMap* InputMap)
 			{
 				if (InputMap == nullptr)
 				{
@@ -87,29 +86,29 @@ namespace UE::UAF::Transformers
 				if (Transformer == nullptr)
 				{
 					// No transformer was found for this value type, retain the input as-is
-					FAttributeMap* MapCopy = InputMap->Duplicate(AdditiveOutput.GetAllocator());
+					FUnboundValueMap* MapCopy = InputMap->Duplicate(AdditiveOutput.GetAllocator());
 					AdditiveOutput.AppendMap(MapCopy);
 					return;
 				}
 
 				// Transform our matching pair, note that the base map might be missing
 				// If it is missing, transformer specializations should use a default value
-				FAttributeMap* OutputMap = nullptr;
+				FUnboundValueMap* OutputMap = nullptr;
 
-				FTransformMapFunc TypedTransformer = reinterpret_cast<FTransformMapFunc>(Transformer);
+				FTransformUnboundValueMapFunc TypedTransformer = reinterpret_cast<FTransformUnboundValueMapFunc>(Transformer);
 				(*TypedTransformer)(BaseMap, InputMap, OutputMap, AdditiveOutput.GetAllocator());
 
 				checkf(OutputMap != nullptr, TEXT("Transformer function is responsible for allocating an output buffer"));
 				checkf(OutputMap->GetAllocator() == AdditiveOutput.GetAllocator(), TEXT("Transformer function is responsible for allocating using the output allocator"));
 				AdditiveOutput.AppendMap(OutputMap);
 			},
-			TransformerList->CreateMapTransformerIterator(),
+			TransformerList->CreateUnboundValueMapTransformerIterator(),
 			Base.CreateConstMapIterator(),
 			Input.CreateConstMapIterator());
 	}
 
 	void FApplyAdditiveSpace::Apply(
-		const FAttributeValueTransformerMapPtr& TransformerMap,
+		const FValueTransformerMapPtr& TransformerMap,
 		const FAttributeCollection& Base,
 		const FAttributeCollection& Additive,
 		float AdditiveWeight,
@@ -128,7 +127,7 @@ namespace UE::UAF::Transformers
 	}
 
 	void FApplyAdditiveSpace::Apply(
-		const FAttributeValueTransformerMapPtr& TransformerMap,
+		const FValueTransformerMapPtr& TransformerMap,
 		const FAttributeCollection& Base,
 		const FAttributeCollection& Additive,
 		const FAttributeCollection& PerValueWeights,
@@ -146,7 +145,7 @@ namespace UE::UAF::Transformers
 		checkf(AdditiveSpace.IsAdditive(), TEXT("Additive input must be in additive space"));
 		checkf(BaseSpace.GetType() == AdditiveSpace.GetType(), TEXT("Additive value space must match the base (local or mesh space additive)"));
 
-		const FAttributeValueTransformerList* TransformerList = TransformerMap->Find(TransformerName);
+		const FValueTransformerList* TransformerList = TransformerMap->Find(TransformerName);
 		if (TransformerList == nullptr)
 		{
 			// No transformer implementations have been registered
@@ -165,11 +164,11 @@ namespace UE::UAF::Transformers
 			// Set maps
 			{
 				// Keep a list of set maps that needs to be appended to our output
-				TArray<FAttributeSetMap*, TInlineAllocator<8>> PendingSetMaps;
+				TArray<FBoundValueMap*, TInlineAllocator<8>> PendingMaps;
 
 				OuterJoinBy(
-					FAttributeValueTransformerListWithAttributeCollectionSetMapJoinOp(),
-					[&PerValueWeights, DefaultWeight, bIsOutputBase, FloatValueType, OutputAllocator, &PendingSetMaps](FRawTransformerFunc Transformer, FAttributeSetMap* BaseMap, FAttributeSetMap* AdditiveMap)
+					FValueTransformerListWithBoundMapCollectionJoinOp(),
+					[&PerValueWeights, DefaultWeight, bIsOutputBase, FloatValueType, OutputAllocator, &PendingMaps](FRawTransformerFunc Transformer, FBoundValueMap* BaseMap, FBoundValueMap* AdditiveMap)
 					{
 						if (AdditiveMap == nullptr)
 						{
@@ -183,7 +182,7 @@ namespace UE::UAF::Transformers
 							if (BaseMap == nullptr)
 							{
 								// We have no base, duplicate the additive
-								PendingSetMaps.Add(AdditiveMap->Duplicate(OutputAllocator));
+								PendingMaps.Add(AdditiveMap->Duplicate(OutputAllocator));
 							}
 							else if (DefaultWeight >= 0.5f)
 							{
@@ -216,44 +215,44 @@ namespace UE::UAF::Transformers
 
 						// Transform our matching pair, note that the base map might be missing
 						// If it is missing, transformer specializations should use a default value
-						FAttributeSetMap* OutputSetMap = bIsOutputBase ? BaseMap : AdditiveMap;
-						const bool bIsOutputMissing = OutputSetMap == nullptr;
+						FBoundValueMap* OutputMap = bIsOutputBase ? BaseMap : AdditiveMap;
+						const bool bIsOutputMissing = OutputMap == nullptr;
 
 						// We cannot join our weights because their value type differs
 						const FAttributeMappingKey MappingKey = AdditiveMap->GetMappingKey().To(FloatValueType);
-						const FAttributeSetMap* PerValueWeightMap = PerValueWeights.FindSetMap(MappingKey);
+						const FBoundValueMap* PerValueWeightMap = PerValueWeights.FindSetMap(MappingKey);
 
-						FTransformSetMapFunc TypedTransformer = reinterpret_cast<FTransformSetMapFunc>(Transformer);
-						(*TypedTransformer)(BaseMap, AdditiveMap, PerValueWeightMap, DefaultWeight, OutputSetMap, OutputAllocator);
+						FTransformBoundValueMapFunc TypedTransformer = reinterpret_cast<FTransformBoundValueMapFunc>(Transformer);
+						(*TypedTransformer)(BaseMap, AdditiveMap, PerValueWeightMap, DefaultWeight, OutputMap, OutputAllocator);
 
 						if (bIsOutputMissing)
 						{
 							// Our output was missing, we should have allocated a buffer for it
-							checkf(OutputSetMap != nullptr, TEXT("Transformer function is responsible for allocating an output buffer"));
-							checkf(OutputSetMap->GetAllocator() == OutputAllocator, TEXT("Transformer function is responsible for allocating using the output allocator"));
-							PendingSetMaps.Add(OutputSetMap);
+							checkf(OutputMap != nullptr, TEXT("Transformer function is responsible for allocating an output buffer"));
+							checkf(OutputMap->GetAllocator() == OutputAllocator, TEXT("Transformer function is responsible for allocating using the output allocator"));
+							PendingMaps.Add(OutputMap);
 						}
 					},
-					TransformerList->CreateSetMapTransformerIterator(),
+					TransformerList->CreateBoundValueMapTransformerIterator(),
 					// Remove the const qualifier since we'll mutate one in place as our output
 					const_cast<FAttributeCollection&>(Base).CreateSetMapIterator(),
 					const_cast<FAttributeCollection&>(Additive).CreateSetMapIterator());
 
 				// Append our pending entries
-				for (FAttributeSetMap* SetMap : PendingSetMaps)
+				for (FBoundValueMap* Map : PendingMaps)
 				{
-					Output.AddSetMap(SetMap);
+					Output.AddSetMap(Map);
 				}
 			}
 
 			// Name/value maps
 			{
 				// Keep a list of maps that needs to be appended to our output
-				TArray<FAttributeMap*, TInlineAllocator<8>> PendingMaps;
+				TArray<FUnboundValueMap*, TInlineAllocator<8>> PendingMaps;
 
 				OuterJoinBy(
-					FAttributeValueTransformerListWithAttributeCollectionMapJoinOp(),
-					[DefaultWeight, bIsOutputBase, OutputAllocator, &PendingMaps](FRawTransformerFunc Transformer, FAttributeMap* BaseMap, FAttributeMap* AdditiveMap)
+					FValueTransformerListWithUnboundMapCollectionJoinOp(),
+					[DefaultWeight, bIsOutputBase, OutputAllocator, &PendingMaps](FRawTransformerFunc Transformer, FUnboundValueMap* BaseMap, FUnboundValueMap* AdditiveMap)
 					{
 						if (AdditiveMap == nullptr)
 						{
@@ -308,9 +307,9 @@ namespace UE::UAF::Transformers
 
 						// Transform our matching pair, note that the base map might be missing
 						// If it is missing, transformer specializations should use a default value
-						FAttributeMap* OutputMap = nullptr;
+						FUnboundValueMap* OutputMap = nullptr;
 
-						FTransformMapFunc TypedTransformer = reinterpret_cast<FTransformMapFunc>(Transformer);
+						FTransformUnboundValueMapFunc TypedTransformer = reinterpret_cast<FTransformUnboundValueMapFunc>(Transformer);
 						(*TypedTransformer)(BaseMap, AdditiveMap, DefaultWeight, OutputMap, OutputAllocator);
 
 						checkf(OutputMap != nullptr, TEXT("Transformer function is responsible for allocating an output buffer"));
@@ -324,7 +323,7 @@ namespace UE::UAF::Transformers
 								OutputMap->MoveTo(*BaseMap);
 
 								// Destroy our temporary output map
-								ReleaseAttributeMap(OutputMap);
+								ReleaseUnboundValueMap(OutputMap);
 							}
 							else
 							{
@@ -338,16 +337,16 @@ namespace UE::UAF::Transformers
 							OutputMap->MoveTo(*AdditiveMap);
 
 							// Destroy our temporary output map
-							ReleaseAttributeMap(OutputMap);
+							ReleaseUnboundValueMap(OutputMap);
 						}
 					},
-					TransformerList->CreateMapTransformerIterator(),
+					TransformerList->CreateUnboundValueMapTransformerIterator(),
 					// Remove the const qualifier since we'll mutate one in place as our output
 					const_cast<FAttributeCollection&>(Base).CreateMapIterator(),
 					const_cast<FAttributeCollection&>(Additive).CreateMapIterator());
 
 				// Append our pending entries
-				for (FAttributeMap* Map : PendingMaps)
+				for (FUnboundValueMap* Map : PendingMaps)
 				{
 					Output.AddMap(Map);
 				}
@@ -361,8 +360,8 @@ namespace UE::UAF::Transformers
 
 			// Set maps
 			OuterJoinBy(
-				FAttributeValueTransformerListWithAttributeCollectionSetMapJoinOp(),
-				[&PerValueWeights, DefaultWeight, FloatValueType, &Output](FRawTransformerFunc Transformer, const FAttributeSetMap* BaseMap, const FAttributeSetMap* AdditiveMap)
+				FValueTransformerListWithBoundMapCollectionJoinOp(),
+				[&PerValueWeights, DefaultWeight, FloatValueType, &Output](FRawTransformerFunc Transformer, const FBoundValueMap* BaseMap, const FBoundValueMap* AdditiveMap)
 				{
 					if (AdditiveMap == nullptr)
 					{
@@ -373,7 +372,7 @@ namespace UE::UAF::Transformers
 					if (Transformer == nullptr)
 					{
 						// No transformer was found for this value type, retain the one with the largest weight
-						FAttributeSetMap* MapCopy;
+						FBoundValueMap* MapCopy;
 						if (BaseMap == nullptr || DefaultWeight >= 0.5f)
 						{
 							// Additive map has more weight or we have no base
@@ -391,27 +390,27 @@ namespace UE::UAF::Transformers
 
 					// Transform our matching pair, note that the base map might be missing
 					// If it is missing, transformer specializations should use a default value
-					FAttributeSetMap* OutputSetMap = nullptr;
+					FBoundValueMap* OutputMap = nullptr;
 
 					// We cannot join our weights because their value type differs
 					const FAttributeMappingKey MappingKey = AdditiveMap->GetMappingKey().To(FloatValueType);
-					const FAttributeSetMap* PerValueWeightMap = PerValueWeights.FindSetMap(MappingKey);
+					const FBoundValueMap* PerValueWeightMap = PerValueWeights.FindSetMap(MappingKey);
 
-					FTransformSetMapFunc TypedTransformer = reinterpret_cast<FTransformSetMapFunc>(Transformer);
-					(*TypedTransformer)(BaseMap, AdditiveMap, PerValueWeightMap, DefaultWeight, OutputSetMap, Output.GetAllocator());
+					FTransformBoundValueMapFunc TypedTransformer = reinterpret_cast<FTransformBoundValueMapFunc>(Transformer);
+					(*TypedTransformer)(BaseMap, AdditiveMap, PerValueWeightMap, DefaultWeight, OutputMap, Output.GetAllocator());
 
-					checkf(OutputSetMap != nullptr, TEXT("Transformer function is responsible for allocating an output buffer"));
-					checkf(OutputSetMap->GetAllocator() == Output.GetAllocator(), TEXT("Transformer function is responsible for allocating using the output allocator"));
-					Output.AppendSetMap(OutputSetMap);
+					checkf(OutputMap != nullptr, TEXT("Transformer function is responsible for allocating an output buffer"));
+					checkf(OutputMap->GetAllocator() == Output.GetAllocator(), TEXT("Transformer function is responsible for allocating using the output allocator"));
+					Output.AppendSetMap(OutputMap);
 				},
-				TransformerList->CreateSetMapTransformerIterator(),
+				TransformerList->CreateBoundValueMapTransformerIterator(),
 				Base.CreateConstSetMapIterator(),
 				Additive.CreateConstSetMapIterator());
 
 			// Name/value maps
 			OuterJoinBy(
-				FAttributeValueTransformerListWithAttributeCollectionMapJoinOp(),
-				[DefaultWeight, &Output](FRawTransformerFunc Transformer, const FAttributeMap* BaseMap, const FAttributeMap* AdditiveMap)
+				FValueTransformerListWithUnboundMapCollectionJoinOp(),
+				[DefaultWeight, &Output](FRawTransformerFunc Transformer, const FUnboundValueMap* BaseMap, const FUnboundValueMap* AdditiveMap)
 				{
 					if (AdditiveMap == nullptr)
 					{
@@ -422,7 +421,7 @@ namespace UE::UAF::Transformers
 					if (Transformer == nullptr)
 					{
 						// No transformer was found for this value type, retain the one with the largest weight
-						FAttributeMap* MapCopy;
+						FUnboundValueMap* MapCopy;
 						if (BaseMap == nullptr || DefaultWeight >= 0.5f)
 						{
 							// Additive map has more weight or we have no base
@@ -440,16 +439,16 @@ namespace UE::UAF::Transformers
 
 					// Transform our matching pair, note that the base map might be missing
 					// If it is missing, transformer specializations should use a default value
-					FAttributeMap* OutputMap = nullptr;
+					FUnboundValueMap* OutputMap = nullptr;
 
-					FTransformMapFunc TypedTransformer = reinterpret_cast<FTransformMapFunc>(Transformer);
+					FTransformUnboundValueMapFunc TypedTransformer = reinterpret_cast<FTransformUnboundValueMapFunc>(Transformer);
 					(*TypedTransformer)(BaseMap, AdditiveMap, DefaultWeight, OutputMap, Output.GetAllocator());
 
 					checkf(OutputMap != nullptr, TEXT("Transformer function is responsible for allocating an output buffer"));
 					checkf(OutputMap->GetAllocator() == Output.GetAllocator(), TEXT("Transformer function is responsible for allocating using the output allocator"));
 					Output.AppendMap(OutputMap);
 				},
-				TransformerList->CreateMapTransformerIterator(),
+				TransformerList->CreateUnboundValueMapTransformerIterator(),
 				Base.CreateConstMapIterator(),
 				Additive.CreateConstMapIterator());
 		}
@@ -460,7 +459,7 @@ namespace UE::UAF::Transformers
 		//////////////////////////////////////////////////////////////////////////
 		// MakeAdditiveSpace transformers
 
-		void FMakeAdditiveSpace_BoneTransformAttribute::TransformSetMap(const FAttributeSetMap* Base, const FAttributeSetMap* Input, FAttributeSetMap*& Output, FReallocFun OutputAllocator)
+		void FMakeAdditiveSpace_BoneTransformAttribute::TransformBoundValueMap(const FBoundValueMap* Base, const FBoundValueMap* Input, FBoundValueMap*& Output, FReallocFun OutputAllocator)
 		{
 			if (Base == nullptr)
 			{
@@ -472,13 +471,13 @@ namespace UE::UAF::Transformers
 
 			if (Output == nullptr)
 			{
-				const FAttributeSetMap::FConstructArgs Args(Input->GetTypedSet(), Input->GetValueType(), OutputAllocator);
-				Output = MakeAttributeSetMap<FBoneTransformAnimationAttribute>(Args);
+				const FBoundValueMap::FConstructArgs Args(Input->GetTypedSet(), Input->GetValueType(), OutputAllocator);
+				Output = MakeBoundValueMap<FBoneTransformAnimationAttribute>(Args);
 			}
 
 			if (Base == nullptr)
 			{
-				UE_LOG(LogAnimation, Warning, TEXT("FMakeAdditiveSpace_BoneTransformAttribute::TransformSetMap: Could not transform set map - no base provided and no reference found. Setting to identity instead."));
+				UE_LOG(LogAnimation, Warning, TEXT("FMakeAdditiveSpace_BoneTransformAttribute::TransformBoundValueMap: Could not transform bound value map - no base provided and no reference found. Setting to identity instead."));
 				Output->FillWithIdentity(true);
 				return;
 			}
@@ -491,9 +490,9 @@ namespace UE::UAF::Transformers
 			check(Base->GetTypedSet() == Input->GetTypedSet());
 			check(Base->Num() == Input->Num());
 
-			const TAttributeSetMap<FBoneTransformAnimationAttribute>* BaseBoneTransforms = CastChecked<FBoneTransformAnimationAttribute>(Base);
-			const TAttributeSetMap<FBoneTransformAnimationAttribute>* InputBoneTransforms = CastChecked<FBoneTransformAnimationAttribute>(Input);
-			TAttributeSetMap<FBoneTransformAnimationAttribute>* OutputBoneTransforms = CastChecked<FBoneTransformAnimationAttribute>(Output);
+			const TBoundValueMap<FBoneTransformAnimationAttribute>* BaseBoneTransforms = CastChecked<FBoneTransformAnimationAttribute>(Base);
+			const TBoundValueMap<FBoneTransformAnimationAttribute>* InputBoneTransforms = CastChecked<FBoneTransformAnimationAttribute>(Input);
+			TBoundValueMap<FBoneTransformAnimationAttribute>* OutputBoneTransforms = CastChecked<FBoneTransformAnimationAttribute>(Output);
 
 			const FBoneTransformAnimationAttribute* BaseTransformPtr = BaseBoneTransforms->GetData();
 			const FBoneTransformAnimationAttribute* InputTransformPtr = InputBoneTransforms->GetData();
@@ -517,17 +516,17 @@ namespace UE::UAF::Transformers
 			}
 		}
 
-		void FMakeAdditiveSpace_BoneTransformAttribute::TransformMap(const FAttributeMap* Base, const FAttributeMap* Input, FAttributeMap*& Output, FReallocFun OutputAllocator)
+		void FMakeAdditiveSpace_BoneTransformAttribute::TransformUnboundValueMap(const FUnboundValueMap* Base, const FUnboundValueMap* Input, FUnboundValueMap*& Output, FReallocFun OutputAllocator)
 		{
 			check(Output == nullptr);
 
-			const TAttributeMap<FBoneTransformAnimationAttribute>* BaseBoneTransforms = Cast<FBoneTransformAnimationAttribute>(Base);
-			const TAttributeMap<FBoneTransformAnimationAttribute>* InputBoneTransforms = CastChecked<FBoneTransformAnimationAttribute>(Input);
+			const TUnboundValueMap<FBoneTransformAnimationAttribute>* BaseBoneTransforms = Cast<FBoneTransformAnimationAttribute>(Base);
+			const TUnboundValueMap<FBoneTransformAnimationAttribute>* InputBoneTransforms = CastChecked<FBoneTransformAnimationAttribute>(Input);
 
-			TAttributeMap<FBoneTransformAnimationAttribute>* OutputBoneTransforms = MakeAttributeMap<FBoneTransformAnimationAttribute>(OutputAllocator);
+			TUnboundValueMap<FBoneTransformAnimationAttribute>* OutputBoneTransforms = MakeUnboundValueMap<FBoneTransformAnimationAttribute>(OutputAllocator);
 
 			OuterJoinBy(
-				FAttributeMapJoinOp(),
+				FUnboundValueMapJoinOp(),
 				[OutputBoneTransforms](FName Name, const FBoneTransformAnimationAttribute* BaseValue, const FBoneTransformAnimationAttribute* InputValue)
 				{
 					if (InputValue == nullptr)
@@ -551,13 +550,13 @@ namespace UE::UAF::Transformers
 
 					OutputBoneTransforms->Append(Name, Output);
 				},
-				BaseBoneTransforms != nullptr ? BaseBoneTransforms->CreateConstIterator() : TAttributeMap<FBoneTransformAnimationAttribute>::FConstIterator(),
+				BaseBoneTransforms != nullptr ? BaseBoneTransforms->CreateConstIterator() : TUnboundValueMap<FBoneTransformAnimationAttribute>::FConstIterator(),
 				InputBoneTransforms->CreateConstIterator());
 
 			Output = OutputBoneTransforms;
 		}
 
-		void FMakeAdditiveSpace_FloatAttribute::TransformSetMap(const FAttributeSetMap* Base, const FAttributeSetMap* Input, FAttributeSetMap*& Output, FReallocFun OutputAllocator)
+		void FMakeAdditiveSpace_FloatAttribute::TransformBoundValueMap(const FBoundValueMap* Base, const FBoundValueMap* Input, FBoundValueMap*& Output, FReallocFun OutputAllocator)
 		{
 			if (Base == nullptr)
 			{
@@ -569,13 +568,13 @@ namespace UE::UAF::Transformers
 
 			if (Output == nullptr)
 			{
-				const FAttributeSetMap::FConstructArgs Args(Input->GetTypedSet(), Input->GetValueType(), OutputAllocator);
-				Output = MakeAttributeSetMap<FFloatAnimationAttribute>(Args);
+				const FBoundValueMap::FConstructArgs Args(Input->GetTypedSet(), Input->GetValueType(), OutputAllocator);
+				Output = MakeBoundValueMap<FFloatAnimationAttribute>(Args);
 			}
 
 			if (Base == nullptr)
 			{
-				UE_LOG(LogAnimation, Warning, TEXT("FMakeAdditiveSpace_FloatAttribute::TransformSetMap: Could not transform set map - no base provided and no reference found. Setting to identity instead."));
+				UE_LOG(LogAnimation, Warning, TEXT("FMakeAdditiveSpace_FloatAttribute::TransformBoundValueMap: Could not transform bound value map - no base provided and no reference found. Setting to identity instead."));
 				Output->FillWithIdentity(true);
 				return;
 			}
@@ -588,9 +587,9 @@ namespace UE::UAF::Transformers
 			check(Base->GetTypedSet() == Input->GetTypedSet());
 			check(Base->Num() == Input->Num());
 
-			const TAttributeSetMap<FFloatAnimationAttribute>* BaseFloats = CastChecked<FFloatAnimationAttribute>(Base);
-			const TAttributeSetMap<FFloatAnimationAttribute>* InputFloats = CastChecked<FFloatAnimationAttribute>(Input);
-			TAttributeSetMap<FFloatAnimationAttribute>* OutputFloats = CastChecked<FFloatAnimationAttribute>(Output);
+			const TBoundValueMap<FFloatAnimationAttribute>* BaseFloats = CastChecked<FFloatAnimationAttribute>(Base);
+			const TBoundValueMap<FFloatAnimationAttribute>* InputFloats = CastChecked<FFloatAnimationAttribute>(Input);
+			TBoundValueMap<FFloatAnimationAttribute>* OutputFloats = CastChecked<FFloatAnimationAttribute>(Output);
 
 			const FFloatAnimationAttribute* BaseFloatsPtr = BaseFloats->GetData();
 			const FFloatAnimationAttribute* InputFloatsPtr = InputFloats->GetData();
@@ -607,17 +606,17 @@ namespace UE::UAF::Transformers
 			}
 		}
 
-		void FMakeAdditiveSpace_FloatAttribute::TransformMap(const FAttributeMap* Base, const FAttributeMap* Input, FAttributeMap*& Output, FReallocFun OutputAllocator)
+		void FMakeAdditiveSpace_FloatAttribute::TransformUnboundValueMap(const FUnboundValueMap* Base, const FUnboundValueMap* Input, FUnboundValueMap*& Output, FReallocFun OutputAllocator)
 		{
 			check(Output == nullptr);
 
-			const TAttributeMap<FFloatAnimationAttribute>* BaseFloats = Cast<FFloatAnimationAttribute>(Base);
-			const TAttributeMap<FFloatAnimationAttribute>* InputFloats = CastChecked<FFloatAnimationAttribute>(Input);
+			const TUnboundValueMap<FFloatAnimationAttribute>* BaseFloats = Cast<FFloatAnimationAttribute>(Base);
+			const TUnboundValueMap<FFloatAnimationAttribute>* InputFloats = CastChecked<FFloatAnimationAttribute>(Input);
 
-			TAttributeMap<FFloatAnimationAttribute>* OutputFloats = MakeAttributeMap<FFloatAnimationAttribute>(OutputAllocator);
+			TUnboundValueMap<FFloatAnimationAttribute>* OutputFloats = MakeUnboundValueMap<FFloatAnimationAttribute>(OutputAllocator);
 
 			OuterJoinBy(
-				FAttributeMapJoinOp(),
+				FUnboundValueMapJoinOp(),
 				[OutputFloats](FName Name, const FFloatAnimationAttribute* BaseValue, const FFloatAnimationAttribute* InputValue)
 				{
 					if (InputValue == nullptr)
@@ -632,7 +631,7 @@ namespace UE::UAF::Transformers
 					const FFloatAnimationAttribute Output{ InputValue->Value - BaseFloat };
 					OutputFloats->Append(Name, Output);
 				},
-				BaseFloats != nullptr ? BaseFloats->CreateConstIterator() : TAttributeMap<FFloatAnimationAttribute>::FConstIterator(),
+				BaseFloats != nullptr ? BaseFloats->CreateConstIterator() : TUnboundValueMap<FFloatAnimationAttribute>::FConstIterator(),
 				InputFloats->CreateConstIterator());
 
 			Output = OutputFloats;
@@ -641,7 +640,7 @@ namespace UE::UAF::Transformers
 		//////////////////////////////////////////////////////////////////////////
 		// ApplyAdditiveSpace transformers
 
-		void FApplyAdditiveSpace_BoneTransformAttribute::TransformSetMap(const FAttributeSetMap* Base, const FAttributeSetMap* Additive, const FAttributeSetMap* PerValueWeightMap, float DefaultWeight, FAttributeSetMap*& Output, FReallocFun OutputAllocator)
+		void FApplyAdditiveSpace_BoneTransformAttribute::TransformBoundValueMap(const FBoundValueMap* Base, const FBoundValueMap* Additive, const FBoundValueMap* PerValueWeightMap, float DefaultWeight, FBoundValueMap*& Output, FReallocFun OutputAllocator)
 		{
 			if (Base == nullptr)
 			{
@@ -653,13 +652,13 @@ namespace UE::UAF::Transformers
 
 			if (Output == nullptr)
 			{
-				const FAttributeSetMap::FConstructArgs Args(Additive->GetTypedSet(), Additive->GetValueType(), OutputAllocator);
-				Output = MakeAttributeSetMap<FBoneTransformAnimationAttribute>(Args);
+				const FBoundValueMap::FConstructArgs Args(Additive->GetTypedSet(), Additive->GetValueType(), OutputAllocator);
+				Output = MakeBoundValueMap<FBoneTransformAnimationAttribute>(Args);
 			}
 
 			if (Base == nullptr)
 			{
-				UE_LOG(LogAnimation, Warning, TEXT("FApplyAdditiveSpace_BoneTransformAttribute::TransformSetMap: Could not transform set map - no base provided and no reference found. Setting to identity instead."));
+				UE_LOG(LogAnimation, Warning, TEXT("FApplyAdditiveSpace_BoneTransformAttribute::TransformBoundValueMap: Could not transform bound value map - no base provided and no reference found. Setting to identity instead."));
 				Output->FillWithIdentity(false);
 				return;
 			}
@@ -672,16 +671,16 @@ namespace UE::UAF::Transformers
 			check(Base->GetTypedSet() == Additive->GetTypedSet());
 			check(Base->Num() == Additive->Num());
 
-			const TAttributeSetMap<FBoneTransformAnimationAttribute>* BaseBoneTransforms = CastChecked<FBoneTransformAnimationAttribute>(Base);
-			const TAttributeSetMap<FBoneTransformAnimationAttribute>* AdditiveBoneTransforms = CastChecked<FBoneTransformAnimationAttribute>(Additive);
-			TAttributeSetMap<FBoneTransformAnimationAttribute>* OutputBoneTransforms = CastChecked<FBoneTransformAnimationAttribute>(Output);
+			const TBoundValueMap<FBoneTransformAnimationAttribute>* BaseBoneTransforms = CastChecked<FBoneTransformAnimationAttribute>(Base);
+			const TBoundValueMap<FBoneTransformAnimationAttribute>* AdditiveBoneTransforms = CastChecked<FBoneTransformAnimationAttribute>(Additive);
+			TBoundValueMap<FBoneTransformAnimationAttribute>* OutputBoneTransforms = CastChecked<FBoneTransformAnimationAttribute>(Output);
 
 			const FBoneTransformAnimationAttribute* BaseTransformPtr = BaseBoneTransforms->GetData();
 			const FBoneTransformAnimationAttribute* AdditiveTransformPtr = AdditiveBoneTransforms->GetData();
 			FBoneTransformAnimationAttribute* OutputTransformPtr = OutputBoneTransforms->GetData();
 			const FBoneTransformAnimationAttribute* BaseTransformPtrEnd = BaseTransformPtr + BaseBoneTransforms->Num();
 
-			if (const TAttributeSetMap<FFloatAnimationAttribute>* PerValueWeights = Cast<FFloatAnimationAttribute>(PerValueWeightMap))
+			if (const TBoundValueMap<FFloatAnimationAttribute>* PerValueWeights = Cast<FFloatAnimationAttribute>(PerValueWeightMap))
 			{
 				check(Base->GetTypedSet() == PerValueWeights->GetTypedSet());
 				check(Base->Num() == PerValueWeights->Num());
@@ -739,14 +738,14 @@ namespace UE::UAF::Transformers
 			}
 		}
 
-		void FApplyAdditiveSpace_BoneTransformAttribute::TransformMap(const FAttributeMap* Base, const FAttributeMap* Additive, float AdditiveWeight, FAttributeMap*& Output, FReallocFun OutputAllocator)
+		void FApplyAdditiveSpace_BoneTransformAttribute::TransformUnboundValueMap(const FUnboundValueMap* Base, const FUnboundValueMap* Additive, float AdditiveWeight, FUnboundValueMap*& Output, FReallocFun OutputAllocator)
 		{
 			check(Output == nullptr);
 
-			const TAttributeMap<FBoneTransformAnimationAttribute>* BaseBoneTransforms = Cast<FBoneTransformAnimationAttribute>(Base);
-			const TAttributeMap<FBoneTransformAnimationAttribute>* AdditiveBoneTransforms = CastChecked<FBoneTransformAnimationAttribute>(Additive);
+			const TUnboundValueMap<FBoneTransformAnimationAttribute>* BaseBoneTransforms = Cast<FBoneTransformAnimationAttribute>(Base);
+			const TUnboundValueMap<FBoneTransformAnimationAttribute>* AdditiveBoneTransforms = CastChecked<FBoneTransformAnimationAttribute>(Additive);
 
-			TAttributeMap<FBoneTransformAnimationAttribute>* OutputBoneTransforms = MakeAttributeMap<FBoneTransformAnimationAttribute>(OutputAllocator);
+			TUnboundValueMap<FBoneTransformAnimationAttribute>* OutputBoneTransforms = MakeUnboundValueMap<FBoneTransformAnimationAttribute>(OutputAllocator);
 
 			const ScalarRegister VAdditiveWeight(AdditiveWeight);
 
@@ -754,7 +753,7 @@ namespace UE::UAF::Transformers
 			{
 				// Our additive has full weight, no need to interpolate with identity
 				OuterJoinBy(
-					FAttributeMapJoinOp(),
+					FUnboundValueMapJoinOp(),
 					[VAdditiveWeight, OutputBoneTransforms](FName Name, const FBoneTransformAnimationAttribute* BaseValue, const FBoneTransformAnimationAttribute* AdditiveValue)
 					{
 						// Missing values are set to the identity
@@ -770,14 +769,14 @@ namespace UE::UAF::Transformers
 
 						OutputBoneTransforms->Append(Name, Output);
 					},
-					BaseBoneTransforms != nullptr ? BaseBoneTransforms->CreateConstIterator() : TAttributeMap<FBoneTransformAnimationAttribute>::FConstIterator(),
+					BaseBoneTransforms != nullptr ? BaseBoneTransforms->CreateConstIterator() : TUnboundValueMap<FBoneTransformAnimationAttribute>::FConstIterator(),
 					AdditiveBoneTransforms->CreateConstIterator());
 			}
 			else
 			{
 				// Interpolate with identity before adding
 				OuterJoinBy(
-					FAttributeMapJoinOp(),
+					FUnboundValueMapJoinOp(),
 					[VAdditiveWeight, OutputBoneTransforms](FName Name, const FBoneTransformAnimationAttribute* BaseValue, const FBoneTransformAnimationAttribute* AdditiveValue)
 					{
 						// Missing values are set to the identity
@@ -792,14 +791,14 @@ namespace UE::UAF::Transformers
 
 						OutputBoneTransforms->Append(Name, Output);
 					},
-					BaseBoneTransforms != nullptr ? BaseBoneTransforms->CreateConstIterator() : TAttributeMap<FBoneTransformAnimationAttribute>::FConstIterator(),
+					BaseBoneTransforms != nullptr ? BaseBoneTransforms->CreateConstIterator() : TUnboundValueMap<FBoneTransformAnimationAttribute>::FConstIterator(),
 					AdditiveBoneTransforms->CreateConstIterator());
 			}
 
 			Output = OutputBoneTransforms;
 		}
 
-		void FApplyAdditiveSpace_FloatAttribute::TransformSetMap(const FAttributeSetMap* Base, const FAttributeSetMap* Additive, const FAttributeSetMap* PerValueWeightMap, float DefaultWeight, FAttributeSetMap*& Output, FReallocFun OutputAllocator)
+		void FApplyAdditiveSpace_FloatAttribute::TransformBoundValueMap(const FBoundValueMap* Base, const FBoundValueMap* Additive, const FBoundValueMap* PerValueWeightMap, float DefaultWeight, FBoundValueMap*& Output, FReallocFun OutputAllocator)
 		{
 			if (Base == nullptr)
 			{
@@ -811,13 +810,13 @@ namespace UE::UAF::Transformers
 
 			if (Output == nullptr)
 			{
-				const FAttributeSetMap::FConstructArgs Args(Additive->GetTypedSet(), Additive->GetValueType(), OutputAllocator);
-				Output = MakeAttributeSetMap<FFloatAnimationAttribute>(Args);
+				const FBoundValueMap::FConstructArgs Args(Additive->GetTypedSet(), Additive->GetValueType(), OutputAllocator);
+				Output = MakeBoundValueMap<FFloatAnimationAttribute>(Args);
 			}
 
 			if (Base == nullptr)
 			{
-				UE_LOG(LogAnimation, Warning, TEXT("FApplyAdditiveSpace_FloatAttribute::TransformSetMap: Could not transform set map - no base provided and no reference found. Setting to identity instead."));
+				UE_LOG(LogAnimation, Warning, TEXT("FApplyAdditiveSpace_FloatAttribute::TransformBoundValueMap: Could not transform bound value map - no base provided and no reference found. Setting to identity instead."));
 				Output->FillWithIdentity(false);
 				return;
 			}
@@ -830,16 +829,16 @@ namespace UE::UAF::Transformers
 			check(Base->GetTypedSet() == Additive->GetTypedSet());
 			check(Base->Num() == Additive->Num());
 
-			const TAttributeSetMap<FFloatAnimationAttribute>* BaseFloats = CastChecked<FFloatAnimationAttribute>(Base);
-			const TAttributeSetMap<FFloatAnimationAttribute>* AdditiveFloats = CastChecked<FFloatAnimationAttribute>(Additive);
-			TAttributeSetMap<FFloatAnimationAttribute>* OutputFloats = CastChecked<FFloatAnimationAttribute>(Output);
+			const TBoundValueMap<FFloatAnimationAttribute>* BaseFloats = CastChecked<FFloatAnimationAttribute>(Base);
+			const TBoundValueMap<FFloatAnimationAttribute>* AdditiveFloats = CastChecked<FFloatAnimationAttribute>(Additive);
+			TBoundValueMap<FFloatAnimationAttribute>* OutputFloats = CastChecked<FFloatAnimationAttribute>(Output);
 
 			const FFloatAnimationAttribute* BaseFloatsPtr = BaseFloats->GetData();
 			const FFloatAnimationAttribute* AdditiveFloatsPtr = AdditiveFloats->GetData();
 			FFloatAnimationAttribute* OutputFloatsPtr = OutputFloats->GetData();
 			const FFloatAnimationAttribute* BaseFloatsPtrEnd = BaseFloatsPtr + BaseFloats->Num();
 
-			if (const TAttributeSetMap<FFloatAnimationAttribute>* PerValueWeights = Cast<FFloatAnimationAttribute>(PerValueWeightMap))
+			if (const TBoundValueMap<FFloatAnimationAttribute>* PerValueWeights = Cast<FFloatAnimationAttribute>(PerValueWeightMap))
 			{
 				check(Base->GetTypedSet() == PerValueWeights->GetTypedSet());
 				check(Base->Num() == PerValueWeights->Num());
@@ -869,17 +868,17 @@ namespace UE::UAF::Transformers
 			}
 		}
 
-		void FApplyAdditiveSpace_FloatAttribute::TransformMap(const FAttributeMap* Base, const FAttributeMap* Additive, float AdditiveWeight, FAttributeMap*& Output, FReallocFun OutputAllocator)
+		void FApplyAdditiveSpace_FloatAttribute::TransformUnboundValueMap(const FUnboundValueMap* Base, const FUnboundValueMap* Additive, float AdditiveWeight, FUnboundValueMap*& Output, FReallocFun OutputAllocator)
 		{
 			check(Output == nullptr);
 
-			const TAttributeMap<FFloatAnimationAttribute>* BaseFloats = Cast<FFloatAnimationAttribute>(Base);
-			const TAttributeMap<FFloatAnimationAttribute>* AdditiveFloats = CastChecked<FFloatAnimationAttribute>(Additive);
+			const TUnboundValueMap<FFloatAnimationAttribute>* BaseFloats = Cast<FFloatAnimationAttribute>(Base);
+			const TUnboundValueMap<FFloatAnimationAttribute>* AdditiveFloats = CastChecked<FFloatAnimationAttribute>(Additive);
 
-			TAttributeMap<FFloatAnimationAttribute>* OutputFloats = MakeAttributeMap<FFloatAnimationAttribute>(OutputAllocator);
+			TUnboundValueMap<FFloatAnimationAttribute>* OutputFloats = MakeUnboundValueMap<FFloatAnimationAttribute>(OutputAllocator);
 
 			OuterJoinBy(
-				FAttributeMapJoinOp(),
+				FUnboundValueMapJoinOp(),
 				[AdditiveWeight, OutputFloats](FName Name, const FFloatAnimationAttribute* BaseValue, const FFloatAnimationAttribute* AdditiveValue)
 				{
 					// Missing values are set to 0.0
@@ -889,7 +888,7 @@ namespace UE::UAF::Transformers
 					const FFloatAnimationAttribute Output{ BaseFloat + (AdditiveWeight * AdditiveFloat) };
 					OutputFloats->Append(Name, Output);
 				},
-				BaseFloats != nullptr ? BaseFloats->CreateConstIterator() : TAttributeMap<FFloatAnimationAttribute>::FConstIterator(),
+				BaseFloats != nullptr ? BaseFloats->CreateConstIterator() : TUnboundValueMap<FFloatAnimationAttribute>::FConstIterator(),
 				AdditiveFloats->CreateConstIterator());
 
 			Output = OutputFloats;
diff --git a/Engine/Plugins/Experimental/UAF/UAF/Source/UAF/Private/UAF/ValueRuntime/Transformers/BoneSpace.cpp b/Engine/Plugins/Experimental/UAF/UAF/Source/UAF/Private/UAF/ValueRuntime/Transformers/BoneSpace.cpp
--- a/Engine/Plugins/Experimental/UAF/UAF/Source/UAF/Private/UAF/ValueRuntime/Transformers/BoneSpace.cpp
+++ b/Engine/Plugins/Experimental/UAF/UAF/Source/UAF/Private/UAF/ValueRuntime/Transformers/BoneSpace.cpp
@@ -21,7 +21,7 @@ namespace UE::UAF::Transformers
 		OutputValues.SetValueSpace(FValueSpace(EValueSpaceType::Component));
 
 		// Modify in place
-		TAttributeSetMap<FBoneTransformAnimationAttribute>* BoneTransforms = OutputValues.FindSetMap<FBoneTransformAnimationAttribute>(FAttributeMappingKey::MakeFromTo<FBoneTransformAnimationAttribute>());
+		TBoundValueMap<FBoneTransformAnimationAttribute>* BoneTransforms = OutputValues.FindSetMap<FBoneTransformAnimationAttribute>(FAttributeMappingKey::MakeFromTo<FBoneTransformAnimationAttribute>());
 		if (BoneTransforms == nullptr)
 		{
 			// No bone transforms, nothing to convert
@@ -66,7 +66,7 @@ namespace UE::UAF::Transformers
 		OutputValues.SetValueSpace(FValueSpace(EValueSpaceType::Local));
 
 		// Modify in place
-		TAttributeSetMap<FBoneTransformAnimationAttribute>* BoneTransforms = OutputValues.FindSetMap<FBoneTransformAnimationAttribute>(FAttributeMappingKey::MakeFromTo<FBoneTransformAnimationAttribute>());
+		TBoundValueMap<FBoneTransformAnimationAttribute>* BoneTransforms = OutputValues.FindSetMap<FBoneTransformAnimationAttribute>(FAttributeMappingKey::MakeFromTo<FBoneTransformAnimationAttribute>());
 		if (BoneTransforms == nullptr)
 		{
 			// No bone transforms, nothing to convert
@@ -116,7 +116,7 @@ namespace UE::UAF::Transformers
 		OutputValues.SetValueSpace(FValueSpace(MixedSpaceFlags));
 
 		// Modify in place
-		TAttributeSetMap<FBoneTransformAnimationAttribute>* BoneTransforms = OutputValues.FindSetMap<FBoneTransformAnimationAttribute>(FAttributeMappingKey::MakeFromTo<FBoneTransformAnimationAttribute>());
+		TBoundValueMap<FBoneTransformAnimationAttribute>* BoneTransforms = OutputValues.FindSetMap<FBoneTransformAnimationAttribute>(FAttributeMappingKey::MakeFromTo<FBoneTransformAnimationAttribute>());
 		if (BoneTransforms == nullptr)
 		{
 			// No bone transforms, nothing to convert
@@ -174,7 +174,7 @@ namespace UE::UAF::Transformers
 		OutputValues.SetValueSpace(OutputSpace);
 
 		// Modify in place
-		TAttributeSetMap<FBoneTransformAnimationAttribute>* BoneTransforms = OutputValues.FindSetMap<FBoneTransformAnimationAttribute>(FAttributeMappingKey::MakeFromTo<FBoneTransformAnimationAttribute>());
+		TBoundValueMap<FBoneTransformAnimationAttribute>* BoneTransforms = OutputValues.FindSetMap<FBoneTransformAnimationAttribute>(FAttributeMappingKey::MakeFromTo<FBoneTransformAnimationAttribute>());
 		if (BoneTransforms == nullptr)
 		{
 			// No bone transforms, nothing to convert
@@ -224,7 +224,7 @@ namespace UE::UAF::Transformers
 		OutputValues.SetValueSpace(FValueSpace(MixedSpaceFlags));
 
 		// Modify in place
-		TAttributeSetMap<FBoneTransformAnimationAttribute>* BoneTransforms = OutputValues.FindSetMap<FBoneTransformAnimationAttribute>(FAttributeMappingKey::MakeFromTo<FBoneTransformAnimationAttribute>());
+		TBoundValueMap<FBoneTransformAnimationAttribute>* BoneTransforms = OutputValues.FindSetMap<FBoneTransformAnimationAttribute>(FAttributeMappingKey::MakeFromTo<FBoneTransformAnimationAttribute>());
 		if (BoneTransforms == nullptr)
 		{
 			// No bone transforms, nothing to convert
@@ -284,7 +284,7 @@ namespace UE::UAF::Transformers
 		OutputValues.SetValueSpace(OutputSpace);
 
 		// Modify in place
-		TAttributeSetMap<FBoneTransformAnimationAttribute>* BoneTransforms = OutputValues.FindSetMap<FBoneTransformAnimationAttribute>(FAttributeMappingKey::MakeFromTo<FBoneTransformAnimationAttribute>());
+		TBoundValueMap<FBoneTransformAnimationAttribute>* BoneTransforms = OutputValues.FindSetMap<FBoneTransformAnimationAttribute>(FAttributeMappingKey::MakeFromTo<FBoneTransformAnimationAttribute>());
 		if (BoneTransforms == nullptr)
 		{
 			// No bone transforms, nothing to convert
@@ -335,7 +335,7 @@ namespace UE::UAF::Transformers
 		OutputValues.SetValueSpace(FValueSpace(MixedSpaceFlags));
 
 		// Modify in place
-		TAttributeSetMap<FBoneTransformAnimationAttribute>* BoneTransforms = OutputValues.FindSetMap<FBoneTransformAnimationAttribute>(FAttributeMappingKey::MakeFromTo<FBoneTransformAnimationAttribute>());
+		TBoundValueMap<FBoneTransformAnimationAttribute>* BoneTransforms = OutputValues.FindSetMap<FBoneTransformAnimationAttribute>(FAttributeMappingKey::MakeFromTo<FBoneTransformAnimationAttribute>());
 		if (BoneTransforms == nullptr)
 		{
 			// No bone transforms, nothing to convert
@@ -394,7 +394,7 @@ namespace UE::UAF::Transformers
 		OutputValues.SetValueSpace(OutputSpace);
 
 		// Modify in place
-		TAttributeSetMap<FBoneTransformAnimationAttribute>* BoneTransforms = OutputValues.FindSetMap<FBoneTransformAnimationAttribute>(FAttributeMappingKey::MakeFromTo<FBoneTransformAnimationAttribute>());
+		TBoundValueMap<FBoneTransformAnimationAttribute>* BoneTransforms = OutputValues.FindSetMap<FBoneTransformAnimationAttribute>(FAttributeMappingKey::MakeFromTo<FBoneTransformAnimationAttribute>());
 		if (BoneTransforms == nullptr)
 		{
 			// No bone transforms, nothing to convert
