Rename UAF ValueRuntime types

diff --git a/Engine/Plugins/Experimental/UAF/UAF/Source/UAF/Private/UAF/ValueRuntime/Transformers/Interpolate.cpp b/Engine/Plugins/Experimental/UAF/UAF/Source/UAF/Private/UAF/ValueRuntime/Transformers/Interpolate.cpp
--- a/Engine/Plugins/Experimental/UAF/UAF/Source/UAF/Private/UAF/ValueRuntime/Transformers/Interpolate.cpp
+++ b/Engine/Plugins/Experimental/UAF/UAF/Source/UAF/Private/UAF/ValueRuntime/Transformers/Interpolate.cpp
@@ -9,7 +9,7 @@ namespace UE::UAF::Transformers
 	FName FInterpolate::TransformerName = TEXT("UAF::Transformers::Interpolate");
 
 	void FInterpolate::Apply(
-		const FAttributeValueTransformerMapPtr& TransformerMap,
+		const FValueTransformerMapPtr& TransformerMap,
 		const FAttributeCollection& InputA,
 		const FAttributeCollection& InputB,
 		float WeightB,
@@ -39,7 +39,7 @@ namespace UE::UAF::Transformers
 	}
 
 	void FInterpolate::Apply(
-		const FAttributeValueTransformerMapPtr& TransformerMap,
+		const FValueTransformerMapPtr& TransformerMap,
 		const FAttributeCollection& InputA,
 		const FAttributeCollection& InputB,
 		const FAttributeCollection& PerValueWeightsB,
@@ -51,7 +51,7 @@ namespace UE::UAF::Transformers
 		checkf(InputA.GetNamedSet() == PerValueWeightsB.GetNamedSet(), TEXT("Input and per-value weight collections must use the same named set"));
 		checkf(InputA.GetValueSpace() == InputB.GetValueSpace(), TEXT("Both input collections must have matching value spaces: [%s] != [%s]"), *InputA.GetValueSpace().ToString(), *InputB.GetValueSpace().ToString());
 
-		const FAttributeValueTransformerList* TransformerList = TransformerMap->Find(TransformerName);
+		const FValueTransformerList* TransformerList = TransformerMap->Find(TransformerName);
 		if (TransformerList == nullptr)
 		{
 			// No transformer implementations have been registered
@@ -67,18 +67,18 @@ namespace UE::UAF::Transformers
 			// Our output collection matches one of our inputs, we'll modify it in place
 			FReallocFun OutputAllocator = Output.GetAllocator();
 
-			// Set maps
+			// Bound value maps
 			{
 				// Keep a list of set maps that needs to be appended to our output
-				TArray<FAttributeSetMap*, TInlineAllocator<8>> PendingSetMaps;
+				TArray<FBoundValueMap*, TInlineAllocator<8>> PendingMaps;
 
 				OuterJoinBy(
-					FAttributeValueTransformerListWithAttributeCollectionSetMapJoinOp(),
-					[&PerValueWeightsB, DefaultWeightB, bIsOutputA, FloatValueType, OutputAllocator, &PendingSetMaps](FRawTransformerFunc Transformer, FAttributeSetMap* SetMapA, FAttributeSetMap* SetMapB)
+					FValueTransformerListWithBoundMapCollectionJoinOp(),
+					[&PerValueWeightsB, DefaultWeightB, bIsOutputA, FloatValueType, OutputAllocator, &PendingMaps](FRawTransformerFunc Transformer, FBoundValueMap* MapA, FBoundValueMap* MapB)
 					{
-						if (SetMapA == nullptr)
+						if (MapA == nullptr)
 						{
-							if (SetMapB == nullptr)
+							if (MapB == nullptr)
 							{
 								// No maps found for the transformer type
 							}
@@ -88,7 +88,7 @@ namespace UE::UAF::Transformers
 								if (bIsOutputA)
 								{
 									// A is our output, duplicate it
-									PendingSetMaps.Add(SetMapB->Duplicate(OutputAllocator));
+									PendingMaps.Add(MapB->Duplicate(OutputAllocator));
 								}
 								else
 								{
@@ -98,7 +98,7 @@ namespace UE::UAF::Transformers
 						}
 						else
 						{
-							if (SetMapB == nullptr)
+							if (MapB == nullptr)
 							{
 								// B doesn't have a matching pair
 								if (bIsOutputA)
@@ -108,7 +108,7 @@ namespace UE::UAF::Transformers
 								else
 								{
 									// B is our output, duplicate it
-									PendingSetMaps.Add(SetMapA->Duplicate(OutputAllocator));
+									PendingMaps.Add(MapA->Duplicate(OutputAllocator));
 								}
 							}
 							else
@@ -122,7 +122,7 @@ namespace UE::UAF::Transformers
 										if (bIsOutputA)
 										{
 											// A is our output, copy it over
-											SetMapB->CopyTo(*SetMapA);
+											MapB->CopyTo(*MapA);
 										}
 										else
 										{
@@ -139,45 +139,45 @@ namespace UE::UAF::Transformers
 										else
 										{
 											// B is our output, copy it over
-											SetMapA->CopyTo(*SetMapB);
+											MapA->CopyTo(*MapB);
 										}
 									}
 								}
 								else
 								{
 									// Transform our matching pair in place
-									FAttributeSetMap* OutputSetMap = bIsOutputA ? SetMapA : SetMapB;
+									FBoundValueMap* OutputMap = bIsOutputA ? MapA : MapB;
 
 									// We cannot join our weights because their value type differs
-									const FAttributeMappingKey MappingKey = SetMapA->GetMappingKey().To(FloatValueType);
-									const FAttributeSetMap* PerValueWeightMapB = PerValueWeightsB.FindSetMap(MappingKey);
+									const FAttributeMappingKey MappingKey = MapA->GetMappingKey().To(FloatValueType);
+									const FBoundValueMap* PerValueWeightMapB = PerValueWeightsB.FindSetMap(MappingKey);
 
-									FTransformSetMapFunc TypedTransformer = reinterpret_cast<FTransformSetMapFunc>(Transformer);
-									(*TypedTransformer)(SetMapA, SetMapB, PerValueWeightMapB, DefaultWeightB, OutputSetMap, nullptr);
+									FTransformBoundValueMapFunc TypedTransformer = reinterpret_cast<FTransformBoundValueMapFunc>(Transformer);
+									(*TypedTransformer)(MapA, MapB, PerValueWeightMapB, DefaultWeightB, OutputMap, nullptr);
 								}
 							}
 						}
 					},
-					TransformerList->CreateSetMapTransformerIterator(),
+					TransformerList->CreateBoundValueMapTransformerIterator(),
 					// Remove the const qualifier since we'll mutate one in place as our output
 					const_cast<FAttributeCollection&>(InputA).CreateSetMapIterator(),
 					const_cast<FAttributeCollection&>(InputB).CreateSetMapIterator());
 
 				// Append our pending entries
-				for (FAttributeSetMap* SetMap : PendingSetMaps)
+				for (FBoundValueMap* Map : PendingMaps)
 				{
-					Output.AddSetMap(SetMap);
+					Output.AddSetMap(Map);
 				}
 			}
 
 			// Name/value maps
 			{
 				// Keep a list of maps that needs to be appended to our output
-				TArray<FAttributeMap*, TInlineAllocator<8>> PendingMaps;
+				TArray<FUnboundValueMap*, TInlineAllocator<8>> PendingMaps;
 
 				OuterJoinBy(
-					FAttributeValueTransformerListWithAttributeCollectionMapJoinOp(),
-					[DefaultWeightB, bIsOutputA, OutputAllocator, &PendingMaps](FRawTransformerFunc Transformer, FAttributeMap* MapA, FAttributeMap* MapB)
+					FValueTransformerListWithUnboundMapCollectionJoinOp(),
+					[DefaultWeightB, bIsOutputA, OutputAllocator, &PendingMaps](FRawTransformerFunc Transformer, FUnboundValueMap* MapA, FUnboundValueMap* MapB)
 					{
 						if (MapA == nullptr)
 						{
@@ -249,9 +249,9 @@ namespace UE::UAF::Transformers
 								else
 								{
 									// Transform our matching pair in a new map
-									FAttributeMap* OutputMap = nullptr;
+									FUnboundValueMap* OutputMap = nullptr;
 
-									FTransformMapFunc TypedTransformer = reinterpret_cast<FTransformMapFunc>(Transformer);
+									FTransformUnboundValueMapFunc TypedTransformer = reinterpret_cast<FTransformUnboundValueMapFunc>(Transformer);
 									(*TypedTransformer)(MapA, MapB, DefaultWeightB, OutputMap, OutputAllocator);
 
 									checkf(OutputMap != nullptr, TEXT("Transformer function is responsible for allocating an output buffer"));
@@ -261,18 +261,18 @@ namespace UE::UAF::Transformers
 									OutputMap->MoveTo(bIsOutputA ? *MapA : *MapB);
 
 									// Destroy our temporary output map
-									ReleaseAttributeMap(OutputMap);
+									ReleaseUnboundValueMap(OutputMap);
 								}
 							}
 						}
 					},
-					TransformerList->CreateMapTransformerIterator(),
+					TransformerList->CreateUnboundValueMapTransformerIterator(),
 					// Remove the const qualifier since we'll mutate one in place as our output
 					const_cast<FAttributeCollection&>(InputA).CreateMapIterator(),
 					const_cast<FAttributeCollection&>(InputB).CreateMapIterator());
 
 				// Append our pending entries
-				for (FAttributeMap* Map : PendingMaps)
+				for (FUnboundValueMap* Map : PendingMaps)
 				{
 					Output.AddMap(Map);
 				}
@@ -284,74 +284,74 @@ namespace UE::UAF::Transformers
 			Output.Reset(InputA.GetNamedSet());
 			Output.SetValueSpace(InputA.GetValueSpace());
 
-			// Set maps
+			// Bound value maps
 			OuterJoinBy(
-				FAttributeValueTransformerListWithAttributeCollectionSetMapJoinOp(),
-				[&PerValueWeightsB, DefaultWeightB, FloatValueType, &Output](FRawTransformerFunc Transformer, const FAttributeSetMap* SetMapA, const FAttributeSetMap* SetMapB)
+				FValueTransformerListWithBoundMapCollectionJoinOp(),
+				[&PerValueWeightsB, DefaultWeightB, FloatValueType, &Output](FRawTransformerFunc Transformer, const FBoundValueMap* MapA, const FBoundValueMap* MapB)
 				{
-					if (SetMapA == nullptr)
+					if (MapA == nullptr)
 					{
-						if (SetMapB == nullptr)
+						if (MapB == nullptr)
 						{
 							// No maps found for the transformer type
 						}
 						else
 						{
 							// This set map doesn't have a matching pair, duplicate it
-							Output.AppendSetMap(SetMapB->Duplicate(Output.GetAllocator()));
+							Output.AppendSetMap(MapB->Duplicate(Output.GetAllocator()));
 						}
 					}
 					else
 					{
-						if (SetMapB == nullptr)
+						if (MapB == nullptr)
 						{
 							// This set map doesn't have a matching pair, duplicate it
-							Output.AppendSetMap(SetMapA->Duplicate(Output.GetAllocator()));
+							Output.AppendSetMap(MapA->Duplicate(Output.GetAllocator()));
 						}
 						else
 						{
 							if (Transformer == nullptr)
 							{
 								// No transformer was found for this value type, retain the one with the largest weight
-								FAttributeSetMap* SetMapCopy;
+								FBoundValueMap* MapCopy;
 								if (DefaultWeightB >= 0.5f)
 								{
-									SetMapCopy = SetMapB->Duplicate(Output.GetAllocator());
+									MapCopy = MapB->Duplicate(Output.GetAllocator());
 								}
 								else
 								{
-									SetMapCopy = SetMapA->Duplicate(Output.GetAllocator());
+									MapCopy = MapA->Duplicate(Output.GetAllocator());
 								}
 
-								Output.AppendSetMap(SetMapCopy);
+								Output.AppendSetMap(MapCopy);
 							}
 							else
 							{
 								// Transform our matching pair
-								FAttributeSetMap* OutputSetMap = nullptr;
+								FBoundValueMap* OutputMap = nullptr;
 
 								// We cannot join our weights because their value type differs
-								const FAttributeMappingKey MappingKey = SetMapA->GetMappingKey().To(FloatValueType);
-								const FAttributeSetMap* PerValueWeightMapB = PerValueWeightsB.FindSetMap(MappingKey);
+								const FAttributeMappingKey MappingKey = MapA->GetMappingKey().To(FloatValueType);
+								const FBoundValueMap* PerValueWeightMapB = PerValueWeightsB.FindSetMap(MappingKey);
 
-								FTransformSetMapFunc TypedTransformer = reinterpret_cast<FTransformSetMapFunc>(Transformer);
-								(*TypedTransformer)(SetMapA, SetMapB, PerValueWeightMapB, DefaultWeightB, OutputSetMap, Output.GetAllocator());
+								FTransformBoundValueMapFunc TypedTransformer = reinterpret_cast<FTransformBoundValueMapFunc>(Transformer);
+								(*TypedTransformer)(MapA, MapB, PerValueWeightMapB, DefaultWeightB, OutputMap, Output.GetAllocator());
 
-								checkf(OutputSetMap != nullptr, TEXT("Transformer function is responsible for allocating an output buffer"));
-								checkf(OutputSetMap->GetAllocator() == Output.GetAllocator(), TEXT("Transformer function is responsible for allocating using the output allocator"));
-								Output.AppendSetMap(OutputSetMap);
+								checkf(OutputMap != nullptr, TEXT("Transformer function is responsible for allocating an output buffer"));
+								checkf(OutputMap->GetAllocator() == Output.GetAllocator(), TEXT("Transformer function is responsible for allocating using the output allocator"));
+								Output.AppendSetMap(OutputMap);
 							}
 						}
 					}
 				},
-				TransformerList->CreateSetMapTransformerIterator(),
+				TransformerList->CreateBoundValueMapTransformerIterator(),
 				InputA.CreateConstSetMapIterator(),
 				InputB.CreateConstSetMapIterator());
 
 			// Name/value maps
 			OuterJoinBy(
-				FAttributeValueTransformerListWithAttributeCollectionMapJoinOp(),
-				[DefaultWeightB, &Output](FRawTransformerFunc Transformer, const FAttributeMap* MapA, const FAttributeMap* MapB)
+				FValueTransformerListWithUnboundMapCollectionJoinOp(),
+				[DefaultWeightB, &Output](FRawTransformerFunc Transformer, const FUnboundValueMap* MapA, const FUnboundValueMap* MapB)
 				{
 					if (MapA == nullptr)
 					{
@@ -377,7 +377,7 @@ namespace UE::UAF::Transformers
 							if (Transformer == nullptr)
 							{
 								// No transformer was found for this value type, retain the one with the largest weight
-								FAttributeMap* MapCopy;
+								FUnboundValueMap* MapCopy;
 								if (DefaultWeightB >= 0.5f)
 								{
 									MapCopy = MapB->Duplicate(Output.GetAllocator());
@@ -392,9 +392,9 @@ namespace UE::UAF::Transformers
 							else
 							{
 								// Transform our matching pair in place
-								FAttributeMap* OutputMap = nullptr;
+								FUnboundValueMap* OutputMap = nullptr;
 
-								FTransformMapFunc TypedTransformer = reinterpret_cast<FTransformMapFunc>(Transformer);
+								FTransformUnboundValueMapFunc TypedTransformer = reinterpret_cast<FTransformUnboundValueMapFunc>(Transformer);
 								(*TypedTransformer)(MapA, MapB, DefaultWeightB, OutputMap, Output.GetAllocator());
 
 								checkf(OutputMap != nullptr, TEXT("Transformer function is responsible for allocating an output buffer"));
@@ -404,7 +404,7 @@ namespace UE::UAF::Transformers
 						}
 					}
 				},
-				TransformerList->CreateMapTransformerIterator(),
+				TransformerList->CreateUnboundValueMapTransformerIterator(),
 				InputA.CreateConstMapIterator(),
 				InputB.CreateConstMapIterator());
 		}
@@ -412,33 +412,33 @@ namespace UE::UAF::Transformers
 
 	namespace Private
 	{
-		void FInterpolate_BoneTransformAttribute::TransformSetMap(const FAttributeSetMap* SetMapA, const FAttributeSetMap* SetMapB, const FAttributeSetMap* PerValueWeightMapB, float DefaultWeightB, FAttributeSetMap*& OutputSetMap, FReallocFun OutputAllocator)
+		void FInterpolate_BoneTransformAttribute::TransformBoundValueMap(const FBoundValueMap* MapA, const FBoundValueMap* MapB, const FBoundValueMap* PerValueWeightMapB, float DefaultWeightB, FBoundValueMap*& OutputMap, FReallocFun OutputAllocator)
 		{
-			check(SetMapA->GetTypedSet() == SetMapB->GetTypedSet());
-			check(SetMapA->Num() == SetMapB->Num());
+			check(MapA->GetTypedSet() == MapB->GetTypedSet());
+			check(MapA->Num() == MapB->Num());
 
-			if (OutputSetMap == nullptr)
+			if (OutputMap == nullptr)
 			{
-				const FAttributeSetMap::FConstructArgs Args(SetMapA->GetTypedSet(), SetMapA->GetValueType(), OutputAllocator);
-				OutputSetMap = MakeAttributeSetMap<FBoneTransformAnimationAttribute>(Args);
+				const FBoundValueMap::FConstructArgs Args(MapA->GetTypedSet(), MapA->GetValueType(), OutputAllocator);
+				OutputMap = MakeBoundValueMap<FBoneTransformAnimationAttribute>(Args);
 			}
 
-			check(SetMapA->GetTypedSet() == OutputSetMap->GetTypedSet());
-			check(SetMapA->Num() == OutputSetMap->Num());
+			check(MapA->GetTypedSet() == OutputMap->GetTypedSet());
+			check(MapA->Num() == OutputMap->Num());
 
-			const TAttributeSetMap<FBoneTransformAnimationAttribute>* BoneTransformsA = CastChecked<FBoneTransformAnimationAttribute>(SetMapA);
-			const TAttributeSetMap<FBoneTransformAnimationAttribute>* BoneTransformsB = CastChecked<FBoneTransformAnimationAttribute>(SetMapB);
-			TAttributeSetMap<FBoneTransformAnimationAttribute>* OutputBoneTransforms = CastChecked<FBoneTransformAnimationAttribute>(OutputSetMap);
+			const TBoundValueMap<FBoneTransformAnimationAttribute>* BoneTransformsA = CastChecked<FBoneTransformAnimationAttribute>(MapA);
+			const TBoundValueMap<FBoneTransformAnimationAttribute>* BoneTransformsB = CastChecked<FBoneTransformAnimationAttribute>(MapB);
+			TBoundValueMap<FBoneTransformAnimationAttribute>* OutputBoneTransforms = CastChecked<FBoneTransformAnimationAttribute>(OutputMap);
 
 			const FBoneTransformAnimationAttribute* TransformPtrA = BoneTransformsA->GetData();
 			const FBoneTransformAnimationAttribute* TransformPtrB = BoneTransformsB->GetData();
 			FBoneTransformAnimationAttribute* OutputTransformPtr = OutputBoneTransforms->GetData();
 			const FBoneTransformAnimationAttribute* TransformPtrAEnd = TransformPtrA + BoneTransformsA->Num();
 
-			if (const TAttributeSetMap<FFloatAnimationAttribute>* PerValueWeightsB = Cast<FFloatAnimationAttribute>(PerValueWeightMapB))
+			if (const TBoundValueMap<FFloatAnimationAttribute>* PerValueWeightsB = Cast<FFloatAnimationAttribute>(PerValueWeightMapB))
 			{
-				check(SetMapA->GetTypedSet() == PerValueWeightsB->GetTypedSet());
-				check(SetMapA->Num() == PerValueWeightsB->Num());
+				check(MapA->GetTypedSet() == PerValueWeightsB->GetTypedSet());
+				check(MapA->Num() == PerValueWeightsB->Num());
 
 				const FFloatAnimationAttribute* PerValueWeightsBPtr = PerValueWeightsB->GetData();
 
@@ -465,16 +465,16 @@ namespace UE::UAF::Transformers
 			}
 		}
 
-		void FInterpolate_BoneTransformAttribute::TransformMap(const FAttributeMap* MapA, const FAttributeMap* MapB, float WeightB, FAttributeMap*& OutputMap, FReallocFun OutputAllocator)
+		void FInterpolate_BoneTransformAttribute::TransformUnboundValueMap(const FUnboundValueMap* MapA, const FUnboundValueMap* MapB, float WeightB, FUnboundValueMap*& OutputMap, FReallocFun OutputAllocator)
 		{
 			check(OutputMap == nullptr);
 
-			const TAttributeMap<FBoneTransformAnimationAttribute>* BoneTransformsA = CastChecked<FBoneTransformAnimationAttribute>(MapA);
-			const TAttributeMap<FBoneTransformAnimationAttribute>* BoneTransformsB = CastChecked<FBoneTransformAnimationAttribute>(MapB);
-			TAttributeMap<FBoneTransformAnimationAttribute>* OutputBoneTransforms = MakeAttributeMap<FBoneTransformAnimationAttribute>(OutputAllocator);
+			const TUnboundValueMap<FBoneTransformAnimationAttribute>* BoneTransformsA = CastChecked<FBoneTransformAnimationAttribute>(MapA);
+			const TUnboundValueMap<FBoneTransformAnimationAttribute>* BoneTransformsB = CastChecked<FBoneTransformAnimationAttribute>(MapB);
+			TUnboundValueMap<FBoneTransformAnimationAttribute>* OutputBoneTransforms = MakeUnboundValueMap<FBoneTransformAnimationAttribute>(OutputAllocator);
 
 			OuterJoinBy(
-				FAttributeMapJoinOp(),
+				FUnboundValueMapJoinOp(),
 				[WeightB, OutputBoneTransforms](FName Name, const FBoneTransformAnimationAttribute* A, const FBoneTransformAnimationAttribute* B)
 				{
 					// Missing values are set to the identity
@@ -492,33 +492,33 @@ namespace UE::UAF::Transformers
 			OutputMap = OutputBoneTransforms;
 		}
 
-		void FInterpolate_FloatAttribute::TransformSetMap(const FAttributeSetMap* SetMapA, const FAttributeSetMap* SetMapB, const FAttributeSetMap* PerValueWeightMapB, float DefaultWeightB, FAttributeSetMap*& OutputSetMap, FReallocFun OutputAllocator)
+		void FInterpolate_FloatAttribute::TransformBoundValueMap(const FBoundValueMap* MapA, const FBoundValueMap* MapB, const FBoundValueMap* PerValueWeightMapB, float DefaultWeightB, FBoundValueMap*& OutputMap, FReallocFun OutputAllocator)
 		{
-			check(SetMapA->GetTypedSet() == SetMapB->GetTypedSet());
-			check(SetMapA->Num() == SetMapB->Num());
+			check(MapA->GetTypedSet() == MapB->GetTypedSet());
+			check(MapA->Num() == MapB->Num());
 
-			if (OutputSetMap == nullptr)
+			if (OutputMap == nullptr)
 			{
-				const FAttributeSetMap::FConstructArgs Args(SetMapA->GetTypedSet(), SetMapA->GetValueType(), OutputAllocator);
-				OutputSetMap = MakeAttributeSetMap<FFloatAnimationAttribute>(Args);
+				const FBoundValueMap::FConstructArgs Args(MapA->GetTypedSet(), MapA->GetValueType(), OutputAllocator);
+				OutputMap = MakeBoundValueMap<FFloatAnimationAttribute>(Args);
 			}
 
-			check(SetMapA->GetTypedSet() == OutputSetMap->GetTypedSet());
-			check(SetMapA->Num() == OutputSetMap->Num());
+			check(MapA->GetTypedSet() == OutputMap->GetTypedSet());
+			check(MapA->Num() == OutputMap->Num());
 
-			const TAttributeSetMap<FFloatAnimationAttribute>* FloatsA = CastChecked<FFloatAnimationAttribute>(SetMapA);
-			const TAttributeSetMap<FFloatAnimationAttribute>* FloatsB = CastChecked<FFloatAnimationAttribute>(SetMapB);
-			TAttributeSetMap<FFloatAnimationAttribute>* OutputFloats = CastChecked<FFloatAnimationAttribute>(OutputSetMap);
+			const TBoundValueMap<FFloatAnimationAttribute>* FloatsA = CastChecked<FFloatAnimationAttribute>(MapA);
+			const TBoundValueMap<FFloatAnimationAttribute>* FloatsB = CastChecked<FFloatAnimationAttribute>(MapB);
+			TBoundValueMap<FFloatAnimationAttribute>* OutputFloats = CastChecked<FFloatAnimationAttribute>(OutputMap);
 
 			const FFloatAnimationAttribute* FloatsPtrA = FloatsA->GetData();
 			const FFloatAnimationAttribute* FloatsPtrB = FloatsB->GetData();
 			FFloatAnimationAttribute* OutputFloatsPtr = OutputFloats->GetData();
 			const FFloatAnimationAttribute* FloatsPtrAEnd = FloatsPtrA + FloatsA->Num();
 
-			if (const TAttributeSetMap<FFloatAnimationAttribute>* PerValueWeightsB = Cast<FFloatAnimationAttribute>(PerValueWeightMapB))
+			if (const TBoundValueMap<FFloatAnimationAttribute>* PerValueWeightsB = Cast<FFloatAnimationAttribute>(PerValueWeightMapB))
 			{
-				check(SetMapA->GetTypedSet() == PerValueWeightsB->GetTypedSet());
-				check(SetMapA->Num() == PerValueWeightsB->Num());
+				check(MapA->GetTypedSet() == PerValueWeightsB->GetTypedSet());
+				check(MapA->Num() == PerValueWeightsB->Num());
 
 				const FFloatAnimationAttribute* PerValueWeightsBPtr = PerValueWeightsB->GetData();
 
@@ -545,16 +545,16 @@ namespace UE::UAF::Transformers
 			}
 		}
 
-		void FInterpolate_FloatAttribute::TransformMap(const FAttributeMap* MapA, const FAttributeMap* MapB, float WeightB, FAttributeMap*& OutputMap, FReallocFun OutputAllocator)
+		void FInterpolate_FloatAttribute::TransformUnboundValueMap(const FUnboundValueMap* MapA, const FUnboundValueMap* MapB, float WeightB, FUnboundValueMap*& OutputMap, FReallocFun OutputAllocator)
 		{
 			check(OutputMap == nullptr);
 
-			const TAttributeMap<FFloatAnimationAttribute>* FloatsA = CastChecked<FFloatAnimationAttribute>(MapA);
-			const TAttributeMap<FFloatAnimationAttribute>* FloatsB = CastChecked<FFloatAnimationAttribute>(MapB);
-			TAttributeMap<FFloatAnimationAttribute>* OutputFloats = MakeAttributeMap<FFloatAnimationAttribute>(OutputAllocator);
+			const TUnboundValueMap<FFloatAnimationAttribute>* FloatsA = CastChecked<FFloatAnimationAttribute>(MapA);
+			const TUnboundValueMap<FFloatAnimationAttribute>* FloatsB = CastChecked<FFloatAnimationAttribute>(MapB);
+			TUnboundValueMap<FFloatAnimationAttribute>* OutputFloats = MakeUnboundValueMap<FFloatAnimationAttribute>(OutputAllocator);
 
 			OuterJoinBy(
-				FAttributeMapJoinOp(),
+				FUnboundValueMapJoinOp(),
 				[WeightB, OutputFloats](FName Name, const FFloatAnimationAttribute* A, const FFloatAnimationAttribute* B)
 				{
 					// Missing values are set to 0.0
diff --git a/Engine/Plugins/Experimental/UAF/UAF/Source/UAF/Private/UAF/ValueRuntime/Transformers/Layer.cpp b/Engine/Plugins/Experimental/UAF/UAF/Source/UAF/Private/UAF/ValueRuntime/Transformers/Layer.cpp
--- a/Engine/Plugins/Experimental/UAF/UAF/Source/UAF/Private/UAF/ValueRuntime/Transformers/Layer.cpp
+++ b/Engine/Plugins/Experimental/UAF/UAF/Source/UAF/Private/UAF/ValueRuntime/Transformers/Layer.cpp
@@ -9,7 +9,7 @@ namespace UE::UAF::Transformers
 	FName FLayer::TransformerName = TEXT("UAF::Transformers::Layer");
 
 	void FLayer::Apply(
-		const FAttributeValueTransformerMapPtr& TransformerMap,
+		const FValueTransformerMapPtr& TransformerMap,
 		const FAttributeCollection& Base,
 		const FAttributeCollection& Layer,
 		float LayerWeight,
@@ -28,14 +28,14 @@ namespace UE::UAF::Transformers
 	}
 
 	void FLayer::Apply(
-		const FAttributeValueTransformerMapPtr& TransformerMap,
+		const FValueTransformerMapPtr& TransformerMap,
 		const FAttributeCollection& Base,
 		const FAttributeCollection& Layer,
 		const FAttributeCollection& PerValueLayerWeights,
 		float DefaultLayerWeight,
 		FAttributeCollection& Output)
 	{
-		const FAttributeValueTransformerList* TransformerList = TransformerMap->Find(TransformerName);
+		const FValueTransformerList* TransformerList = TransformerMap->Find(TransformerName);
 		if (TransformerList == nullptr)
 		{
 			// No transformer implementations have been registered
@@ -72,10 +72,10 @@ namespace UE::UAF::Transformers
 			// Our output collection matches one of our inputs, we'll modify it in place
 			FReallocFun OutputAllocator = Output.GetAllocator();
 
-			// Set maps
+			// Bound value maps
 			OuterJoinBy(
-				FAttributeValueTransformerListWithAttributeCollectionSetMapJoinOp(),
-				[&PerValueLayerWeights, LayerWeight, bIsOutputBase, FloatValueType, OutputAllocator](FRawTransformerFunc Transformer, FAttributeSetMap* BaseMap, FAttributeSetMap* LayerMap)
+				FValueTransformerListWithBoundMapCollectionJoinOp(),
+				[&PerValueLayerWeights, LayerWeight, bIsOutputBase, FloatValueType, OutputAllocator](FRawTransformerFunc Transformer, FBoundValueMap* BaseMap, FBoundValueMap* LayerMap)
 				{
 					if (BaseMap == nullptr)
 					{
@@ -129,27 +129,27 @@ namespace UE::UAF::Transformers
 							else
 							{
 								// Transform our matching pair in place
-								FAttributeSetMap* OutputSetMap = bIsOutputBase ? BaseMap : LayerMap;
+								FBoundValueMap* OutputMap = bIsOutputBase ? BaseMap : LayerMap;
 
 								// We cannot join our weights because their value type differs
 								const FAttributeMappingKey MappingKey = BaseMap->GetMappingKey().To(FloatValueType);
-								const FAttributeSetMap* PerValueWeightMap = PerValueLayerWeights.FindSetMap(MappingKey);
+								const FBoundValueMap* PerValueWeightMap = PerValueLayerWeights.FindSetMap(MappingKey);
 
-								FTransformSetMapFunc TypedTransformer = reinterpret_cast<FTransformSetMapFunc>(Transformer);
-								(*TypedTransformer)(BaseMap, LayerMap, PerValueWeightMap, LayerWeight, OutputSetMap, nullptr);
+								FTransformBoundValueMapFunc TypedTransformer = reinterpret_cast<FTransformBoundValueMapFunc>(Transformer);
+								(*TypedTransformer)(BaseMap, LayerMap, PerValueWeightMap, LayerWeight, OutputMap, nullptr);
 							}
 						}
 					}
 				},
-				TransformerList->CreateSetMapTransformerIterator(),
+				TransformerList->CreateBoundValueMapTransformerIterator(),
 				// Remove the const qualifier since we'll mutate one in place as our output
 				const_cast<FAttributeCollection&>(Base).CreateSetMapIterator(),
 				const_cast<FAttributeCollection&>(Layer).CreateSetMapIterator());
 
 			// Name/value maps
 			OuterJoinBy(
-				FAttributeValueTransformerListWithAttributeCollectionMapJoinOp(),
-				[LayerWeight, bIsOutputBase, OutputAllocator](FRawTransformerFunc Transformer, FAttributeMap* BaseMap, FAttributeMap* LayerMap)
+				FValueTransformerListWithUnboundMapCollectionJoinOp(),
+				[LayerWeight, bIsOutputBase, OutputAllocator](FRawTransformerFunc Transformer, FUnboundValueMap* BaseMap, FUnboundValueMap* LayerMap)
 				{
 					if (BaseMap == nullptr)
 					{
@@ -203,9 +203,9 @@ namespace UE::UAF::Transformers
 							else
 							{
 								// Transform our matching pair in a new map
-								FAttributeMap* OutputMap = nullptr;
+								FUnboundValueMap* OutputMap = nullptr;
 
-								FTransformMapFunc TypedTransformer = reinterpret_cast<FTransformMapFunc>(Transformer);
+								FTransformUnboundValueMapFunc TypedTransformer = reinterpret_cast<FTransformUnboundValueMapFunc>(Transformer);
 								(*TypedTransformer)(BaseMap, LayerMap, LayerWeight, OutputMap, OutputAllocator);
 
 								checkf(OutputMap != nullptr, TEXT("Transformer function is responsible for allocating an output buffer"));
@@ -215,12 +215,12 @@ namespace UE::UAF::Transformers
 								OutputMap->MoveTo(bIsOutputBase ? *BaseMap : *LayerMap);
 
 								// Destroy our temporary output map
-								ReleaseAttributeMap(OutputMap);
+								ReleaseUnboundValueMap(OutputMap);
 							}
 						}
 					}
 				},
-				TransformerList->CreateMapTransformerIterator(),
+				TransformerList->CreateUnboundValueMapTransformerIterator(),
 				// Remove the const qualifier since we'll mutate one in place as our output
 				const_cast<FAttributeCollection&>(Base).CreateMapIterator(),
 				const_cast<FAttributeCollection&>(Layer).CreateMapIterator());
@@ -231,10 +231,10 @@ namespace UE::UAF::Transformers
 			Output.Reset(Base.GetNamedSet());
 			Output.SetValueSpace(Base.GetValueSpace());
 
-			// Set maps
+			// Bound value maps
 			OuterJoinBy(
-				FAttributeValueTransformerListWithAttributeCollectionSetMapJoinOp(),
-				[&PerValueLayerWeights, LayerWeight, FloatValueType, &Output](FRawTransformerFunc Transformer, const FAttributeSetMap* BaseMap, const FAttributeSetMap* LayerMap)
+				FValueTransformerListWithBoundMapCollectionJoinOp(),
+				[&PerValueLayerWeights, LayerWeight, FloatValueType, &Output](FRawTransformerFunc Transformer, const FBoundValueMap* BaseMap, const FBoundValueMap* LayerMap)
 				{
 					if (BaseMap == nullptr)
 					{
@@ -260,45 +260,45 @@ namespace UE::UAF::Transformers
 							if (Transformer == nullptr)
 							{
 								// No transformer was found for this value type, retain the one with the largest weight
-								FAttributeSetMap* SetMapCopy;
+								FBoundValueMap* MapCopy;
 								if (LayerWeight >= 0.5f)
 								{
-									SetMapCopy = LayerMap->Duplicate(Output.GetAllocator());
+									MapCopy = LayerMap->Duplicate(Output.GetAllocator());
 								}
 								else
 								{
-									SetMapCopy = BaseMap->Duplicate(Output.GetAllocator());
+									MapCopy = BaseMap->Duplicate(Output.GetAllocator());
 								}
 
-								Output.AppendSetMap(SetMapCopy);
+								Output.AppendSetMap(MapCopy);
 							}
 							else
 							{
 								// Transform our matching pair
-								FAttributeSetMap* OutputSetMap = nullptr;
+								FBoundValueMap* OutputMap = nullptr;
 
 								// We cannot join our weights because their value type differs
 								const FAttributeMappingKey MappingKey = BaseMap->GetMappingKey().To(FloatValueType);
-								const FAttributeSetMap* PerValueWeightMap = PerValueLayerWeights.FindSetMap(MappingKey);
+								const FBoundValueMap* PerValueWeightMap = PerValueLayerWeights.FindSetMap(MappingKey);
 
-								FTransformSetMapFunc TypedTransformer = reinterpret_cast<FTransformSetMapFunc>(Transformer);
-								(*TypedTransformer)(BaseMap, LayerMap, PerValueWeightMap, LayerWeight, OutputSetMap, Output.GetAllocator());
+								FTransformBoundValueMapFunc TypedTransformer = reinterpret_cast<FTransformBoundValueMapFunc>(Transformer);
+								(*TypedTransformer)(BaseMap, LayerMap, PerValueWeightMap, LayerWeight, OutputMap, Output.GetAllocator());
 
-								checkf(OutputSetMap != nullptr, TEXT("Transformer function is responsible for allocating an output buffer"));
-								checkf(OutputSetMap->GetAllocator() == Output.GetAllocator(), TEXT("Transformer function is responsible for allocating using the output allocator"));
-								Output.AppendSetMap(OutputSetMap);
+								checkf(OutputMap != nullptr, TEXT("Transformer function is responsible for allocating an output buffer"));
+								checkf(OutputMap->GetAllocator() == Output.GetAllocator(), TEXT("Transformer function is responsible for allocating using the output allocator"));
+								Output.AppendSetMap(OutputMap);
 							}
 						}
 					}
 				},
-				TransformerList->CreateSetMapTransformerIterator(),
+				TransformerList->CreateBoundValueMapTransformerIterator(),
 				Base.CreateConstSetMapIterator(),
 				LayerPtr->CreateConstSetMapIterator());
 
 			// Name/value maps
 			OuterJoinBy(
-				FAttributeValueTransformerListWithAttributeCollectionMapJoinOp(),
-				[LayerWeight, &Output](FRawTransformerFunc Transformer, const FAttributeMap* BaseMap, const FAttributeMap* LayerMap)
+				FValueTransformerListWithUnboundMapCollectionJoinOp(),
+				[LayerWeight, &Output](FRawTransformerFunc Transformer, const FUnboundValueMap* BaseMap, const FUnboundValueMap* LayerMap)
 				{
 					if (BaseMap == nullptr)
 					{
@@ -324,7 +324,7 @@ namespace UE::UAF::Transformers
 							if (Transformer == nullptr)
 							{
 								// No transformer was found for this value type, retain the one with the largest weight
-								FAttributeMap* MapCopy;
+								FUnboundValueMap* MapCopy;
 								if (LayerWeight >= 0.5f)
 								{
 									MapCopy = LayerMap->Duplicate(Output.GetAllocator());
@@ -339,9 +339,9 @@ namespace UE::UAF::Transformers
 							else
 							{
 								// Transform our matching pair in place
-								FAttributeMap* OutputMap = nullptr;
+								FUnboundValueMap* OutputMap = nullptr;
 
-								FTransformMapFunc TypedTransformer = reinterpret_cast<FTransformMapFunc>(Transformer);
+								FTransformUnboundValueMapFunc TypedTransformer = reinterpret_cast<FTransformUnboundValueMapFunc>(Transformer);
 								(*TypedTransformer)(BaseMap, LayerMap, LayerWeight, OutputMap, Output.GetAllocator());
 
 								checkf(OutputMap != nullptr, TEXT("Transformer function is responsible for allocating an output buffer"));
@@ -351,7 +351,7 @@ namespace UE::UAF::Transformers
 						}
 					}
 				},
-				TransformerList->CreateMapTransformerIterator(),
+				TransformerList->CreateUnboundValueMapTransformerIterator(),
 				Base.CreateConstMapIterator(),
 				LayerPtr->CreateConstMapIterator());
 		}
@@ -359,17 +359,17 @@ namespace UE::UAF::Transformers
 
 	namespace Private
 	{
-		void FLayer_BoneTransformAttribute::TransformSetMap(const FAttributeSetMap* Base, const FAttributeSetMap* Layer, const FAttributeSetMap* PerValueWeights, float DefaultLayerWeight, FAttributeSetMap*& OutputSetMap, FReallocFun OutputAllocator)
+		void FLayer_BoneTransformAttribute::TransformBoundValueMap(const FBoundValueMap* Base, const FBoundValueMap* Layer, const FBoundValueMap* PerValueWeights, float DefaultLayerWeight, FBoundValueMap*& OutputMap, FReallocFun OutputAllocator)
 		{
-			if (OutputSetMap == nullptr)
+			if (OutputMap == nullptr)
 			{
-				const FAttributeSetMap::FConstructArgs Args(Base->GetTypedSet(), Base->GetValueType(), OutputAllocator);
-				OutputSetMap = MakeAttributeSetMap<FBoneTransformAnimationAttribute>(Args);
+				const FBoundValueMap::FConstructArgs Args(Base->GetTypedSet(), Base->GetValueType(), OutputAllocator);
+				OutputMap = MakeBoundValueMap<FBoneTransformAnimationAttribute>(Args);
 			}
 
 			// Our output should have the same shape as our base
-			check(Base->GetTypedSet() == OutputSetMap->GetTypedSet());
-			check(Base->Num() == OutputSetMap->Num());
+			check(Base->GetTypedSet() == OutputMap->GetTypedSet());
+			check(Base->Num() == OutputMap->Num());
 
 			// Our per-value weights should have the same shape as our base
 			check(PerValueWeights == nullptr || Base->GetTypedSet() == PerValueWeights->GetTypedSet());
@@ -379,10 +379,10 @@ namespace UE::UAF::Transformers
 			const FAttributeTypedSetPtr& LayerTypedSet = Layer->GetTypedSet();
 			const FAttributeTypedSet* PerValueWeightTypedSet = PerValueWeights != nullptr ? PerValueWeights->GetTypedSet().GetReference() : nullptr;
 
-			const TAttributeSetMap<FBoneTransformAnimationAttribute>* BaseBoneTransforms = CastChecked<FBoneTransformAnimationAttribute>(Base);
-			const TAttributeSetMap<FBoneTransformAnimationAttribute>* LayerBoneTransforms = CastChecked<FBoneTransformAnimationAttribute>(Layer);
-			const TAttributeSetMap<FFloatAnimationAttribute>* PerValueWeightFloats = PerValueWeights != nullptr ? CastChecked<FFloatAnimationAttribute>(PerValueWeights) : nullptr;
-			TAttributeSetMap<FBoneTransformAnimationAttribute>* OutputBoneTransforms = CastChecked<FBoneTransformAnimationAttribute>(OutputSetMap);
+			const TBoundValueMap<FBoneTransformAnimationAttribute>* BaseBoneTransforms = CastChecked<FBoneTransformAnimationAttribute>(Base);
+			const TBoundValueMap<FBoneTransformAnimationAttribute>* LayerBoneTransforms = CastChecked<FBoneTransformAnimationAttribute>(Layer);
+			const TBoundValueMap<FFloatAnimationAttribute>* PerValueWeightFloats = PerValueWeights != nullptr ? CastChecked<FFloatAnimationAttribute>(PerValueWeights) : nullptr;
+			TBoundValueMap<FBoneTransformAnimationAttribute>* OutputBoneTransforms = CastChecked<FBoneTransformAnimationAttribute>(OutputMap);
 
 			// Use an outer join over our base, layer, and weights
 			// If a value is in the base but not the layer, use the base value
@@ -421,17 +421,17 @@ namespace UE::UAF::Transformers
 			}
 		}
 
-		void FLayer_BoneTransformAttribute::TransformMap(const FAttributeMap* Base, const FAttributeMap* Layer, float LayerWeight, FAttributeMap*& OutputMap, FReallocFun OutputAllocator)
+		void FLayer_BoneTransformAttribute::TransformUnboundValueMap(const FUnboundValueMap* Base, const FUnboundValueMap* Layer, float LayerWeight, FUnboundValueMap*& OutputMap, FReallocFun OutputAllocator)
 		{
 			check(OutputMap == nullptr);
 
-			const TAttributeMap<FBoneTransformAnimationAttribute>* BaseBoneTransforms = CastChecked<FBoneTransformAnimationAttribute>(Base);
-			const TAttributeMap<FBoneTransformAnimationAttribute>* LayerBoneTransforms = CastChecked<FBoneTransformAnimationAttribute>(Layer);
+			const TUnboundValueMap<FBoneTransformAnimationAttribute>* BaseBoneTransforms = CastChecked<FBoneTransformAnimationAttribute>(Base);
+			const TUnboundValueMap<FBoneTransformAnimationAttribute>* LayerBoneTransforms = CastChecked<FBoneTransformAnimationAttribute>(Layer);
 
-			TAttributeMap<FBoneTransformAnimationAttribute>* OutputBoneTransforms = MakeAttributeMap<FBoneTransformAnimationAttribute>(OutputAllocator);
+			TUnboundValueMap<FBoneTransformAnimationAttribute>* OutputBoneTransforms = MakeUnboundValueMap<FBoneTransformAnimationAttribute>(OutputAllocator);
 
 			OuterJoinBy(
-				FAttributeMapJoinOp(),
+				FUnboundValueMapJoinOp(),
 				[LayerWeight, OutputBoneTransforms](FName Name, const FBoneTransformAnimationAttribute* BaseValue, const FBoneTransformAnimationAttribute* LayerValue)
 				{
 					// Mismatched pairs are ignored
@@ -449,17 +449,17 @@ namespace UE::UAF::Transformers
 			OutputMap = OutputBoneTransforms;
 		}
 
-		void FLayer_FloatAttribute::TransformSetMap(const FAttributeSetMap* Base, const FAttributeSetMap* Layer, const FAttributeSetMap* PerValueWeights, float DefaultLayerWeight, FAttributeSetMap*& OutputSetMap, FReallocFun OutputAllocator)
+		void FLayer_FloatAttribute::TransformBoundValueMap(const FBoundValueMap* Base, const FBoundValueMap* Layer, const FBoundValueMap* PerValueWeights, float DefaultLayerWeight, FBoundValueMap*& OutputMap, FReallocFun OutputAllocator)
 		{
-			if (OutputSetMap == nullptr)
+			if (OutputMap == nullptr)
 			{
-				const FAttributeSetMap::FConstructArgs Args(Base->GetTypedSet(), Base->GetValueType(), OutputAllocator);
-				OutputSetMap = MakeAttributeSetMap<FFloatAnimationAttribute>(Args);
+				const FBoundValueMap::FConstructArgs Args(Base->GetTypedSet(), Base->GetValueType(), OutputAllocator);
+				OutputMap = MakeBoundValueMap<FFloatAnimationAttribute>(Args);
 			}
 
 			// Our output should have the same shape as our base
-			check(Base->GetTypedSet() == OutputSetMap->GetTypedSet());
-			check(Base->Num() == OutputSetMap->Num());
+			check(Base->GetTypedSet() == OutputMap->GetTypedSet());
+			check(Base->Num() == OutputMap->Num());
 
 			// Our per-value weights should have the same shape as our base
 			check(PerValueWeights == nullptr || Base->GetTypedSet() == PerValueWeights->GetTypedSet());
@@ -469,10 +469,10 @@ namespace UE::UAF::Transformers
 			const FAttributeTypedSetPtr& LayerTypedSet = Layer->GetTypedSet();
 			const FAttributeTypedSet* PerValueWeightTypedSet = PerValueWeights != nullptr ? PerValueWeights->GetTypedSet().GetReference() : nullptr;
 
-			const TAttributeSetMap<FFloatAnimationAttribute>* BaseFloats = CastChecked<FFloatAnimationAttribute>(Base);
-			const TAttributeSetMap<FFloatAnimationAttribute>* LayerFloats = CastChecked<FFloatAnimationAttribute>(Layer);
-			const TAttributeSetMap<FFloatAnimationAttribute>* PerValueWeightFloats = PerValueWeights != nullptr ? CastChecked<FFloatAnimationAttribute>(PerValueWeights) : nullptr;
-			TAttributeSetMap<FFloatAnimationAttribute>* OutputFloats = CastChecked<FFloatAnimationAttribute>(OutputSetMap);
+			const TBoundValueMap<FFloatAnimationAttribute>* BaseFloats = CastChecked<FFloatAnimationAttribute>(Base);
+			const TBoundValueMap<FFloatAnimationAttribute>* LayerFloats = CastChecked<FFloatAnimationAttribute>(Layer);
+			const TBoundValueMap<FFloatAnimationAttribute>* PerValueWeightFloats = PerValueWeights != nullptr ? CastChecked<FFloatAnimationAttribute>(PerValueWeights) : nullptr;
+			TBoundValueMap<FFloatAnimationAttribute>* OutputFloats = CastChecked<FFloatAnimationAttribute>(OutputMap);
 
 			// Use an outer join over our base, layer, and weights
 			// If a value is in the base but not the layer, use the base value
@@ -512,17 +512,17 @@ namespace UE::UAF::Transformers
 			}
 		}
 
-		void FLayer_FloatAttribute::TransformMap(const FAttributeMap* Base, const FAttributeMap* Layer, float LayerWeight, FAttributeMap*& OutputMap, FReallocFun OutputAllocator)
+		void FLayer_FloatAttribute::TransformUnboundValueMap(const FUnboundValueMap* Base, const FUnboundValueMap* Layer, float LayerWeight, FUnboundValueMap*& OutputMap, FReallocFun OutputAllocator)
 		{
 			check(OutputMap == nullptr);
 
-			const TAttributeMap<FFloatAnimationAttribute>* BaseFloats = CastChecked<FFloatAnimationAttribute>(Base);
-			const TAttributeMap<FFloatAnimationAttribute>* LayerFloats = CastChecked<FFloatAnimationAttribute>(Layer);
+			const TUnboundValueMap<FFloatAnimationAttribute>* BaseFloats = CastChecked<FFloatAnimationAttribute>(Base);
+			const TUnboundValueMap<FFloatAnimationAttribute>* LayerFloats = CastChecked<FFloatAnimationAttribute>(Layer);
 
-			TAttributeMap<FFloatAnimationAttribute>* OutputFloats = MakeAttributeMap<FFloatAnimationAttribute>(OutputAllocator);
+			TUnboundValueMap<FFloatAnimationAttribute>* OutputFloats = MakeUnboundValueMap<FFloatAnimationAttribute>(OutputAllocator);
 
 			OuterJoinBy(
-				FAttributeMapJoinOp(),
+				FUnboundValueMapJoinOp(),
 				[LayerWeight, OutputFloats](FName Name, const FFloatAnimationAttribute* BaseValue, const FFloatAnimationAttribute* LayerValue)
 				{
 					// Mismatched pairs are ignored
diff --git a/Engine/Plugins/Experimental/UAF/UAF/Source/UAF/Private/UAF/ValueRuntime/Transformers/Overwrite.cpp b/Engine/Plugins/Experimental/UAF/UAF/Source/UAF/Private/UAF/ValueRuntime/Transformers/Overwrite.cpp
--- a/Engine/Plugins/Experimental/UAF/UAF/Source/UAF/Private/UAF/ValueRuntime/Transformers/Overwrite.cpp
+++ b/Engine/Plugins/Experimental/UAF/UAF/Source/UAF/Private/UAF/ValueRuntime/Transformers/Overwrite.cpp
@@ -8,7 +8,7 @@ namespace UE::UAF::Transformers
 {
 	FName FOverwrite::TransformerName = TEXT("UAF::Transformers::Overwrite");
 
-	void FOverwrite::Apply(const FAttributeValueTransformerMapPtr& TransformerMap, const FAttributeCollection& Input, float Weight, FAttributeCollection& Output)
+	void FOverwrite::Apply(const FValueTransformerMapPtr& TransformerMap, const FAttributeCollection& Input, float Weight, FAttributeCollection& Output)
 	{
 		const bool bIsInputOutput = &Input == &Output;
 		if (bIsInputOutput && FAnimWeight::IsFullWeight(Weight))
@@ -22,11 +22,11 @@ namespace UE::UAF::Transformers
 		FOverwrite::Apply(TransformerMap, Input, EmptyCollection, Weight, Output);
 	}
 
-	void FOverwrite::Apply(const FAttributeValueTransformerMapPtr& TransformerMap, const FAttributeCollection& Input, const FAttributeCollection& PerValueWeights, float DefaultWeight, FAttributeCollection& Output)
+	void FOverwrite::Apply(const FValueTransformerMapPtr& TransformerMap, const FAttributeCollection& Input, const FAttributeCollection& PerValueWeights, float DefaultWeight, FAttributeCollection& Output)
 	{
 		checkf(Input.GetNamedSet() == PerValueWeights.GetNamedSet(), TEXT("Input and per-value weight collections must use the same named set"));
 
-		const FAttributeValueTransformerList* TransformerList = TransformerMap->Find(TransformerName);
+		const FValueTransformerList* TransformerList = TransformerMap->Find(TransformerName);
 		if (TransformerList == nullptr)
 		{
 			// No transformer implementations have been registered
@@ -41,13 +41,13 @@ namespace UE::UAF::Transformers
 			// Our output collection matches one of our inputs, we'll modify it in place
 			FReallocFun OutputAllocator = Output.GetAllocator();
 
-			// Set maps
+			// Bound value maps
 			{
 				OuterJoinBy(
-					FAttributeValueTransformerListWithAttributeCollectionSetMapJoinOp(),
-					[&PerValueWeights, DefaultWeight, FloatValueType, OutputAllocator](FRawTransformerFunc Transformer, FAttributeSetMap* InputSetMap)
+					FValueTransformerListWithBoundMapCollectionJoinOp(),
+					[&PerValueWeights, DefaultWeight, FloatValueType, OutputAllocator](FRawTransformerFunc Transformer, FBoundValueMap* InputMap)
 					{
-						if (InputSetMap == nullptr)
+						if (InputMap == nullptr)
 						{
 							// No maps found for the transformer type
 						}
@@ -60,18 +60,18 @@ namespace UE::UAF::Transformers
 							else
 							{
 								// Transform our input in place
-								FAttributeSetMap* OutputSetMap = InputSetMap;
+								FBoundValueMap* OutputMap = InputMap;
 
 								// We cannot join our weights because their value type differs
-								const FAttributeMappingKey MappingKey = InputSetMap->GetMappingKey().To(FloatValueType);
-								const FAttributeSetMap* PerValueWeightMap = PerValueWeights.FindSetMap(MappingKey);
+								const FAttributeMappingKey MappingKey = InputMap->GetMappingKey().To(FloatValueType);
+								const FBoundValueMap* PerValueWeightMap = PerValueWeights.FindSetMap(MappingKey);
 
-								FTransformSetMapFunc TypedTransformer = reinterpret_cast<FTransformSetMapFunc>(Transformer);
-								(*TypedTransformer)(InputSetMap, PerValueWeightMap, DefaultWeight, OutputSetMap, nullptr);
+								FTransformBoundValueMapFunc TypedTransformer = reinterpret_cast<FTransformBoundValueMapFunc>(Transformer);
+								(*TypedTransformer)(InputMap, PerValueWeightMap, DefaultWeight, OutputMap, nullptr);
 							}
 						}
 					},
-					TransformerList->CreateSetMapTransformerIterator(),
+					TransformerList->CreateBoundValueMapTransformerIterator(),
 					// Our input is our output
 					Output.CreateSetMapIterator());
 			}
@@ -79,8 +79,8 @@ namespace UE::UAF::Transformers
 			// Name/value maps
 			{
 				OuterJoinBy(
-					FAttributeValueTransformerListWithAttributeCollectionMapJoinOp(),
-					[DefaultWeight, OutputAllocator](FRawTransformerFunc Transformer, FAttributeMap* InputMap)
+					FValueTransformerListWithUnboundMapCollectionJoinOp(),
+					[DefaultWeight, OutputAllocator](FRawTransformerFunc Transformer, FUnboundValueMap* InputMap)
 					{
 						if (InputMap == nullptr)
 						{
@@ -95,14 +95,14 @@ namespace UE::UAF::Transformers
 							else
 							{
 								// Transform our input in place
-								FAttributeMap* OutputMap = InputMap;
+								FUnboundValueMap* OutputMap = InputMap;
 
-								FTransformMapFunc TypedTransformer = reinterpret_cast<FTransformMapFunc>(Transformer);
+								FTransformUnboundValueMapFunc TypedTransformer = reinterpret_cast<FTransformUnboundValueMapFunc>(Transformer);
 								(*TypedTransformer)(InputMap, DefaultWeight, OutputMap, nullptr);
 							}
 						}
 					},
-					TransformerList->CreateMapTransformerIterator(),
+					TransformerList->CreateUnboundValueMapTransformerIterator(),
 					// Our input is our output
 					Output.CreateMapIterator());
 			}
@@ -113,12 +113,12 @@ namespace UE::UAF::Transformers
 			Output.Reset(Input.GetNamedSet());
 			Output.SetValueSpace(Input.GetValueSpace());
 
-			// Set maps
+			// Bound value maps
 			OuterJoinBy(
-				FAttributeValueTransformerListWithAttributeCollectionSetMapJoinOp(),
-				[&PerValueWeights, DefaultWeight, FloatValueType, &Output](FRawTransformerFunc Transformer, const FAttributeSetMap* InputSetMap)
+				FValueTransformerListWithBoundMapCollectionJoinOp(),
+				[&PerValueWeights, DefaultWeight, FloatValueType, &Output](FRawTransformerFunc Transformer, const FBoundValueMap* InputMap)
 				{
-					if (InputSetMap == nullptr)
+					if (InputMap == nullptr)
 					{
 						// No maps found for the transformer type
 					}
@@ -127,33 +127,33 @@ namespace UE::UAF::Transformers
 						if (Transformer == nullptr)
 						{
 							// No transformer was found for this value type, retain it as-is
-							Output.AppendSetMap(InputSetMap->Duplicate(Output.GetAllocator()));
+							Output.AppendSetMap(InputMap->Duplicate(Output.GetAllocator()));
 						}
 						else
 						{
 							// Transform our matching pair
-							FAttributeSetMap* OutputSetMap = nullptr;
+							FBoundValueMap* OutputMap = nullptr;
 
 							// We cannot join our weights because their value type differs
-							const FAttributeMappingKey MappingKey = InputSetMap->GetMappingKey().To(FloatValueType);
-							const FAttributeSetMap* PerValueWeightMap = PerValueWeights.FindSetMap(MappingKey);
+							const FAttributeMappingKey MappingKey = InputMap->GetMappingKey().To(FloatValueType);
+							const FBoundValueMap* PerValueWeightMap = PerValueWeights.FindSetMap(MappingKey);
 
-							FTransformSetMapFunc TypedTransformer = reinterpret_cast<FTransformSetMapFunc>(Transformer);
-							(*TypedTransformer)(InputSetMap, PerValueWeightMap, DefaultWeight, OutputSetMap, Output.GetAllocator());
+							FTransformBoundValueMapFunc TypedTransformer = reinterpret_cast<FTransformBoundValueMapFunc>(Transformer);
+							(*TypedTransformer)(InputMap, PerValueWeightMap, DefaultWeight, OutputMap, Output.GetAllocator());
 
-							checkf(OutputSetMap != nullptr, TEXT("Transformer function is responsible for allocating an output buffer"));
-							checkf(OutputSetMap->GetAllocator() == Output.GetAllocator(), TEXT("Transformer function is responsible for allocating using the output allocator"));
-							Output.AppendSetMap(OutputSetMap);
+							checkf(OutputMap != nullptr, TEXT("Transformer function is responsible for allocating an output buffer"));
+							checkf(OutputMap->GetAllocator() == Output.GetAllocator(), TEXT("Transformer function is responsible for allocating using the output allocator"));
+							Output.AppendSetMap(OutputMap);
 						}
 					}
 				},
-				TransformerList->CreateSetMapTransformerIterator(),
+				TransformerList->CreateBoundValueMapTransformerIterator(),
 				Input.CreateConstSetMapIterator());
 
 			// Name/value maps
 			OuterJoinBy(
-				FAttributeValueTransformerListWithAttributeCollectionMapJoinOp(),
-				[DefaultWeight, &Output](FRawTransformerFunc Transformer, const FAttributeMap* InputMap)
+				FValueTransformerListWithUnboundMapCollectionJoinOp(),
+				[DefaultWeight, &Output](FRawTransformerFunc Transformer, const FUnboundValueMap* InputMap)
 				{
 					if (InputMap == nullptr)
 					{
@@ -169,9 +169,9 @@ namespace UE::UAF::Transformers
 						else
 						{
 							// Transform our matching pair
-							FAttributeMap* OutputMap = nullptr;
+							FUnboundValueMap* OutputMap = nullptr;
 
-							FTransformMapFunc TypedTransformer = reinterpret_cast<FTransformMapFunc>(Transformer);
+							FTransformUnboundValueMapFunc TypedTransformer = reinterpret_cast<FTransformUnboundValueMapFunc>(Transformer);
 							(*TypedTransformer)(InputMap, DefaultWeight, OutputMap, Output.GetAllocator());
 
 							checkf(OutputMap != nullptr, TEXT("Transformer function is responsible for allocating an output buffer"));
@@ -180,35 +180,35 @@ namespace UE::UAF::Transformers
 						}
 					}
 				},
-				TransformerList->CreateMapTransformerIterator(),
+				TransformerList->CreateUnboundValueMapTransformerIterator(),
 				Input.CreateConstMapIterator());
 		}
 	}
 
 	namespace Private
 	{
-		void FOverwrite_BoneTransformAttribute::TransformSetMap(const FAttributeSetMap* InputSetMap, const FAttributeSetMap* PerValueWeightMap, float DefaultWeight, FAttributeSetMap*& OutputSetMap, FReallocFun OutputAllocator)
+		void FOverwrite_BoneTransformAttribute::TransformBoundValueMap(const FBoundValueMap* InputMap, const FBoundValueMap* PerValueWeightMap, float DefaultWeight, FBoundValueMap*& OutputMap, FReallocFun OutputAllocator)
 		{
-			if (OutputSetMap == nullptr)
+			if (OutputMap == nullptr)
 			{
-				const FAttributeSetMap::FConstructArgs Args(InputSetMap->GetTypedSet(), InputSetMap->GetValueType(), OutputAllocator);
-				OutputSetMap = MakeAttributeSetMap<FBoneTransformAnimationAttribute>(Args);
+				const FBoundValueMap::FConstructArgs Args(InputMap->GetTypedSet(), InputMap->GetValueType(), OutputAllocator);
+				OutputMap = MakeBoundValueMap<FBoneTransformAnimationAttribute>(Args);
 			}
 
-			check(InputSetMap->GetTypedSet() == OutputSetMap->GetTypedSet());
-			check(InputSetMap->Num() == OutputSetMap->Num());
+			check(InputMap->GetTypedSet() == OutputMap->GetTypedSet());
+			check(InputMap->Num() == OutputMap->Num());
 
-			const TAttributeSetMap<FBoneTransformAnimationAttribute>* InputBoneTransforms = CastChecked<FBoneTransformAnimationAttribute>(InputSetMap);
-			TAttributeSetMap<FBoneTransformAnimationAttribute>* OutputBoneTransforms = CastChecked<FBoneTransformAnimationAttribute>(OutputSetMap);
+			const TBoundValueMap<FBoneTransformAnimationAttribute>* InputBoneTransforms = CastChecked<FBoneTransformAnimationAttribute>(InputMap);
+			TBoundValueMap<FBoneTransformAnimationAttribute>* OutputBoneTransforms = CastChecked<FBoneTransformAnimationAttribute>(OutputMap);
 
 			const FBoneTransformAnimationAttribute* InputTransformPtr = InputBoneTransforms->GetData();
 			FBoneTransformAnimationAttribute* OutputTransformPtr = OutputBoneTransforms->GetData();
 			const FBoneTransformAnimationAttribute* InputTransformPtrEnd = InputTransformPtr + InputBoneTransforms->Num();
 
-			if (const TAttributeSetMap<FFloatAnimationAttribute>* PerValueWeights = Cast<FFloatAnimationAttribute>(PerValueWeightMap))
+			if (const TBoundValueMap<FFloatAnimationAttribute>* PerValueWeights = Cast<FFloatAnimationAttribute>(PerValueWeightMap))
 			{
-				check(InputSetMap->GetTypedSet() == PerValueWeights->GetTypedSet());
-				check(InputSetMap->Num() == PerValueWeights->Num());
+				check(InputMap->GetTypedSet() == PerValueWeights->GetTypedSet());
+				check(InputMap->Num() == PerValueWeights->Num());
 
 				const FFloatAnimationAttribute* PerValueWeightsPtr = PerValueWeights->GetData();
 
@@ -235,16 +235,16 @@ namespace UE::UAF::Transformers
 			}
 		}
 
-		void FOverwrite_BoneTransformAttribute::TransformMap(const FAttributeMap* InputMap, float Weight, FAttributeMap*& OutputMap, FReallocFun OutputAllocator)
+		void FOverwrite_BoneTransformAttribute::TransformUnboundValueMap(const FUnboundValueMap* InputMap, float Weight, FUnboundValueMap*& OutputMap, FReallocFun OutputAllocator)
 		{
 			const bool bAllocateOutput = OutputMap == nullptr;
 			if (bAllocateOutput)
 			{
-				OutputMap = MakeAttributeMap<FBoneTransformAnimationAttribute>(OutputAllocator);
+				OutputMap = MakeUnboundValueMap<FBoneTransformAnimationAttribute>(OutputAllocator);
 			}
 
-			const TAttributeMap<FBoneTransformAnimationAttribute>* InputBoneTransforms = CastChecked<FBoneTransformAnimationAttribute>(InputMap);
-			TAttributeMap<FBoneTransformAnimationAttribute>* OutputBoneTransforms = CastChecked<FBoneTransformAnimationAttribute>(OutputMap);
+			const TUnboundValueMap<FBoneTransformAnimationAttribute>* InputBoneTransforms = CastChecked<FBoneTransformAnimationAttribute>(InputMap);
+			TUnboundValueMap<FBoneTransformAnimationAttribute>* OutputBoneTransforms = CastChecked<FBoneTransformAnimationAttribute>(OutputMap);
 
 			const ScalarRegister VWeight(Weight);
 
@@ -269,28 +269,28 @@ namespace UE::UAF::Transformers
 			}
 		}
 
-		void FOverwrite_FloatAttribute::TransformSetMap(const FAttributeSetMap* InputSetMap, const FAttributeSetMap* PerValueWeightMap, float DefaultWeight, FAttributeSetMap*& OutputSetMap, FReallocFun OutputAllocator)
+		void FOverwrite_FloatAttribute::TransformBoundValueMap(const FBoundValueMap* InputMap, const FBoundValueMap* PerValueWeightMap, float DefaultWeight, FBoundValueMap*& OutputMap, FReallocFun OutputAllocator)
 		{
-			if (OutputSetMap == nullptr)
+			if (OutputMap == nullptr)
 			{
-				const FAttributeSetMap::FConstructArgs Args(InputSetMap->GetTypedSet(), InputSetMap->GetValueType(), OutputAllocator);
-				OutputSetMap = MakeAttributeSetMap<FFloatAnimationAttribute>(Args);
+				const FBoundValueMap::FConstructArgs Args(InputMap->GetTypedSet(), InputMap->GetValueType(), OutputAllocator);
+				OutputMap = MakeBoundValueMap<FFloatAnimationAttribute>(Args);
 			}
 
-			check(InputSetMap->GetTypedSet() == OutputSetMap->GetTypedSet());
-			check(InputSetMap->Num() == OutputSetMap->Num());
+			check(InputMap->GetTypedSet() == OutputMap->GetTypedSet());
+			check(InputMap->Num() == OutputMap->Num());
 
-			const TAttributeSetMap<FFloatAnimationAttribute>* InputFloats = CastChecked<FFloatAnimationAttribute>(InputSetMap);
-			TAttributeSetMap<FFloatAnimationAttribute>* OutputFloats = CastChecked<FFloatAnimationAttribute>(OutputSetMap);
+			const TBoundValueMap<FFloatAnimationAttribute>* InputFloats = CastChecked<FFloatAnimationAttribute>(InputMap);
+			TBoundValueMap<FFloatAnimationAttribute>* OutputFloats = CastChecked<FFloatAnimationAttribute>(OutputMap);
 
 			const FFloatAnimationAttribute* FloatsPtr = InputFloats->GetData();
 			FFloatAnimationAttribute* OutputFloatsPtr = OutputFloats->GetData();
 			const FFloatAnimationAttribute* FloatsPtrEnd = FloatsPtr + InputFloats->Num();
 
-			if (const TAttributeSetMap<FFloatAnimationAttribute>* PerValueWeights = Cast<FFloatAnimationAttribute>(PerValueWeightMap))
+			if (const TBoundValueMap<FFloatAnimationAttribute>* PerValueWeights = Cast<FFloatAnimationAttribute>(PerValueWeightMap))
 			{
-				check(InputSetMap->GetTypedSet() == PerValueWeights->GetTypedSet());
-				check(InputSetMap->Num() == PerValueWeights->Num());
+				check(InputMap->GetTypedSet() == PerValueWeights->GetTypedSet());
+				check(InputMap->Num() == PerValueWeights->Num());
 
 				const FFloatAnimationAttribute* PerValueWeightsPtr = PerValueWeights->GetData();
 
@@ -315,16 +315,16 @@ namespace UE::UAF::Transformers
 			}
 		}
 
-		void FOverwrite_FloatAttribute::TransformMap(const FAttributeMap* InputMap, float Weight, FAttributeMap*& OutputMap, FReallocFun OutputAllocator)
+		void FOverwrite_FloatAttribute::TransformUnboundValueMap(const FUnboundValueMap* InputMap, float Weight, FUnboundValueMap*& OutputMap, FReallocFun OutputAllocator)
 		{
 			const bool bAllocateOutput = OutputMap == nullptr;
 			if (bAllocateOutput)
 			{
-				OutputMap = MakeAttributeMap<FFloatAnimationAttribute>(OutputAllocator);
+				OutputMap = MakeUnboundValueMap<FFloatAnimationAttribute>(OutputAllocator);
 			}
 
-			const TAttributeMap<FFloatAnimationAttribute>* InputFloats = CastChecked<FFloatAnimationAttribute>(InputMap);
-			TAttributeMap<FFloatAnimationAttribute>* OutputFloats = CastChecked<FFloatAnimationAttribute>(OutputMap);
+			const TUnboundValueMap<FFloatAnimationAttribute>* InputFloats = CastChecked<FFloatAnimationAttribute>(InputMap);
+			TUnboundValueMap<FFloatAnimationAttribute>* OutputFloats = CastChecked<FFloatAnimationAttribute>(OutputMap);
 
 			if (bAllocateOutput)
 			{
diff --git a/Engine/Plugins/Experimental/UAF/UAF/Source/UAF/Private/UAF/ValueRuntime/Transformers/Sanitize.cpp b/Engine/Plugins/Experimental/UAF/UAF/Source/UAF/Private/UAF/ValueRuntime/Transformers/Sanitize.cpp
--- a/Engine/Plugins/Experimental/UAF/UAF/Source/UAF/Private/UAF/ValueRuntime/Transformers/Sanitize.cpp
+++ b/Engine/Plugins/Experimental/UAF/UAF/Source/UAF/Private/UAF/ValueRuntime/Transformers/Sanitize.cpp
@@ -8,9 +8,9 @@ namespace UE::UAF::Transformers
 {
 	FName FSanitize::TransformerName = TEXT("UAF::Transformers::Sanitize");
 
-	void FSanitize::Apply(const FAttributeValueTransformerMapPtr& TransformerMap, FAttributeCollection& Collection)
+	void FSanitize::Apply(const FValueTransformerMapPtr& TransformerMap, FAttributeCollection& Collection)
 	{
-		const FAttributeValueTransformerList* TransformerList = TransformerMap->Find(TransformerName);
+		const FValueTransformerList* TransformerList = TransformerMap->Find(TransformerName);
 		if (TransformerList == nullptr)
 		{
 			// No transformer implementations have been registered
@@ -19,24 +19,24 @@ namespace UE::UAF::Transformers
 
 		// Set maps
 		InnerJoinBy(
-			FAttributeValueTransformerListWithAttributeCollectionSetMapJoinOp(),
-			[](FRawTransformerFunc Transformer, FAttributeSetMap* SetMap)
+			FValueTransformerListWithBoundMapCollectionJoinOp(),
+			[](FRawTransformerFunc Transformer, FBoundValueMap* Map)
 			{
-				FTransformSetMapFunc TypedTransformer = reinterpret_cast<FTransformSetMapFunc>(Transformer);
-				(*TypedTransformer)(SetMap);
+				FTransformBoundValueMapFunc TypedTransformer = reinterpret_cast<FTransformBoundValueMapFunc>(Transformer);
+				(*TypedTransformer)(Map);
 			},
-			TransformerList->CreateSetMapTransformerIterator(),
+			TransformerList->CreateBoundValueMapTransformerIterator(),
 			Collection.CreateSetMapIterator());
 
 		// Name/value maps
 		InnerJoinBy(
-			FAttributeValueTransformerListWithAttributeCollectionMapJoinOp(),
-			[](FRawTransformerFunc Transformer, FAttributeMap* Map)
+			FValueTransformerListWithUnboundMapCollectionJoinOp(),
+			[](FRawTransformerFunc Transformer, FUnboundValueMap* Map)
 			{
-				FTransformMapFunc TypedTransformer = reinterpret_cast<FTransformMapFunc>(Transformer);
+				FTransformUnboundValueMapFunc TypedTransformer = reinterpret_cast<FTransformUnboundValueMapFunc>(Transformer);
 				(*TypedTransformer)(Map);
 			},
-			TransformerList->CreateMapTransformerIterator(),
+			TransformerList->CreateUnboundValueMapTransformerIterator(),
 			Collection.CreateMapIterator());
 	}
 
@@ -44,7 +44,7 @@ namespace UE::UAF::Transformers
 	{
 #if UE_UAF_VALIDATE_MAPPED_VALUES
 		// Don't inline this function to keep the stack usage down
-		static FORCENOINLINE void IsBoneTransformSetMapValid(const TAttributeSetMap<FBoneTransformAnimationAttribute>* BoneTransforms)
+		static FORCENOINLINE void IsBoneTransformMapValid(const TBoundValueMap<FBoneTransformAnimationAttribute>* BoneTransforms)
 		{
 			const FBoneTransformAnimationAttribute* TransformPtr = BoneTransforms->GetData();
 			const FBoneTransformAnimationAttribute* EndTransformPtr = TransformPtr + BoneTransforms->Num();
@@ -69,7 +69,7 @@ namespace UE::UAF::Transformers
 		}
 
 		// Don't inline this function to keep the stack usage down
-		static FORCENOINLINE void IsBoneTransformMapValid(const TAttributeMap<FBoneTransformAnimationAttribute>* BoneTransforms)
+		static FORCENOINLINE void IsBoneTransformMapValid(const TUnboundValueMap<FBoneTransformAnimationAttribute>* BoneTransforms)
 		{
 			for (auto It = BoneTransforms->CreateConstIterator(); It; ++It)
 			{
@@ -89,9 +89,9 @@ namespace UE::UAF::Transformers
 		}
 #endif
 
-		void FSanitize_BoneTransformAttribute::TransformSetMap(FAttributeSetMap* SetMap)
+		void FSanitize_BoneTransformAttribute::TransformBoundValueMap(FBoundValueMap* Map)
 		{
-			TAttributeSetMap<FBoneTransformAnimationAttribute>* BoneTransforms = CastChecked<FBoneTransformAnimationAttribute>(SetMap);
+			TBoundValueMap<FBoneTransformAnimationAttribute>* BoneTransforms = CastChecked<FBoneTransformAnimationAttribute>(Map);
 
 			FBoneTransformAnimationAttribute* TransformPtr = BoneTransforms->GetData();
 			const FBoneTransformAnimationAttribute* TransformPtrEnd = TransformPtr + BoneTransforms->Num();
@@ -104,13 +104,13 @@ namespace UE::UAF::Transformers
 			}
 
 #if UE_UAF_VALIDATE_MAPPED_VALUES
-			IsBoneTransformSetMapValid(BoneTransforms);
+			IsBoneTransformMapValid(BoneTransforms);
 #endif
 		}
 
-		void FSanitize_BoneTransformAttribute::TransformMap(FAttributeMap* Map)
+		void FSanitize_BoneTransformAttribute::TransformUnboundValueMap(FUnboundValueMap* Map)
 		{
-			TAttributeMap<FBoneTransformAnimationAttribute>* BoneTransforms = CastChecked<FBoneTransformAnimationAttribute>(Map);
+			TUnboundValueMap<FBoneTransformAnimationAttribute>* BoneTransforms = CastChecked<FBoneTransformAnimationAttribute>(Map);
 
 			for (auto It = BoneTransforms->CreateIterator(); It; ++It)
 			{
@@ -123,9 +123,9 @@ namespace UE::UAF::Transformers
 		}
 
 #if UE_UAF_VALIDATE_MAPPED_VALUES
-		void FSanitize_FloatAttribute::TransformSetMap(FAttributeSetMap* SetMap)
+		void FSanitize_FloatAttribute::TransformBoundValueMap(FBoundValueMap* Map)
 		{
-			const TAttributeSetMap<FFloatAnimationAttribute>* Floats = CastChecked<FFloatAnimationAttribute>(SetMap);
+			const TBoundValueMap<FFloatAnimationAttribute>* Floats = CastChecked<FFloatAnimationAttribute>(Map);
 
 			const FFloatAnimationAttribute* FloatsPtr = Floats->GetData();
 			const FFloatAnimationAttribute* FloatsPtrEnd = FloatsPtr + Floats->Num();
@@ -140,9 +140,9 @@ namespace UE::UAF::Transformers
 			}
 		}
 
-		void FSanitize_FloatAttribute::TransformMap(FAttributeMap* Map)
+		void FSanitize_FloatAttribute::TransformUnboundValueMap(FUnboundValueMap* Map)
 		{
-			const TAttributeMap<FFloatAnimationAttribute>* Floats = CastChecked<FFloatAnimationAttribute>(Map);
+			const TUnboundValueMap<FFloatAnimationAttribute>* Floats = CastChecked<FFloatAnimationAttribute>(Map);
 
 			for (auto It = Floats->CreateConstIterator(); It; ++It)
 			{
diff --git a/Engine/Plugins/Experimental/UAF/UAF/Source/UAF/Private/UAF/ValueRuntime/Transformers/Transformer.cpp b/Engine/Plugins/Experimental/UAF/UAF/Source/UAF/Private/UAF/ValueRuntime/Transformers/Transformer.cpp
--- a/Engine/Plugins/Experimental/UAF/UAF/Source/UAF/Private/UAF/ValueRuntime/Transformers/Transformer.cpp
+++ b/Engine/Plugins/Experimental/UAF/UAF/Source/UAF/Private/UAF/ValueRuntime/Transformers/Transformer.cpp
@@ -4,21 +4,21 @@
 
 namespace UE::UAF
 {
-	FAttributeValueTransformerList::FAttributeValueTransformerList() = default;
+	FValueTransformerList::FValueTransformerList() = default;
 
-	FAttributeValueTransformerList::FAttributeValueTransformerList(FName InTransformerName)
+	FValueTransformerList::FValueTransformerList(FName InTransformerName)
 		: TransformerName(InTransformerName)
 	{
 	}
 
-	FName FAttributeValueTransformerList::GetTransformerName() const
+	FName FValueTransformerList::GetTransformerName() const
 	{
 		return TransformerName;
 	}
 
-	bool FAttributeValueTransformerList::AddSetMapTransformer(UScriptStruct* ValueType, FRawTransformerFunc TransformerFunc)
+	bool FValueTransformerList::AddBoundValueMapTransformer(UScriptStruct* ValueType, FRawTransformerFunc TransformerFunc)
 	{
-		const int32 EntryIndex = Algo::LowerBound(MakeArrayView(ValueTypes.GetData(), NumSetMapTransformers()), ValueType);
+		const int32 EntryIndex = Algo::LowerBound(MakeArrayView(ValueTypes.GetData(), NumBoundValueMapTransformers()), ValueType);
 		if (EntryIndex < ValueTypes.Num() && ValueTypes[EntryIndex] == ValueType)
 		{
 			// A transformer for this value type has already been registered
@@ -27,20 +27,20 @@ namespace UE::UAF
 
 		ValueTypes.Insert(ValueType, EntryIndex);
 		Transformers.Insert(TransformerFunc, EntryIndex);
-		SetMapTransformerCount++;
+		BoundValueMapTransformerCount++;
 
 		return true;
 	}
 
-	bool FAttributeValueTransformerList::AddMapTransformer(UScriptStruct* ValueType, FRawTransformerFunc TransformerFunc)
+	bool FValueTransformerList::AddUnboundValueMapTransformer(UScriptStruct* ValueType, FRawTransformerFunc TransformerFunc)
 	{
-		const int32 NumSetMapEntries = NumSetMapTransformers();
-		const int32 NumMapEntries = NumMapTransformers();
+		const int32 NumBoundValueMapEntries = NumBoundValueMapTransformers();
+		const int32 NumUnboundValueMapEntries = NumUnboundValueMapTransformers();
 
-		const int32 PerContainerEntryIndex = Algo::LowerBound(MakeArrayView(ValueTypes.GetData() + NumSetMapEntries, NumMapEntries), ValueType);
+		const int32 PerContainerEntryIndex = Algo::LowerBound(MakeArrayView(ValueTypes.GetData() + NumBoundValueMapEntries, NumUnboundValueMapEntries), ValueType);
 
 		// Fixup our index
-		const int32 EntryIndex = PerContainerEntryIndex + NumSetMapEntries;
+		const int32 EntryIndex = PerContainerEntryIndex + NumBoundValueMapEntries;
 
 		if (EntryIndex < ValueTypes.Num() && ValueTypes[EntryIndex] == ValueType)
 		{
diff --git a/Engine/Plugins/Experimental/UAF/UAF/Source/UAF/Private/UAF/ValueRuntime/ValueRuntimeRegistry.cpp b/Engine/Plugins/Experimental/UAF/UAF/Source/UAF/Private/UAF/ValueRuntime/ValueRuntimeRegistry.cpp
--- a/Engine/Plugins/Experimental/UAF/UAF/Source/UAF/Private/UAF/ValueRuntime/ValueRuntimeRegistry.cpp
+++ b/Engine/Plugins/Experimental/UAF/UAF/Source/UAF/Private/UAF/ValueRuntime/ValueRuntimeRegistry.cpp
@@ -36,10 +36,10 @@ namespace UE::UAF
 
 	FValueRuntimeRegistry::FValueRuntimeRegistry() = default;
 
-	FAttributeSetMap* FValueRuntimeRegistry::ConstructSetMap(const FAttributeSetMap::FConstructArgs& Args) const
+	FBoundValueMap* FValueRuntimeRegistry::ConstructBoundValueMap(const FBoundValueMap::FConstructArgs& Args) const
 	{
-		FReadScopeLock Lock(SetMapInitializerLock);
-		if (const SetMapConstructorFun* ConstructorFun = SetMapInitializers.Find(Args.ValueType))
+		FReadScopeLock Lock(BoundValueMapInitializerLock);
+		if (const BoundValueMapConstructorFun* ConstructorFun = BoundValueMapInitializers.Find(Args.ValueType))
 		{
 			return (*ConstructorFun)(Args);
 		}
@@ -48,10 +48,10 @@ namespace UE::UAF
 		return nullptr;
 	}
 
-	FAttributeMap* FValueRuntimeRegistry::ConstructMap(UScriptStruct* ValueType, FReallocFun ReallocFun) const
+	FUnboundValueMap* FValueRuntimeRegistry::ConstructUnboundValueMap(UScriptStruct* ValueType, FReallocFun ReallocFun) const
 	{
-		FReadScopeLock Lock(SetMapInitializerLock);
-		if (const MapConstructorFun* ConstructorFun = MapInitializers.Find(ValueType))
+		FReadScopeLock Lock(UnboundValueMapInitializerLock);
+		if (const UnboundValueMapConstructorFun* ConstructorFun = UnboundValueMapInitializers.Find(ValueType))
 		{
 			return (*ConstructorFun)(ReallocFun);
 		}
@@ -60,44 +60,44 @@ namespace UE::UAF
 		return nullptr;
 	}
 
-	FAttributeValueTransformerMapPtr FValueRuntimeRegistry::GetTransformerMap() const
+	FValueTransformerMapPtr FValueRuntimeRegistry::GetTransformerMap() const
 	{
 		FReadScopeLock Lock(TransformerMapLock);
 		return TransformerMap;
 	}
 
-	void FValueRuntimeRegistry::RegisterSetMapInitializer(UScriptStruct* ValueType, SetMapConstructorFun Initializer)
+	void FValueRuntimeRegistry::RegisterBoundValueMapInitializer(UScriptStruct* ValueType, BoundValueMapConstructorFun Initializer)
 	{
-		FWriteScopeLock Lock(SetMapInitializerLock);
-		SetMapInitializers.Add(ValueType, Initializer);
+		FWriteScopeLock Lock(BoundValueMapInitializerLock);
+		BoundValueMapInitializers.Add(ValueType, Initializer);
 	}
 
-	void FValueRuntimeRegistry::RegisterMapInitializer(UScriptStruct* ValueType, MapConstructorFun Initializer)
+	void FValueRuntimeRegistry::RegisterUnboundValueMapInitializer(UScriptStruct* ValueType, UnboundValueMapConstructorFun Initializer)
 	{
-		FWriteScopeLock Lock(MapInitializerLock);
-		MapInitializers.Add(ValueType, Initializer);
+		FWriteScopeLock Lock(UnboundValueMapInitializerLock);
+		UnboundValueMapInitializers.Add(ValueType, Initializer);
 	}
 
-	bool FValueRuntimeRegistry::RegisterAttributeValueSetMapTransformer(FName TransformerName, UScriptStruct* ValueType, FRawTransformerFunc TransformerFunc)
+	bool FValueRuntimeRegistry::RegisterBouneValueMapTransformer(FName TransformerName, UScriptStruct* ValueType, FRawTransformerFunc TransformerFunc)
 	{
 		FWriteScopeLock Lock(TransformerMapLock);
 
 		// If someone is already holding a reference to our named operators, we cannot mutate
 		// When this occurs, we make a copy and swap it, keeping the old one alive until it is no longer needed
-		FAttributeValueTransformerMapPtr MutableTransformerMap = TransformerMap;
+		FValueTransformerMapPtr MutableTransformerMap = TransformerMap;
 		if (!MutableTransformerMap || MutableTransformerMap.GetSharedReferenceCount() != 2)
 		{
-			MutableTransformerMap = MakeShared<typename FAttributeValueTransformerMapPtr::ElementType>();
+			MutableTransformerMap = MakeShared<typename FValueTransformerMapPtr::ElementType>();
 		}
 
 		const uint32 TransformerNameHash = GetTypeHash(TransformerName);
-		FAttributeValueTransformerList* TransformerList = MutableTransformerMap->FindByHash(TransformerNameHash, TransformerName);
+		FValueTransformerList* TransformerList = MutableTransformerMap->FindByHash(TransformerNameHash, TransformerName);
 		if (TransformerList == nullptr)
 		{
-			TransformerList = &MutableTransformerMap->AddByHash(TransformerNameHash, TransformerName, FAttributeValueTransformerList(TransformerName));
+			TransformerList = &MutableTransformerMap->AddByHash(TransformerNameHash, TransformerName, FValueTransformerList(TransformerName));
 		}
 
-		if (!TransformerList->AddSetMapTransformer(ValueType, TransformerFunc))
+		if (!TransformerList->AddBoundValueMapTransformer(ValueType, TransformerFunc))
 		{
 			// We failed to add this transformer
 			return false;
@@ -111,26 +111,26 @@ namespace UE::UAF
 		return true;
 	}
 
-	bool FValueRuntimeRegistry::RegisterAttributeValueMapTransformer(FName TransformerName, UScriptStruct* ValueType, FRawTransformerFunc TransformerFunc)
+	bool FValueRuntimeRegistry::RegisterUnboundValueMapTransformer(FName TransformerName, UScriptStruct* ValueType, FRawTransformerFunc TransformerFunc)
 	{
 		FWriteScopeLock Lock(TransformerMapLock);
 
 		// If someone is already holding a reference to our named operators, we cannot mutate
 		// When this occurs, we make a copy and swap it, keeping the old one alive until it is no longer needed
-		FAttributeValueTransformerMapPtr MutableTransformerMap = TransformerMap;
+		FValueTransformerMapPtr MutableTransformerMap = TransformerMap;
 		if (!MutableTransformerMap || MutableTransformerMap.GetSharedReferenceCount() != 2)
 		{
-			MutableTransformerMap = MakeShared<typename FAttributeValueTransformerMapPtr::ElementType>();
+			MutableTransformerMap = MakeShared<typename FValueTransformerMapPtr::ElementType>();
 		}
 
 		const uint32 TransformerNameHash = GetTypeHash(TransformerName);
-		FAttributeValueTransformerList* TransformerList = MutableTransformerMap->FindByHash(TransformerNameHash, TransformerName);
+		FValueTransformerList* TransformerList = MutableTransformerMap->FindByHash(TransformerNameHash, TransformerName);
 		if (TransformerList == nullptr)
 		{
-			TransformerList = &MutableTransformerMap->AddByHash(TransformerNameHash, TransformerName, FAttributeValueTransformerList(TransformerName));
+			TransformerList = &MutableTransformerMap->AddByHash(TransformerNameHash, TransformerName, FValueTransformerList(TransformerName));
 		}
 
-		if (!TransformerList->AddMapTransformer(ValueType, TransformerFunc))
+		if (!TransformerList->AddUnboundValueMapTransformer(ValueType, TransformerFunc))
 		{
 			// We failed to add this transformer
 			return false;
diff --git a/Engine/Plugins/Experimental/UAF/UAF/Source/UAF/Public/UAF/ValueRuntime/AttributeCollection.h b/Engine/Plugins/Experimental/UAF/UAF/Source/UAF/Public/UAF/ValueRuntime/AttributeCollection.h
--- a/Engine/Plugins/Experimental/UAF/UAF/Source/UAF/Public/UAF/ValueRuntime/AttributeCollection.h
+++ b/Engine/Plugins/Experimental/UAF/UAF/Source/UAF/Public/UAF/ValueRuntime/AttributeCollection.h
@@ -5,10 +5,10 @@
 #include "CoreMinimal.h"
 #include "Misc/MemStack.h"
 #include "UAF/Attributes/AttributeNamedSet.h"
-#include "UAF/ValueRuntime/AttributeMap.h"
 #include "UAF/ValueRuntime/AttributeMappingKey.h"
-#include "UAF/ValueRuntime/AttributeSetMap.h"
+#include "UAF/ValueRuntime/BoundValueMap.h"
 #include "UAF/ValueRuntime/IndirectAllocator.h"
+#include "UAF/ValueRuntime/UnboundValueMap.h"
 #include "UAF/ValueRuntime/ValueSpace.h"
 
 #include "AttributeCollection.generated.h"
@@ -25,7 +25,7 @@ namespace UE::UAF
 		static constexpr bool bIsContainerConst = std::is_const_v<ContainerType>;
 
 		using EntryPtrType = typename std::conditional<bIsContainerConst, const typename ContainerType::FSetMapEntry*, typename ContainerType::FSetMapEntry*>::type;
-		using SetMapType = typename std::conditional<bIsContainerConst, const FAttributeSetMap*, FAttributeSetMap*>::type;
+		using MapType = typename std::conditional<bIsContainerConst, const FBoundValueMap*, FBoundValueMap*>::type;
 
 		// Creates an empty iterator
 		[[nodiscard]] TAttributeCollectionSetMapIterator() = default;
@@ -49,7 +49,7 @@ namespace UE::UAF
 		[[nodiscard]] UScriptStruct* GetValueType() const;
 
 		// Returns the set map from the current entry
-		[[nodiscard]] SetMapType GetMap() const;
+		[[nodiscard]] MapType GetMap() const;
 
 	private:
 		// Current set map entry
@@ -67,7 +67,7 @@ namespace UE::UAF
 		static constexpr bool bIsContainerConst = std::is_const_v<ContainerType>;
 
 		using EntryPtrType = typename std::conditional<bIsContainerConst, const typename ContainerType::FMapEntry*, typename ContainerType::FMapEntry*>::type;
-		using MapType = typename std::conditional<bIsContainerConst, const FAttributeMap*, FAttributeMap*>::type;
+		using MapType = typename std::conditional<bIsContainerConst, const FUnboundValueMap*, FUnboundValueMap*>::type;
 
 		// Creates an empty iterator
 		[[nodiscard]] TAttributeCollectionMapIterator() = default;
@@ -111,7 +111,7 @@ namespace UE::UAF
 			FAttributeMappingKey Key;
 
 			// The set map containing our values
-			FAttributeSetMap* Map = nullptr;
+			FBoundValueMap* Map = nullptr;
 		};
 
 		struct FMapEntry
@@ -120,7 +120,7 @@ namespace UE::UAF
 			UScriptStruct* ValueType = nullptr;
 
 			// The map containing our values
-			FAttributeMap* Map = nullptr;
+			FUnboundValueMap* Map = nullptr;
 		};
 
 	public:
@@ -169,23 +169,23 @@ namespace UE::UAF
 
 		// Constructs and adds a new mapping of ValueType to a set of AttributeType to this collection.
 		// Returns nullptr if we fail to add an entry (e.g. duplicate mapping)
-		template <class ValueType, typename... TArgs> [[nodiscard]] TAttributeSetMap<ValueType>* AddSetMap(const FAttributeMappingKey& MappingKey, TArgs... Args);
+		template <class ValueType, typename... TArgs> [[nodiscard]] TBoundValueMap<ValueType>* AddSetMap(const FAttributeMappingKey& MappingKey, TArgs... Args);
 
 		// Constructs and adds a new set mapping to this collection.
 		// Returns nullptr if we fail to add an entry (e.g. duplicate mapping)
-		[[nodiscard]] UE_API FAttributeSetMap* AddSetMap(const FAttributeMappingKey& MappingKey);
+		[[nodiscard]] UE_API FBoundValueMap* AddSetMap(const FAttributeMappingKey& MappingKey);
 
 		// Adds a new set mapping to this collection.
 		// Returns false if we fail to add an entry (e.g. duplicate mapping, mismatched allocator)
-		UE_API bool AddSetMap(FAttributeSetMap* Map);
+		UE_API bool AddSetMap(FBoundValueMap* Map);
 
 		// Appends a new set mapping to this collection (the input must belong at the end of the sorted set maps).
 		// Returns false if we fail to add an entry (e.g. duplicate mapping, mismatched allocator, not sorted)
-		UE_API bool AppendSetMap(FAttributeSetMap* Map);
+		UE_API bool AppendSetMap(FBoundValueMap* Map);
 
 		// Removes the specified value set mapping from this collection
 		// Returns whether or not the entry was successfully removed
-		UE_API bool RemoveSetMap(FAttributeSetMap* Map);
+		UE_API bool RemoveSetMap(FBoundValueMap* Map);
 
 		// Removes the specified set mapping from this collection
 		// Returns whether or not the entry was successfully removed
@@ -193,27 +193,27 @@ namespace UE::UAF
 
 		// Finds the mapping of ValueType to a set of AttributeType in this collection
 		// Returns nullptr if this collection does not contain such a mapping
-		template <class ValueType> [[nodiscard]] TAttributeSetMap<ValueType>* FindSetMap(const FAttributeMappingKey& MappingKey);
+		template <class ValueType> [[nodiscard]] TBoundValueMap<ValueType>* FindSetMap(const FAttributeMappingKey& MappingKey);
 
 		// Finds the mapping of ValueType to a set of AttributeType in this collection
 		// Returns nullptr if this collection does not contain such a mapping
-		template <class ValueType> [[nodiscard]] const TAttributeSetMap<ValueType>* FindSetMap(const FAttributeMappingKey& MappingKey) const;
+		template <class ValueType> [[nodiscard]] const TBoundValueMap<ValueType>* FindSetMap(const FAttributeMappingKey& MappingKey) const;
 
 		// Finds the mapping of ValueType to a set of AttributeType in this collection
 		// Returns nullptr if this collection does not contain such a mapping
-		[[nodiscard]] UE_API FAttributeSetMap* FindSetMap(const FAttributeMappingKey& MappingKey);
+		[[nodiscard]] UE_API FBoundValueMap* FindSetMap(const FAttributeMappingKey& MappingKey);
 
 		// Finds the mapping of ValueType to a set of AttributeType in this collection
 		// Returns nullptr if this collection does not contain such a mapping
-		[[nodiscard]] UE_API const FAttributeSetMap* FindSetMap(const FAttributeMappingKey& MappingKey) const;
+		[[nodiscard]] UE_API const FBoundValueMap* FindSetMap(const FAttributeMappingKey& MappingKey) const;
 
 		// Finds or adds a mapping of ValueType to a set of AttributeType in this collection
 		// Returns nullptr if we fail to add an entry (e.g. duplicate mapping)
-		template <class ValueType, typename... TArgs> [[nodiscard]] TAttributeSetMap<ValueType>* FindOrAddSetMap(const FAttributeMappingKey& MappingKey, TArgs... Args);
+		template <class ValueType, typename... TArgs> [[nodiscard]] TBoundValueMap<ValueType>* FindOrAddSetMap(const FAttributeMappingKey& MappingKey, TArgs... Args);
 
 		// Finds or adds a mapping of ValueType to a set of AttributeType in this collection
 		// Returns nullptr if we fail to add an entry (e.g. duplicate mapping)
-		[[nodiscard]] UE_API FAttributeSetMap* FindOrAddSetMap(const FAttributeMappingKey& MappingKey);
+		[[nodiscard]] UE_API FBoundValueMap* FindOrAddSetMap(const FAttributeMappingKey& MappingKey);
 
 		// Iterators
 		using FSetMapIterator = TAttributeCollectionSetMapIterator<FAttributeCollection>;
@@ -228,23 +228,23 @@ namespace UE::UAF
 
 		// Constructs and adds a new name/value map to this collection.
 		// Returns nullptr if we fail to add an entry (e.g. duplicate mapping)
-		template <class ValueType> [[nodiscard]] TAttributeMap<ValueType>* AddMap();
+		template <class ValueType> [[nodiscard]] TUnboundValueMap<ValueType>* AddMap();
 
 		// Constructs and adds a new name/value mapping to this collection.
 		// Returns nullptr if we fail to add an entry (e.g. duplicate mapping)
-		[[nodiscard]] UE_API FAttributeMap* AddMap(UScriptStruct* ValueType);
+		[[nodiscard]] UE_API FUnboundValueMap* AddMap(UScriptStruct* ValueType);
 
 		// Adds a new name/value mapping to this collection.
 		// Returns false if we fail to add an entry (e.g. duplicate mapping, mismatched allocator)
-		UE_API bool AddMap(FAttributeMap* Map);
+		UE_API bool AddMap(FUnboundValueMap* Map);
 
 		// Appends a new name/value mapping to this collection (the input must belong at the end of the sorted maps).
 		// Returns false if we fail to add an entry (e.g. duplicate mapping, mismatched allocator, not sorted)
-		UE_API bool AppendMap(FAttributeMap* Map);
+		UE_API bool AppendMap(FUnboundValueMap* Map);
 
 		// Removes the specified name/value mapping from this collection
 		// Returns whether or not the entry was successfully removed
-		UE_API bool RemoveMap(FAttributeMap* Map);
+		UE_API bool RemoveMap(FUnboundValueMap* Map);
 
 		// Removes the specified name/value mapping from this collection
 		// Returns whether or not the entry was successfully removed
@@ -252,27 +252,27 @@ namespace UE::UAF
 
 		// Finds the name/value mapping of the specified ValueType in this collection
 		// Returns nullptr if this collection does not contain such a mapping
-		template <class ValueType> [[nodiscard]] TAttributeMap<ValueType>* FindMap();
+		template <class ValueType> [[nodiscard]] TUnboundValueMap<ValueType>* FindMap();
 
 		// Finds the name/value mapping of the specified ValueType in this collection
 		// Returns nullptr if this collection does not contain such a mapping
-		template <class ValueType> [[nodiscard]] const TAttributeMap<ValueType>* FindMap() const;
+		template <class ValueType> [[nodiscard]] const TUnboundValueMap<ValueType>* FindMap() const;
 
 		// Finds the name/value mapping of the specified ValueType in this collection
 		// Returns nullptr if this collection does not contain such a mapping
-		[[nodiscard]] UE_API FAttributeMap* FindMap(UScriptStruct* ValueType);
+		[[nodiscard]] UE_API FUnboundValueMap* FindMap(UScriptStruct* ValueType);
 
 		// Finds the name/value mapping of the specified ValueType in this collection
 		// Returns nullptr if this collection does not contain such a mapping
-		[[nodiscard]] UE_API const FAttributeMap* FindMap(UScriptStruct* ValueType) const;
+		[[nodiscard]] UE_API const FUnboundValueMap* FindMap(UScriptStruct* ValueType) const;
 
 		// Finds or adds the name/value mapping of the specified ValueType in this collection
 		// Returns nullptr if we fail to add an entry (e.g. duplicate mapping)
-		template <class ValueType> [[nodiscard]] TAttributeMap<ValueType>* FindOrAddMap();
+		template <class ValueType> [[nodiscard]] TUnboundValueMap<ValueType>* FindOrAddMap();
 
 		// Finds or adds the name/value mapping of the specified ValueType in this collection
 		// Returns nullptr if we fail to add an entry (e.g. duplicate mapping)
-		[[nodiscard]] UE_API FAttributeMap* FindOrAddMap(UScriptStruct* ValueType);
+		[[nodiscard]] UE_API FUnboundValueMap* FindOrAddMap(UScriptStruct* ValueType);
 
 		// Iterators
 		using FMapIterator = TAttributeCollectionMapIterator<FAttributeCollection>;
@@ -283,8 +283,8 @@ namespace UE::UAF
 		[[nodiscard]] FConstMapIterator CreateConstMapIterator() const;
 
 	private:
-		UE_API FAttributeSetMap** AddSetMapImpl(const FAttributeSetMap::FConstructArgs& Args, bool bConstructSetMap);
-		UE_API FAttributeMap** AddMapImpl(UScriptStruct* ValueType, bool bConstructMap);
+		UE_API FBoundValueMap** AddSetMapImpl(const FBoundValueMap::FConstructArgs& Args, bool bConstructSetMap);
+		UE_API FUnboundValueMap** AddMapImpl(UScriptStruct* ValueType, bool bConstructMap);
 
 		// The named set we are based on, every set mapping uses a typed set within it
 		FAttributeNamedSetPtr NamedSet;
@@ -355,8 +355,8 @@ namespace UE::UAF
 		const FAttributeMappingKey& GetKey(const FAttributeCollection::FSetMapIterator& It) const { return It.GetKey(); }
 		const FAttributeMappingKey& GetKey(const FAttributeCollection::FConstSetMapIterator& It) const { return It.GetKey(); }
 
-		FAttributeSetMap* GetValue(const FAttributeCollection::FSetMapIterator& It) const { return It.GetMap(); }
-		const FAttributeSetMap* GetValue(const FAttributeCollection::FConstSetMapIterator& It) const { return It.GetMap(); }
+		FBoundValueMap* GetValue(const FAttributeCollection::FSetMapIterator& It) const { return It.GetMap(); }
+		const FBoundValueMap* GetValue(const FAttributeCollection::FConstSetMapIterator& It) const { return It.GetMap(); }
 	};
 
 	// An operator structure to join attribute collections using InnerJoin/OuterJoin/etc
@@ -370,8 +370,8 @@ namespace UE::UAF
 		UScriptStruct* GetKey(const FAttributeCollection::FMapIterator& It) const { return It.GetValueType(); }
 		UScriptStruct* GetKey(const FAttributeCollection::FConstMapIterator& It) const { return It.GetValueType(); }
 
-		FAttributeMap* GetValue(const FAttributeCollection::FMapIterator& It) const { return It.GetMap(); }
-		const FAttributeMap* GetValue(const FAttributeCollection::FConstMapIterator& It) const { return It.GetMap(); }
+		FUnboundValueMap* GetValue(const FAttributeCollection::FMapIterator& It) const { return It.GetMap(); }
+		const FUnboundValueMap* GetValue(const FAttributeCollection::FConstMapIterator& It) const { return It.GetMap(); }
 	};
 }
 
@@ -399,7 +399,7 @@ struct FUAFAttributeCollection
 namespace UE::UAF
 {
 	template <class ValueType, typename... TArgs>
-	inline TAttributeSetMap<ValueType>* FAttributeCollection::AddSetMap(const FAttributeMappingKey& MappingKey, TArgs... Args)
+	inline TBoundValueMap<ValueType>* FAttributeCollection::AddSetMap(const FAttributeMappingKey& MappingKey, TArgs... Args)
 	{
 		if (!NamedSet)
 		{
@@ -421,18 +421,18 @@ namespace UE::UAF
 		}
 
 		// Forward the extra arguments
-		typename TAttributeSetMap<ValueType>::FConstructArgs ConstructArgs(TypedSet, ValueTypeStruct, ReallocFun, Args...);
+		typename TBoundValueMap<ValueType>::FConstructArgs ConstructArgs(TypedSet, ValueTypeStruct, ReallocFun, Args...);
 
 		// Don't construct the set map internally, it'll have to use indirection to do so and here we know the concrete type
 		constexpr bool bConstructSetMap = false;
 
-		FAttributeSetMap** SetMapPtr = AddSetMapImpl(ConstructArgs, bConstructSetMap);
+		FBoundValueMap** SetMapPtr = AddSetMapImpl(ConstructArgs, bConstructSetMap);
 		if (SetMapPtr == nullptr)
 		{
 			return nullptr;
 		}
 
-		TAttributeSetMap<ValueType>* SetMap = MakeAttributeSetMap<ValueType>(ConstructArgs);
+		TBoundValueMap<ValueType>* SetMap = MakeBoundValueMap<ValueType>(ConstructArgs);
 
 		check(*SetMapPtr == nullptr);
 		*SetMapPtr = SetMap;
@@ -441,7 +441,7 @@ namespace UE::UAF
 	}
 
 	template <class ValueType>
-	inline TAttributeSetMap<ValueType>* FAttributeCollection::FindSetMap(const FAttributeMappingKey& MappingKey)
+	inline TBoundValueMap<ValueType>* FAttributeCollection::FindSetMap(const FAttributeMappingKey& MappingKey)
 	{
 		UScriptStruct* ValueTypeStruct = ValueType::StaticStruct();
 		if (!ensureMsgf(ValueTypeStruct == MappingKey.GetValueType(), TEXT("Mismatch between the return value type and mapping key value type")))
@@ -453,7 +453,7 @@ namespace UE::UAF
 	}
 
 	template <class ValueType>
-	inline const TAttributeSetMap<ValueType>* FAttributeCollection::FindSetMap(const FAttributeMappingKey& MappingKey) const
+	inline const TBoundValueMap<ValueType>* FAttributeCollection::FindSetMap(const FAttributeMappingKey& MappingKey) const
 	{
 		UScriptStruct* ValueTypeStruct = ValueType::StaticStruct();
 		if (!ensureMsgf(ValueTypeStruct == MappingKey.GetValueType(), TEXT("Mismatch between the return value type and mapping key value type")))
@@ -465,7 +465,7 @@ namespace UE::UAF
 	}
 
 	template <class ValueType, typename... TArgs>
-	inline TAttributeSetMap<ValueType>* FAttributeCollection::FindOrAddSetMap(const FAttributeMappingKey& MappingKey, TArgs... Args)
+	inline TBoundValueMap<ValueType>* FAttributeCollection::FindOrAddSetMap(const FAttributeMappingKey& MappingKey, TArgs... Args)
 	{
 		UScriptStruct* ValueTypeStruct = ValueType::StaticStruct();
 		if (!ensureMsgf(ValueTypeStruct == MappingKey.GetValueType(), TEXT("Mismatch between the return value type and mapping key value type")))
@@ -473,9 +473,9 @@ namespace UE::UAF
 			return nullptr;
 		}
 
-		if (FAttributeSetMap* SetMap = FindSetMap(MappingKey))
+		if (FBoundValueMap* Map = FindSetMap(MappingKey))
 		{
-			return Cast<ValueType>(SetMap);
+			return Cast<ValueType>(Map);
 		}
 
 		return AddSetMap<ValueType>(MappingKey, Args...);
@@ -492,7 +492,7 @@ namespace UE::UAF
 	}
 
 	template <class ValueType>
-	inline TAttributeMap<ValueType>* FAttributeCollection::AddMap()
+	inline TUnboundValueMap<ValueType>* FAttributeCollection::AddMap()
 	{
 		if (!NamedSet)
 		{
@@ -505,13 +505,13 @@ namespace UE::UAF
 		// Don't construct the map internally, it'll have to use indirection to do so and here we know the concrete type
 		constexpr bool bConstructMap = false;
 
-		FAttributeMap** MapPtr = AddMapImpl(ValueTypeStruct, bConstructMap);
+		FUnboundValueMap** MapPtr = AddMapImpl(ValueTypeStruct, bConstructMap);
 		if (MapPtr == nullptr)
 		{
 			return nullptr;
 		}
 
-		TAttributeMap<ValueType>* Map = MakeAttributeMap<ValueType>(ReallocFun);
+		TUnboundValueMap<ValueType>* Map = MakeUnboundValueMap<ValueType>(ReallocFun);
 
 		check(*MapPtr == nullptr);
 		*MapPtr = Map;
@@ -520,21 +520,21 @@ namespace UE::UAF
 	}
 
 	template <class ValueType>
-	inline TAttributeMap<ValueType>* FAttributeCollection::FindMap()
+	inline TUnboundValueMap<ValueType>* FAttributeCollection::FindMap()
 	{
 		return Cast<ValueType>(FindMap(ValueType::StaticStruct()));
 	}
 
 	template <class ValueType>
-	inline const TAttributeMap<ValueType>* FAttributeCollection::FindMap() const
+	inline const TUnboundValueMap<ValueType>* FAttributeCollection::FindMap() const
 	{
 		return Cast<ValueType>(FindMap(ValueType::StaticStruct()));
 	}
 
 	template <class ValueType>
-	inline TAttributeMap<ValueType>* FAttributeCollection::FindOrAddMap()
+	inline TUnboundValueMap<ValueType>* FAttributeCollection::FindOrAddMap()
 	{
-		if (TAttributeMap<ValueType>* Map = FindMap<ValueType>())
+		if (TUnboundValueMap<ValueType>* Map = FindMap<ValueType>())
 		{
 			return Map;
 		}
@@ -611,7 +611,7 @@ namespace UE::UAF
 	}
 
 	template<class ContainerType>
-	inline typename TAttributeCollectionSetMapIterator<ContainerType>::SetMapType TAttributeCollectionSetMapIterator<ContainerType>::GetMap() const
+	inline typename TAttributeCollectionSetMapIterator<ContainerType>::MapType TAttributeCollectionSetMapIterator<ContainerType>::GetMap() const
 	{
 		return EntryPtr->Map;
 	}
diff --git a/Engine/Plugins/Experimental/UAF/UAF/Source/UAF/Public/UAF/ValueRuntime/Transformers/Accumulate.h b/Engine/Plugins/Experimental/UAF/UAF/Source/UAF/Public/UAF/ValueRuntime/Transformers/Accumulate.h
--- a/Engine/Plugins/Experimental/UAF/UAF/Source/UAF/Public/UAF/ValueRuntime/Transformers/Accumulate.h
+++ b/Engine/Plugins/Experimental/UAF/UAF/Source/UAF/Public/UAF/ValueRuntime/Transformers/Accumulate.h
@@ -14,26 +14,26 @@ namespace UE::UAF::Transformers
 	// Attribute value transformer to accumulate values
 	// Accumulate is the repeating step within an N-step interpolation: Overwrite, Accumulate0, Accumulate1, ...
 	// Make sure to also implement variations for Overwrite and Interpolate
-	struct FAccumulate final : public FAttributeValueTransformer
+	struct FAccumulate final : public FValueTransformer
 	{
 		static UE_API FName TransformerName;
 
 		// If the output buffer is nullptr, an allocator is provided and the implementation is responsible for allocating the output buffer and returning it
 		// Per-value map is optional, when not present, the default weight should be used
 		// Unbound values do not support per-value weights
-		using FTransformSetMapFunc = void (*)(const FAttributeSetMap* SetMapA, const FAttributeSetMap* SetMapB, const FAttributeSetMap* PerValueWeightMapB, float DefaultWeightB, FAttributeSetMap*& OutputSetMap, FReallocFun OutputAllocator);
-		using FTransformMapFunc = void (*)(const FAttributeMap* MapA, const FAttributeMap* MapB, float WeightB, FAttributeMap*& OutputMap, FReallocFun OutputAllocator);
+		using FTransformBoundValueMapFunc = void (*)(const FBoundValueMap* MapA, const FBoundValueMap* MapB, const FBoundValueMap* PerValueWeightMapB, float DefaultWeightB, FBoundValueMap*& OutputMap, FReallocFun OutputAllocator);
+		using FTransformUnboundValueMapFunc = void (*)(const FUnboundValueMap* MapA, const FUnboundValueMap* MapB, float WeightB, FUnboundValueMap*& OutputMap, FReallocFun OutputAllocator);
 
 		// Accumulate attribute values
 		// Output = A + (B * WeightB)
 		// This function supports outputting to either input
-		static UE_API void Apply(const FAttributeValueTransformerMapPtr& TransformerMap, const FAttributeCollection& InputA, const FAttributeCollection& InputB, float WeightB, FAttributeCollection& Output);
+		static UE_API void Apply(const FValueTransformerMapPtr& TransformerMap, const FAttributeCollection& InputA, const FAttributeCollection& InputB, float WeightB, FAttributeCollection& Output);
 
 		// Accumulate attribute values
 		// If a weight isn't specified, the default weight will be used
 		// Output = A + (B * (HasValueWeightB ? ValueWeightB : DefaultWeightB))
 		// This function supports outputting to either input
-		static UE_API void Apply(const FAttributeValueTransformerMapPtr& TransformerMap, const FAttributeCollection& InputA, const FAttributeCollection& InputB, const FAttributeCollection& PerValueWeightsB, float DefaultWeightB, FAttributeCollection& Output);
+		static UE_API void Apply(const FValueTransformerMapPtr& TransformerMap, const FAttributeCollection& InputA, const FAttributeCollection& InputB, const FAttributeCollection& PerValueWeightsB, float DefaultWeightB, FAttributeCollection& Output);
 	};
 
 	//////////////////////////////////////////////////////////////////////////
@@ -46,8 +46,8 @@ namespace UE::UAF::Transformers
 			using TransformerType = FAccumulate;
 			using ValueType = FBoneTransformAnimationAttribute;
 
-			static void TransformSetMap(const FAttributeSetMap* SetMapA, const FAttributeSetMap* SetMapB, const FAttributeSetMap* PerValueWeightMapB, float DefaultWeightB, FAttributeSetMap*& OutputSetMap, FReallocFun OutputAllocator);
-			static void TransformMap(const FAttributeMap* MapA, const FAttributeMap* MapB, float WeightB, FAttributeMap*& OutputMap, FReallocFun OutputAllocator);
+			static void TransformBoundValueMap(const FBoundValueMap* MapA, const FBoundValueMap* MapB, const FBoundValueMap* PerValueWeightMapB, float DefaultWeightB, FBoundValueMap*& OutputMap, FReallocFun OutputAllocator);
+			static void TransformUnboundValueMap(const FUnboundValueMap* MapA, const FUnboundValueMap* MapB, float WeightB, FUnboundValueMap*& OutputMap, FReallocFun OutputAllocator);
 		};
 
 		struct FAccumulate_FloatAttribute
@@ -55,8 +55,8 @@ namespace UE::UAF::Transformers
 			using TransformerType = FAccumulate;
 			using ValueType = FFloatAnimationAttribute;
 
-			static void TransformSetMap(const FAttributeSetMap* SetMapA, const FAttributeSetMap* SetMapB, const FAttributeSetMap* PerValueWeightMapB, float DefaultWeightB, FAttributeSetMap*& OutputSetMap, FReallocFun OutputAllocator);
-			static void TransformMap(const FAttributeMap* MapA, const FAttributeMap* MapB, float WeightB, FAttributeMap*& OutputMap, FReallocFun OutputAllocator);
+			static void TransformBoundValueMap(const FBoundValueMap* MapA, const FBoundValueMap* MapB, const FBoundValueMap* PerValueWeightMapB, float DefaultWeightB, FBoundValueMap*& OutputMap, FReallocFun OutputAllocator);
+			static void TransformUnboundValueMap(const FUnboundValueMap* MapA, const FUnboundValueMap* MapB, float WeightB, FUnboundValueMap*& OutputMap, FReallocFun OutputAllocator);
 		};
 	}
 }
