Rename UAF attribute collection and split into 3 types

#Synced-CL 49329447

diff --git a/Engine/Plugins/Experimental/MovieSceneAnimMixer/Source/MovieSceneAnimMixer/Private/Systems/MovieSceneAnimMixerSystem.cpp b/Engine/Plugins/Experimental/MovieSceneAnimMixer/Source/MovieSceneAnimMixer/Private/Systems/MovieSceneAnimMixerSystem.cpp
--- a/Engine/Plugins/Experimental/MovieSceneAnimMixer/Source/MovieSceneAnimMixer/Private/Systems/MovieSceneAnimMixerSystem.cpp
+++ b/Engine/Plugins/Experimental/MovieSceneAnimMixer/Source/MovieSceneAnimMixer/Private/Systems/MovieSceneAnimMixerSystem.cpp
@@ -353,7 +353,7 @@ void FAnimNextBlendTwoKeyframesPreserveRootMotionTask::Execute(UE::UAF::FEvaluat
 	if (VM.GetActiveNamedSet())
 	{
 		// TODO: Implement with new attribute runtime
-		TUniquePtr<FAttributeCollection> DiscardedCollection;
+		TUniquePtr<FValueBundle> DiscardedCollection;
 		(void)VM.PopValue(ATTRIBUTE_STACK_NAME, DiscardedCollection);
 		return;
 	}
@@ -488,7 +488,7 @@ void FMovieSceneAccumulateAbsoluteBlendTask::Execute(UE::UAF::FEvaluationVM& VM)
 	if (VM.GetActiveNamedSet())
 	{
 		// TODO: Implement with new attribute runtime
-		TUniquePtr<FAttributeCollection> DiscardedCollection;
+		TUniquePtr<FValueBundle> DiscardedCollection;
 		(void)VM.PopValue(ATTRIBUTE_STACK_NAME, DiscardedCollection);
 		return;
 
diff --git a/Engine/Plugins/Experimental/MovieSceneAnimMixer/Source/MovieSceneAnimMixer/Private/Systems/MovieSceneMixedSkeletalAnimationSystem.cpp b/Engine/Plugins/Experimental/MovieSceneAnimMixer/Source/MovieSceneAnimMixer/Private/Systems/MovieSceneMixedSkeletalAnimationSystem.cpp
--- a/Engine/Plugins/Experimental/MovieSceneAnimMixer/Source/MovieSceneAnimMixer/Private/Systems/MovieSceneMixedSkeletalAnimationSystem.cpp
+++ b/Engine/Plugins/Experimental/MovieSceneAnimMixer/Source/MovieSceneAnimMixer/Private/Systems/MovieSceneMixedSkeletalAnimationSystem.cpp
@@ -406,10 +406,10 @@ void FMovieSceneSkeletalAnimationEvaluationTask::Execute(UE::UAF::FEvaluationVM&
 	if (VM.GetActiveNamedSet())
 	{
 		// TODO: Implement with new attribute runtime
-		FAttributeCollectionStack Collection(VM.GetActiveNamedSet());
+		FValueBundleStack Collection(VM.GetActiveNamedSet());
 		Collection.InitWithValueSpace(FValueSpace(EValueSpaceType::Local));
 
-		VM.PushValue(ATTRIBUTE_STACK_NAME, MakeUnique<FAttributeCollection>(MoveTemp(Collection)));
+		VM.PushValue(ATTRIBUTE_STACK_NAME, MakeUnique<FValueBundle>(MoveTemp(Collection)));
 		return;
 	}
 
diff --git a/Engine/Plugins/Experimental/MovieScenePoseSearchTracks/Source/MovieScenePoseSearchTracks/Private/Sections/MovieSceneStitchAnimSection.cpp b/Engine/Plugins/Experimental/MovieScenePoseSearchTracks/Source/MovieScenePoseSearchTracks/Private/Sections/MovieSceneStitchAnimSection.cpp
--- a/Engine/Plugins/Experimental/MovieScenePoseSearchTracks/Source/MovieScenePoseSearchTracks/Private/Sections/MovieSceneStitchAnimSection.cpp
+++ b/Engine/Plugins/Experimental/MovieScenePoseSearchTracks/Source/MovieScenePoseSearchTracks/Private/Sections/MovieSceneStitchAnimSection.cpp
@@ -62,10 +62,10 @@ void FMovieSceneStitchAnimEvaluationTask::Execute(UE::UAF::FEvaluationVM& VM) co
 	if (VM.GetActiveNamedSet())
 	{
 		// TODO: Implement with new attribute runtime
-		FAttributeCollectionStack Collection(VM.GetActiveNamedSet());
+		FValueBundleStack Collection(VM.GetActiveNamedSet());
 		Collection.InitWithValueSpace(FValueSpace(EValueSpaceType::Local));
 
-		VM.PushValue(ATTRIBUTE_STACK_NAME, MakeUnique<FAttributeCollection>(MoveTemp(Collection)));
+		VM.PushValue(ATTRIBUTE_STACK_NAME, MakeUnique<FValueBundle>(MoveTemp(Collection)));
 		return;
 	}
 
diff --git a/Engine/Plugins/Experimental/UAF/UAF/Source/UAF/Private/UAF/Attributes/AttributeBindingData.cpp b/Engine/Plugins/Experimental/UAF/UAF/Source/UAF/Private/UAF/Attributes/AttributeBindingData.cpp
--- a/Engine/Plugins/Experimental/UAF/UAF/Source/UAF/Private/UAF/Attributes/AttributeBindingData.cpp
+++ b/Engine/Plugins/Experimental/UAF/UAF/Source/UAF/Private/UAF/Attributes/AttributeBindingData.cpp
@@ -336,7 +336,7 @@ namespace UE::UAF
 				FAttributeNamedSetPtr NamedSet = *NamedSetIt;
 				const FName SetName = NamedSet->GetName();
 
-				FAttributeCollectionPtr DefaultAttributeValues = MakeShared<FAttributeCollectionHeap>(NamedSet);
+				FAttributeCollectionPtr DefaultAttributeValues = MakeShared<FValueBundleHeap>(NamedSet);
 				DefaultAttributeValues->SetValueSpace(FValueSpace(EValueSpaceType::Local));
 
 				for (auto TypedSetIt = NamedSet->CreateTypedSetIterator(); TypedSetIt; ++TypedSetIt)
@@ -348,7 +348,7 @@ namespace UE::UAF
 					}
 
 					FAttributeMappingKey MappingKey = FAttributeMappingKey::MakeFrom(TypedSetIt->GetType()).To<FBoneTransformAnimationAttribute>();
-					TBoundValueMap<FBoneTransformAnimationAttribute>* DefaultValues = DefaultAttributeValues->AddSetMap<FBoneTransformAnimationAttribute>(MappingKey);
+					TBoundValueMap<FBoneTransformAnimationAttribute>* DefaultValues = DefaultAttributeValues->GetBoundValueMaps().Add<FBoneTransformAnimationAttribute>(MappingKey);
 
 					// Set every value within our map
 					const int32 NumValues = DefaultValues->Num();
diff --git a/Engine/Plugins/Experimental/UAF/UAF/Source/UAF/Private/UAF/ValueRuntime/Transformers/Accumulate.cpp b/Engine/Plugins/Experimental/UAF/UAF/Source/UAF/Private/UAF/ValueRuntime/Transformers/Accumulate.cpp
--- a/Engine/Plugins/Experimental/UAF/UAF/Source/UAF/Private/UAF/ValueRuntime/Transformers/Accumulate.cpp
+++ b/Engine/Plugins/Experimental/UAF/UAF/Source/UAF/Private/UAF/ValueRuntime/Transformers/Accumulate.cpp
@@ -8,7 +8,7 @@ namespace UE::UAF::Transformers
 {
 	FName FAccumulate::TransformerName = TEXT("UAF::Transformers::Accumulate");
 
-	void FAccumulate::Apply(const FValueTransformerMapPtr& TransformerMap, const FAttributeCollection& InputA, const FAttributeCollection& InputB, float WeightB, FAttributeCollection& Output)
+	void FAccumulate::Apply(const FValueTransformerMapPtr& TransformerMap, const FValueBundle& InputA, const FValueBundle& InputB, float WeightB, FValueBundle& Output)
 	{
 		// Named sets must match as it ensures that our inputs have the same sizes/shapes
 		checkf(InputA.GetNamedSet() == InputB.GetNamedSet(), TEXT("Both input collections must use the same named set"));
@@ -22,11 +22,11 @@ namespace UE::UAF::Transformers
 		}
 
 		// Allocator is irrelevant since it will remain empty
-		FAttributeCollectionHeap EmptyCollection(InputA.GetNamedSet());
+		FValueBundleHeap EmptyCollection(InputA.GetNamedSet());
 		FAccumulate::Apply(TransformerMap, InputA, InputB, EmptyCollection, WeightB, Output);
 	}
 
-	void FAccumulate::Apply(const FValueTransformerMapPtr& TransformerMap, const FAttributeCollection& InputA, const FAttributeCollection& InputB, const FAttributeCollection& PerValueWeightsB, float DefaultWeightB, FAttributeCollection& Output)
+	void FAccumulate::Apply(const FValueTransformerMapPtr& TransformerMap, const FValueBundle& InputA, const FValueBundle& InputB, const FValueBundle& PerValueWeightsB, float DefaultWeightB, FValueBundle& Output)
 	{
 		// Named sets must match as it ensures that our inputs have the same sizes/shapes
 		checkf(InputA.GetNamedSet() == InputB.GetNamedSet(), TEXT("Both input collections must use the same named set"));
@@ -51,12 +51,14 @@ namespace UE::UAF::Transformers
 
 			// Bound value maps
 			{
+				const FBoundMapCollection& PerValueWeightBMaps = PerValueWeightsB.GetBoundValueMaps();
+
 				// Keep a list of bound value maps that needs to be appended to our output
 				TArray<FBoundValueMap*, TInlineAllocator<8>> PendingMaps;
 
 				OuterJoinBy(
 					FValueTransformerListWithBoundMapCollectionJoinOp(),
-					[&PerValueWeightsB, DefaultWeightB, bIsOutputA, FloatValueType, OutputAllocator, &PendingMaps](FRawTransformerFunc Transformer, FBoundValueMap* MapA, FBoundValueMap* MapB)
+					[&PerValueWeightBMaps, DefaultWeightB, bIsOutputA, FloatValueType, OutputAllocator, &PendingMaps](FRawTransformerFunc Transformer, FBoundValueMap* MapA, FBoundValueMap* MapB)
 					{
 						if (MapA == nullptr)
 						{
@@ -132,7 +134,7 @@ namespace UE::UAF::Transformers
 
 									// We cannot join our weights because their value type differs
 									const FAttributeMappingKey MappingKey = MapA->GetMappingKey().To(FloatValueType);
-									const FBoundValueMap* PerValueWeightMapB = PerValueWeightsB.FindSetMap(MappingKey);
+									const FBoundValueMap* PerValueWeightMapB = PerValueWeightBMaps.Find(MappingKey);
 
 									FTransformBoundValueMapFunc TypedTransformer = reinterpret_cast<FTransformBoundValueMapFunc>(Transformer);
 									(*TypedTransformer)(MapA, MapB, PerValueWeightMapB, DefaultWeightB, OutputMap, nullptr);
@@ -142,13 +144,14 @@ namespace UE::UAF::Transformers
 					},
 					TransformerList->CreateBoundValueMapTransformerIterator(),
 					// Remove the const qualifier since we'll mutate one in place as our output
-					const_cast<FAttributeCollection&>(InputA).CreateSetMapIterator(),
-					const_cast<FAttributeCollection&>(InputB).CreateSetMapIterator());
+					const_cast<FValueBundle&>(InputA).GetBoundValueMaps().CreateIterator(),
+					const_cast<FValueBundle&>(InputB).GetBoundValueMaps().CreateIterator());
 
 				// Append our pending entries
+				FBoundMapCollection& OutputMaps = Output.GetBoundValueMaps();
 				for (FBoundValueMap* Map : PendingMaps)
 				{
-					Output.AddSetMap(Map);
+					OutputMaps.Add(Map);
 				}
 			}
 
@@ -250,13 +253,14 @@ namespace UE::UAF::Transformers
 					},
 					TransformerList->CreateUnboundValueMapTransformerIterator(),
 					// Remove the const qualifier since we'll mutate one in place as our output
-					const_cast<FAttributeCollection&>(InputA).CreateMapIterator(),
-					const_cast<FAttributeCollection&>(InputB).CreateMapIterator());
+					const_cast<FValueBundle&>(InputA).GetUnboundValueMaps().CreateIterator(),
+					const_cast<FValueBundle&>(InputB).GetUnboundValueMaps().CreateIterator());
 
 				// Append our pending entries
+				FUnboundMapCollection& OutputMaps = Output.GetUnboundValueMaps();
 				for (FUnboundValueMap* Map : PendingMaps)
 				{
-					Output.AddMap(Map);
+					OutputMaps.Add(Map);
 				}
 			}
 		}
@@ -267,128 +271,137 @@ namespace UE::UAF::Transformers
 			Output.SetValueSpace(InputA.GetValueSpace());
 
 			// Bound value maps
-			OuterJoinBy(
-				FValueTransformerListWithBoundMapCollectionJoinOp(),
-				[&PerValueWeightsB, DefaultWeightB, FloatValueType, &Output](FRawTransformerFunc Transformer, const FBoundValueMap* MapA, const FBoundValueMap* MapB)
-				{
-					if (MapA == nullptr)
-					{
-						if (MapB == nullptr)
-						{
-							// No maps found for the transformer type
-						}
-						else
-						{
-							// This set map doesn't have a matching pair, duplicate it
-							Output.AppendSetMap(MapB->Duplicate(Output.GetAllocator()));
-						}
-					}
-					else
+			{
+				const FBoundMapCollection& PerValueWeightBMaps = PerValueWeightsB.GetBoundValueMaps();
+				FBoundMapCollection& OutputMaps = Output.GetBoundValueMaps();
+
+				OuterJoinBy(
+					FValueTransformerListWithBoundMapCollectionJoinOp(),
+					[&PerValueWeightBMaps, DefaultWeightB, FloatValueType, &OutputMaps](FRawTransformerFunc Transformer, const FBoundValueMap* MapA, const FBoundValueMap* MapB)
 					{
-						if (MapB == nullptr)
+						if (MapA == nullptr)
 						{
-							// This set map doesn't have a matching pair, duplicate it
-							Output.AppendSetMap(MapA->Duplicate(Output.GetAllocator()));
+							if (MapB == nullptr)
+							{
+								// No maps found for the transformer type
+							}
+							else
+							{
+								// This set map doesn't have a matching pair, duplicate it
+								OutputMaps.Append(MapB->Duplicate(OutputMaps.GetAllocator()));
+							}
 						}
 						else
 						{
-							if (Transformer == nullptr)
+							if (MapB == nullptr)
+							{
+								// This set map doesn't have a matching pair, duplicate it
+								OutputMaps.Append(MapA->Duplicate(OutputMaps.GetAllocator()));
+							}
+							else
 							{
-								// No transformer was found for this value type, retain the one with the largest weight
-								FBoundValueMap* MapCopy;
-								if (DefaultWeightB >= 0.5f)
+								if (Transformer == nullptr)
 								{
-									MapCopy = MapB->Duplicate(Output.GetAllocator());
+									// No transformer was found for this value type, retain the one with the largest weight
+									FBoundValueMap* MapCopy;
+									if (DefaultWeightB >= 0.5f)
+									{
+										MapCopy = MapB->Duplicate(OutputMaps.GetAllocator());
+									}
+									else
+									{
+										MapCopy = MapA->Duplicate(OutputMaps.GetAllocator());
+									}
+
+									OutputMaps.Append(MapCopy);
 								}
 								else
 								{
-									MapCopy = MapA->Duplicate(Output.GetAllocator());
-								}
-
-								Output.AppendSetMap(MapCopy);
-							}
-							else
-							{
-								// Transform our matching pair
-								FBoundValueMap* OutputMap = nullptr;
+									// Transform our matching pair
+									FBoundValueMap* OutputMap = nullptr;
 
-								// We cannot join our weights because their value type differs
-								const FAttributeMappingKey MappingKey = MapA->GetMappingKey().To(FloatValueType);
-								const FBoundValueMap* PerValueWeightMapB = PerValueWeightsB.FindSetMap(MappingKey);
+									// We cannot join our weights because their value type differs
+									const FAttributeMappingKey MappingKey = MapA->GetMappingKey().To(FloatValueType);
+									const FBoundValueMap* PerValueWeightMapB = PerValueWeightBMaps.Find(MappingKey);
 
-								FTransformBoundValueMapFunc TypedTransformer = reinterpret_cast<FTransformBoundValueMapFunc>(Transformer);
-								(*TypedTransformer)(MapA, MapB, PerValueWeightMapB, DefaultWeightB, OutputMap, Output.GetAllocator());
+									FTransformBoundValueMapFunc TypedTransformer = reinterpret_cast<FTransformBoundValueMapFunc>(Transformer);
+									(*TypedTransformer)(MapA, MapB, PerValueWeightMapB, DefaultWeightB, OutputMap, OutputMaps.GetAllocator());
 
-								checkf(OutputMap != nullptr, TEXT("Transformer function is responsible for allocating an output buffer"));
-								checkf(OutputMap->GetAllocator() == Output.GetAllocator(), TEXT("Transformer function is responsible for allocating using the output allocator"));
-								Output.AppendSetMap(OutputMap);
+									checkf(OutputMap != nullptr, TEXT("Transformer function is responsible for allocating an output buffer"));
+									checkf(OutputMap->GetAllocator() == OutputMaps.GetAllocator(), TEXT("Transformer function is responsible for allocating using the output allocator"));
+									OutputMaps.Append(OutputMap);
+								}
 							}
 						}
-					}
-				},
-				TransformerList->CreateBoundValueMapTransformerIterator(),
-				InputA.CreateConstSetMapIterator(),
-				InputB.CreateConstSetMapIterator());
+					},
+					TransformerList->CreateBoundValueMapTransformerIterator(),
+					InputA.GetBoundValueMaps().CreateConstIterator(),
+					InputB.GetBoundValueMaps().CreateConstIterator());
+			}
 
 			// Name/value maps
-			OuterJoinBy(
-				FValueTransformerListWithUnboundMapCollectionJoinOp(),
-				[DefaultWeightB, &Output](FRawTransformerFunc Transformer, const FUnboundValueMap* MapA, const FUnboundValueMap* MapB)
-				{
-					if (MapA == nullptr)
-					{
-						if (MapB == nullptr)
-						{
-							// No maps found for the transformer type
-						}
-						else
-						{
-							// This map doesn't have a matching pair, duplicate it
-							Output.AppendMap(MapB->Duplicate(Output.GetAllocator()));
-						}
-					}
-					else
+			{
+				FUnboundMapCollection& OutputMaps = Output.GetUnboundValueMaps();
+
+				OuterJoinBy(
+					FValueTransformerListWithUnboundMapCollectionJoinOp(),
+					[DefaultWeightB, &OutputMaps](FRawTransformerFunc Transformer, const FUnboundValueMap* MapA, const FUnboundValueMap* MapB)
 					{
-						if (MapB == nullptr)
+						if (MapA == nullptr)
 						{
-							// This map doesn't have a matching pair
-							// Our output is A and thus already contains this
+							if (MapB == nullptr)
+							{
+								// No maps found for the transformer type
+							}
+							else
+							{
+								// This map doesn't have a matching pair, duplicate it
+								OutputMaps.Append(MapB->Duplicate(OutputMaps.GetAllocator()));
+							}
 						}
 						else
 						{
-							if (Transformer == nullptr)
+							if (MapB == nullptr)
+							{
+								// This map doesn't have a matching pair
+								// Our output is A and thus already contains this
+							}
+							else
 							{
-								// No transformer was found for this value type, retain the one with the largest weight
-								FUnboundValueMap* MapCopy;
-								if (DefaultWeightB >= 0.5f)
+								if (Transformer == nullptr)
 								{
-									MapCopy = MapB->Duplicate(Output.GetAllocator());
+									// No transformer was found for this value type, retain the one with the largest weight
+									FUnboundValueMap* MapCopy;
+									if (DefaultWeightB >= 0.5f)
+									{
+										MapCopy = MapB->Duplicate(OutputMaps.GetAllocator());
+									}
+									else
+									{
+										MapCopy = MapA->Duplicate(OutputMaps.GetAllocator());
+									}
+
+									OutputMaps.Append(MapCopy);
 								}
 								else
 								{
-									MapCopy = MapA->Duplicate(Output.GetAllocator());
-								}
-
-								Output.AppendMap(MapCopy);
-							}
-							else
-							{
-								// Transform our matching pair in place
-								FUnboundValueMap* OutputMap = nullptr;
+									// Transform our matching pair in place
+									FUnboundValueMap* OutputMap = nullptr;
 
-								FTransformUnboundValueMapFunc TypedTransformer = reinterpret_cast<FTransformUnboundValueMapFunc>(Transformer);
-								(*TypedTransformer)(MapA, MapB, DefaultWeightB, OutputMap, Output.GetAllocator());
+									FTransformUnboundValueMapFunc TypedTransformer = reinterpret_cast<FTransformUnboundValueMapFunc>(Transformer);
+									(*TypedTransformer)(MapA, MapB, DefaultWeightB, OutputMap, OutputMaps.GetAllocator());
 
-								checkf(OutputMap != nullptr, TEXT("Transformer function is responsible for allocating an output buffer"));
-								checkf(OutputMap->GetAllocator() == Output.GetAllocator(), TEXT("Transformer function is responsible for allocating using the output allocator"));
-								Output.AppendMap(OutputMap);
+									checkf(OutputMap != nullptr, TEXT("Transformer function is responsible for allocating an output buffer"));
+									checkf(OutputMap->GetAllocator() == OutputMaps.GetAllocator(), TEXT("Transformer function is responsible for allocating using the output allocator"));
+									OutputMaps.Append(OutputMap);
+								}
 							}
 						}
-					}
-				},
-				TransformerList->CreateUnboundValueMapTransformerIterator(),
-				InputA.CreateConstMapIterator(),
-				InputB.CreateConstMapIterator());
+					},
+					TransformerList->CreateUnboundValueMapTransformerIterator(),
+					InputA.GetUnboundValueMaps().CreateConstIterator(),
+					InputB.GetUnboundValueMaps().CreateConstIterator());
+			}
 		}
 	}
 
diff --git a/Engine/Plugins/Experimental/UAF/UAF/Source/UAF/Private/UAF/ValueRuntime/Transformers/AdditiveSpace.cpp b/Engine/Plugins/Experimental/UAF/UAF/Source/UAF/Private/UAF/ValueRuntime/Transformers/AdditiveSpace.cpp
--- a/Engine/Plugins/Experimental/UAF/UAF/Source/UAF/Private/UAF/ValueRuntime/Transformers/AdditiveSpace.cpp
+++ b/Engine/Plugins/Experimental/UAF/UAF/Source/UAF/Private/UAF/ValueRuntime/Transformers/AdditiveSpace.cpp
@@ -12,9 +12,9 @@ namespace UE::UAF::Transformers
 
 	void FMakeAdditiveSpace::Apply(
 		const FValueTransformerMapPtr& TransformerMap,
-		const FAttributeCollection& Base,
-		const FAttributeCollection& Input,
-		FAttributeCollection& AdditiveOutput)
+		const FValueBundle& Base,
+		const FValueBundle& Input,
+		FValueBundle& AdditiveOutput)
 	{
 		// Named sets must match as it ensures that our inputs have the same sizes/shapes
 		checkf(Base.GetNamedSet() == Input.GetNamedSet(), TEXT("Base and Input collections must use the same named set"));
@@ -40,79 +40,86 @@ namespace UE::UAF::Transformers
 		AdditiveOutput.SetValueSpace(OutputSpace);
 
 		// Bound value maps
-		OuterJoinBy(
-			FValueTransformerListWithBoundMapCollectionJoinOp(),
-			[&AdditiveOutput](FRawTransformerFunc Transformer, const FBoundValueMap* BaseMap, const FBoundValueMap* InputMap)
-			{
-				if (InputMap == nullptr)
-				{
-					// The base doesn't have a matching pair in our input, nothing to convert
-					return;
-				}
+		{
+			FBoundMapCollection& AdditiveOutputMaps = AdditiveOutput.GetBoundValueMaps();
 
-				if (Transformer == nullptr)
+			OuterJoinBy(
+				FValueTransformerListWithBoundMapCollectionJoinOp(),
+				[&AdditiveOutputMaps](FRawTransformerFunc Transformer, const FBoundValueMap* BaseMap, const FBoundValueMap* InputMap)
 				{
-					// No transformer was found for this value type, retain the input as-is
-					AdditiveOutput.AppendSetMap(InputMap->Duplicate(AdditiveOutput.GetAllocator()));
-					return;
-				}
+					if (InputMap == nullptr)
+					{
+						// The base doesn't have a matching pair in our input, nothing to convert
+						return;
+					}
 
-				// Transform our matching pair, note that the base map might be missing
-				// If it is missing, transformer specializations should use a default value
-				FBoundValueMap* OutputMap = nullptr;
-
-				FTransformBoundValueMapFunc TypedTransformer = reinterpret_cast<FTransformBoundValueMapFunc>(Transformer);
-				(*TypedTransformer)(BaseMap, InputMap, OutputMap, AdditiveOutput.GetAllocator());
-
-				checkf(OutputMap != nullptr, TEXT("Transformer function is responsible for allocating an output buffer"));
-				checkf(OutputMap->GetAllocator() == AdditiveOutput.GetAllocator(), TEXT("Transformer function is responsible for allocating using the output allocator"));
-				AdditiveOutput.AppendSetMap(OutputMap);
-			},
-			TransformerList->CreateBoundValueMapTransformerIterator(),
-			Base.CreateConstSetMapIterator(),
-			Input.CreateConstSetMapIterator());
-
-		// Name/value maps
-		OuterJoinBy(
-			FValueTransformerListWithUnboundMapCollectionJoinOp(),
-			[&AdditiveOutput](FRawTransformerFunc Transformer, const FUnboundValueMap* BaseMap, const FUnboundValueMap* InputMap)
-			{
-				if (InputMap == nullptr)
-				{
-					// The base doesn't have a matching pair in our input, nothing to convert
-					return;
-				}
+					if (Transformer == nullptr)
+					{
+						// No transformer was found for this value type, retain the input as-is
+						AdditiveOutputMaps.Append(InputMap->Duplicate(AdditiveOutputMaps.GetAllocator()));
+						return;
+					}
 
-				if (Transformer == nullptr)
+					// Transform our matching pair, note that the base map might be missing
+					// If it is missing, transformer specializations should use a default value
+					FBoundValueMap* OutputMap = nullptr;
+
+					FTransformBoundValueMapFunc TypedTransformer = reinterpret_cast<FTransformBoundValueMapFunc>(Transformer);
+					(*TypedTransformer)(BaseMap, InputMap, OutputMap, AdditiveOutputMaps.GetAllocator());
+
+					checkf(OutputMap != nullptr, TEXT("Transformer function is responsible for allocating an output buffer"));
+					checkf(OutputMap->GetAllocator() == AdditiveOutputMaps.GetAllocator(), TEXT("Transformer function is responsible for allocating using the output allocator"));
+					AdditiveOutputMaps.Append(OutputMap);
+				},
+				TransformerList->CreateBoundValueMapTransformerIterator(),
+				Base.GetBoundValueMaps().CreateConstIterator(),
+				Input.GetBoundValueMaps().CreateConstIterator());
+		}
+
+		// Unbound value maps
+		{
+			FUnboundMapCollection& AdditiveOutputMaps = AdditiveOutput.GetUnboundValueMaps();
+
+			OuterJoinBy(
+				FValueTransformerListWithUnboundMapCollectionJoinOp(),
+				[&AdditiveOutputMaps](FRawTransformerFunc Transformer, const FUnboundValueMap* BaseMap, const FUnboundValueMap* InputMap)
 				{
-					// No transformer was found for this value type, retain the input as-is
-					FUnboundValueMap* MapCopy = InputMap->Duplicate(AdditiveOutput.GetAllocator());
-					AdditiveOutput.AppendMap(MapCopy);
-					return;
-				}
+					if (InputMap == nullptr)
+					{
+						// The base doesn't have a matching pair in our input, nothing to convert
+						return;
+					}
 
-				// Transform our matching pair, note that the base map might be missing
-				// If it is missing, transformer specializations should use a default value
-				FUnboundValueMap* OutputMap = nullptr;
+					if (Transformer == nullptr)
+					{
+						// No transformer was found for this value type, retain the input as-is
+						AdditiveOutputMaps.Append(InputMap->Duplicate(AdditiveOutputMaps.GetAllocator()));
+						return;
+					}
 
-				FTransformUnboundValueMapFunc TypedTransformer = reinterpret_cast<FTransformUnboundValueMapFunc>(Transformer);
-				(*TypedTransformer)(BaseMap, InputMap, OutputMap, AdditiveOutput.GetAllocator());
+					// Transform our matching pair, note that the base map might be missing
+					// If it is missing, transformer specializations should use a default value
+					FUnboundValueMap* OutputMap = nullptr;
 
-				checkf(OutputMap != nullptr, TEXT("Transformer function is responsible for allocating an output buffer"));
-				checkf(OutputMap->GetAllocator() == AdditiveOutput.GetAllocator(), TEXT("Transformer function is responsible for allocating using the output allocator"));
-				AdditiveOutput.AppendMap(OutputMap);
-			},
-			TransformerList->CreateUnboundValueMapTransformerIterator(),
-			Base.CreateConstMapIterator(),
-			Input.CreateConstMapIterator());
+					FTransformUnboundValueMapFunc TypedTransformer = reinterpret_cast<FTransformUnboundValueMapFunc>(Transformer);
+					(*TypedTransformer)(BaseMap, InputMap, OutputMap, AdditiveOutputMaps.GetAllocator());
+
+					checkf(OutputMap != nullptr, TEXT("Transformer function is responsible for allocating an output buffer"));
+					checkf(OutputMap->GetAllocator() == AdditiveOutputMaps.GetAllocator(), TEXT("Transformer function is responsible for allocating using the output allocator"));
+					AdditiveOutputMaps.Append(OutputMap);
+				},
+				TransformerList->CreateUnboundValueMapTransformerIterator(),
+				Base.GetUnboundValueMaps().CreateConstIterator(),
+				Input.GetUnboundValueMaps().CreateConstIterator());
+		}
 	}
 
 	void FApplyAdditiveSpace::Apply(
 		const FValueTransformerMapPtr& TransformerMap,
-		const FAttributeCollection& Base,
-		const FAttributeCollection& Additive,
+		const FValueBundle& Base,
+		const FValueBundle& Additive,
 		float AdditiveWeight,
-		FAttributeCollection& Output)
+		FValueBundle& Output)
 	{
 		const bool bIsOutputBase = &Base == &Output;
 		if (bIsOutputBase && !FAnimWeight::IsRelevant(AdditiveWeight))
@@ -122,17 +129,17 @@ namespace UE::UAF::Transformers
 		}
 
 		// Allocator is irrelevant since it will remain empty
-		FAttributeCollectionHeap EmptyCollection(Additive.GetNamedSet());
+		FValueBundleHeap EmptyCollection(Additive.GetNamedSet());
 		FApplyAdditiveSpace::Apply(TransformerMap, Base, Additive, EmptyCollection, AdditiveWeight, Output);
 	}
 
 	void FApplyAdditiveSpace::Apply(
 		const FValueTransformerMapPtr& TransformerMap,
-		const FAttributeCollection& Base,
-		const FAttributeCollection& Additive,
-		const FAttributeCollection& PerValueWeights,
+		const FValueBundle& Base,
+		const FValueBundle& Additive,
+		const FValueBundle& PerValueWeights,
 		float DefaultWeight,
-		FAttributeCollection& Output)
+		FValueBundle& Output)
 	{
 		// Named sets must match as it ensures that our inputs have the same sizes / shapes
 		checkf(Base.GetNamedSet() == Additive.GetNamedSet(), TEXT("Base and Additive collections must use the same named set"));
@@ -161,14 +168,16 @@ namespace UE::UAF::Transformers
 			// Our output collection matches one of our inputs, we'll modify it in place
 			FReallocFun OutputAllocator = Output.GetAllocator();
 
-			// Set maps
+			// Bound value maps
 			{
+				const FBoundMapCollection& PerValueWeightMaps = PerValueWeights.GetBoundValueMaps();
+
 				// Keep a list of set maps that needs to be appended to our output
 				TArray<FBoundValueMap*, TInlineAllocator<8>> PendingMaps;
 
 				OuterJoinBy(
 					FValueTransformerListWithBoundMapCollectionJoinOp(),
-					[&PerValueWeights, DefaultWeight, bIsOutputBase, FloatValueType, OutputAllocator, &PendingMaps](FRawTransformerFunc Transformer, FBoundValueMap* BaseMap, FBoundValueMap* AdditiveMap)
+					[&PerValueWeightMaps, DefaultWeight, bIsOutputBase, FloatValueType, OutputAllocator, &PendingMaps](FRawTransformerFunc Transformer, FBoundValueMap* BaseMap, FBoundValueMap* AdditiveMap)
 					{
 						if (AdditiveMap == nullptr)
 						{
@@ -220,7 +229,7 @@ namespace UE::UAF::Transformers
 
 						// We cannot join our weights because their value type differs
 						const FAttributeMappingKey MappingKey = AdditiveMap->GetMappingKey().To(FloatValueType);
-						const FBoundValueMap* PerValueWeightMap = PerValueWeights.FindSetMap(MappingKey);
+						const FBoundValueMap* PerValueWeightMap = PerValueWeightMaps.Find(MappingKey);
 
 						FTransformBoundValueMapFunc TypedTransformer = reinterpret_cast<FTransformBoundValueMapFunc>(Transformer);
 						(*TypedTransformer)(BaseMap, AdditiveMap, PerValueWeightMap, DefaultWeight, OutputMap, OutputAllocator);
@@ -235,17 +244,18 @@ namespace UE::UAF::Transformers
 					},
 					TransformerList->CreateBoundValueMapTransformerIterator(),
 					// Remove the const qualifier since we'll mutate one in place as our output
-					const_cast<FAttributeCollection&>(Base).CreateSetMapIterator(),
-					const_cast<FAttributeCollection&>(Additive).CreateSetMapIterator());
+					const_cast<FValueBundle&>(Base).GetBoundValueMaps().CreateIterator(),
+					const_cast<FValueBundle&>(Additive).GetBoundValueMaps().CreateIterator());
 
 				// Append our pending entries
+				FBoundMapCollection& OutputMaps = Output.GetBoundValueMaps();
 				for (FBoundValueMap* Map : PendingMaps)
 				{
-					Output.AddSetMap(Map);
+					OutputMaps.Add(Map);
 				}
 			}
 
-			// Name/value maps
+			// Unbound value maps
 			{
 				// Keep a list of maps that needs to be appended to our output
 				TArray<FUnboundValueMap*, TInlineAllocator<8>> PendingMaps;
@@ -342,13 +352,14 @@ namespace UE::UAF::Transformers
 					},
 					TransformerList->CreateUnboundValueMapTransformerIterator(),
 					// Remove the const qualifier since we'll mutate one in place as our output
-					const_cast<FAttributeCollection&>(Base).CreateMapIterator(),
-					const_cast<FAttributeCollection&>(Additive).CreateMapIterator());
+					const_cast<FValueBundle&>(Base).GetUnboundValueMaps().CreateIterator(),
+					const_cast<FValueBundle&>(Additive).GetUnboundValueMaps().CreateIterator());
 
 				// Append our pending entries
+				FUnboundMapCollection& OutputMaps = Output.GetUnboundValueMaps();
 				for (FUnboundValueMap* Map : PendingMaps)
 				{
-					Output.AddMap(Map);
+					OutputMaps.Add(Map);
 				}
 			}
 		}
@@ -358,99 +369,108 @@ namespace UE::UAF::Transformers
 			Output.Reset(Base.GetNamedSet());
 			Output.SetValueSpace(BaseSpace);
 
-			// Set maps
-			OuterJoinBy(
-				FValueTransformerListWithBoundMapCollectionJoinOp(),
-				[&PerValueWeights, DefaultWeight, FloatValueType, &Output](FRawTransformerFunc Transformer, const FBoundValueMap* BaseMap, const FBoundValueMap* AdditiveMap)
-				{
-					if (AdditiveMap == nullptr)
-					{
-						// The base doesn't have a matching pair in our input, nothing to apply
-						return;
-					}
+			// Bound value maps
+			{
+				const FBoundMapCollection& PerValueWeightMaps = PerValueWeights.GetBoundValueMaps();
+				FBoundMapCollection& OutputMaps = Output.GetBoundValueMaps();
 
-					if (Transformer == nullptr)
+				OuterJoinBy(
+					FValueTransformerListWithBoundMapCollectionJoinOp(),
+					[&PerValueWeightMaps, DefaultWeight, FloatValueType, &OutputMaps](FRawTransformerFunc Transformer, const FBoundValueMap* BaseMap, const FBoundValueMap* AdditiveMap)
 					{
-						// No transformer was found for this value type, retain the one with the largest weight
-						FBoundValueMap* MapCopy;
-						if (BaseMap == nullptr || DefaultWeight >= 0.5f)
+						if (AdditiveMap == nullptr)
 						{
-							// Additive map has more weight or we have no base
-							MapCopy = AdditiveMap->Duplicate(Output.GetAllocator());
+							// The base doesn't have a matching pair in our input, nothing to apply
+							return;
 						}
-						else
+
+						if (Transformer == nullptr)
 						{
-							// Base map has more weight
-							MapCopy = BaseMap->Duplicate(Output.GetAllocator());
-						}
+							// No transformer was found for this value type, retain the one with the largest weight
+							FBoundValueMap* MapCopy;
+							if (BaseMap == nullptr || DefaultWeight >= 0.5f)
+							{
+								// Additive map has more weight or we have no base
+								MapCopy = AdditiveMap->Duplicate(OutputMaps.GetAllocator());
+							}
+							else
+							{
+								// Base map has more weight
+								MapCopy = BaseMap->Duplicate(OutputMaps.GetAllocator());
+							}
 
-						Output.AppendSetMap(MapCopy);
-						return;
-					}
+							OutputMaps.Append(MapCopy);
+							return;
+						}
 
-					// Transform our matching pair, note that the base map might be missing
-					// If it is missing, transformer specializations should use a default value
-					FBoundValueMap* OutputMap = nullptr;
+						// Transform our matching pair, note that the base map might be missing
+						// If it is missing, transformer specializations should use a default value
+						FBoundValueMap* OutputMap = nullptr;
 
-					// We cannot join our weights because their value type differs
-					const FAttributeMappingKey MappingKey = AdditiveMap->GetMappingKey().To(FloatValueType);
-					const FBoundValueMap* PerValueWeightMap = PerValueWeights.FindSetMap(MappingKey);
+						// We cannot join our weights because their value type differs
+						const FAttributeMappingKey MappingKey = AdditiveMap->GetMappingKey().To(FloatValueType);
+						const FBoundValueMap* PerValueWeightMap = PerValueWeightMaps.Find(MappingKey);
 
-					FTransformBoundValueMapFunc TypedTransformer = reinterpret_cast<FTransformBoundValueMapFunc>(Transformer);
-					(*TypedTransformer)(BaseMap, AdditiveMap, PerValueWeightMap, DefaultWeight, OutputMap, Output.GetAllocator());
+						FTransformBoundValueMapFunc TypedTransformer = reinterpret_cast<FTransformBoundValueMapFunc>(Transformer);
+						(*TypedTransformer)(BaseMap, AdditiveMap, PerValueWeightMap, DefaultWeight, OutputMap, OutputMaps.GetAllocator());
 
-					checkf(OutputMap != nullptr, TEXT("Transformer function is responsible for allocating an output buffer"));
-					checkf(OutputMap->GetAllocator() == Output.GetAllocator(), TEXT("Transformer function is responsible for allocating using the output allocator"));
-					Output.AppendSetMap(OutputMap);
-				},
-				TransformerList->CreateBoundValueMapTransformerIterator(),
-				Base.CreateConstSetMapIterator(),
-				Additive.CreateConstSetMapIterator());
+						checkf(OutputMap != nullptr, TEXT("Transformer function is responsible for allocating an output buffer"));
+						checkf(OutputMap->GetAllocator() == OutputMaps.GetAllocator(), TEXT("Transformer function is responsible for allocating using the output allocator"));
+						OutputMaps.Append(OutputMap);
+					},
+					TransformerList->CreateBoundValueMapTransformerIterator(),
+					Base.GetBoundValueMaps().CreateConstIterator(),
+					Additive.GetBoundValueMaps().CreateConstIterator());
+			}
 
-			// Name/value maps
-			OuterJoinBy(
-				FValueTransformerListWithUnboundMapCollectionJoinOp(),
-				[DefaultWeight, &Output](FRawTransformerFunc Transformer, const FUnboundValueMap* BaseMap, const FUnboundValueMap* AdditiveMap)
-				{
-					if (AdditiveMap == nullptr)
-					{
-						// The base doesn't have a matching pair in our input, nothing to apply
-						return;
-					}
+			// Unbound value maps
+			{
+				FUnboundMapCollection& OutputMaps = Output.GetUnboundValueMaps();
 
-					if (Transformer == nullptr)
+				OuterJoinBy(
+					FValueTransformerListWithUnboundMapCollectionJoinOp(),
+					[DefaultWeight, &OutputMaps](FRawTransformerFunc Transformer, const FUnboundValueMap* BaseMap, const FUnboundValueMap* AdditiveMap)
 					{
-						// No transformer was found for this value type, retain the one with the largest weight
-						FUnboundValueMap* MapCopy;
-						if (BaseMap == nullptr || DefaultWeight >= 0.5f)
+						if (AdditiveMap == nullptr)
 						{
-							// Additive map has more weight or we have no base
-							MapCopy = AdditiveMap->Duplicate(Output.GetAllocator());
+							// The base doesn't have a matching pair in our input, nothing to apply
+							return;
 						}
-						else
+
+						if (Transformer == nullptr)
 						{
-							// Base map has more weight
-							MapCopy = BaseMap->Duplicate(Output.GetAllocator());
-						}
+							// No transformer was found for this value type, retain the one with the largest weight
+							FUnboundValueMap* MapCopy;
+							if (BaseMap == nullptr || DefaultWeight >= 0.5f)
+							{
+								// Additive map has more weight or we have no base
+								MapCopy = AdditiveMap->Duplicate(OutputMaps.GetAllocator());
+							}
+							else
+							{
+								// Base map has more weight
+								MapCopy = BaseMap->Duplicate(OutputMaps.GetAllocator());
+							}
 
-						Output.AppendMap(MapCopy);
-						return;
-					}
+							OutputMaps.Append(MapCopy);
+							return;
+						}
 
-					// Transform our matching pair, note that the base map might be missing
-					// If it is missing, transformer specializations should use a default value
-					FUnboundValueMap* OutputMap = nullptr;
+						// Transform our matching pair, note that the base map might be missing
+						// If it is missing, transformer specializations should use a default value
+						FUnboundValueMap* OutputMap = nullptr;
 
-					FTransformUnboundValueMapFunc TypedTransformer = reinterpret_cast<FTransformUnboundValueMapFunc>(Transformer);
-					(*TypedTransformer)(BaseMap, AdditiveMap, DefaultWeight, OutputMap, Output.GetAllocator());
+						FTransformUnboundValueMapFunc TypedTransformer = reinterpret_cast<FTransformUnboundValueMapFunc>(Transformer);
+						(*TypedTransformer)(BaseMap, AdditiveMap, DefaultWeight, OutputMap, OutputMaps.GetAllocator());
 
-					checkf(OutputMap != nullptr, TEXT("Transformer function is responsible for allocating an output buffer"));
-					checkf(OutputMap->GetAllocator() == Output.GetAllocator(), TEXT("Transformer function is responsible for allocating using the output allocator"));
-					Output.AppendMap(OutputMap);
-				},
-				TransformerList->CreateUnboundValueMapTransformerIterator(),
-				Base.CreateConstMapIterator(),
-				Additive.CreateConstMapIterator());
+						checkf(OutputMap != nullptr, TEXT("Transformer function is responsible for allocating an output buffer"));
+						checkf(OutputMap->GetAllocator() == OutputMaps.GetAllocator(), TEXT("Transformer function is responsible for allocating using the output allocator"));
+						OutputMaps.Append(OutputMap);
+					},
+					TransformerList->CreateUnboundValueMapTransformerIterator(),
+					Base.GetUnboundValueMaps().CreateConstIterator(),
+					Additive.GetUnboundValueMaps().CreateConstIterator());
+			}
 		}
 	}
 
@@ -466,7 +486,7 @@ namespace UE::UAF::Transformers
 				// Our base is missing, use the reference pose
 				const FAttributeTypedSetPtr& TypedSet = Input->GetTypedSet();
 				const FAttributeCollectionPtr ReferenceValues = TypedSet->GetBindingData()->GetDefaultAttributeValues(TypedSet->GetNamedSet()->GetName());
-				Base = ReferenceValues->FindSetMap(Input->GetMappingKey());
+				Base = ReferenceValues->GetBoundValueMaps().Find(Input->GetMappingKey());
 			}
 
 			if (Output == nullptr)
@@ -563,7 +583,7 @@ namespace UE::UAF::Transformers
 				// Our base is missing, use the user authored default values
 				const FAttributeTypedSetPtr& TypedSet = Input->GetTypedSet();
 				const FAttributeCollectionPtr ReferenceValues = TypedSet->GetBindingData()->GetDefaultAttributeValues(TypedSet->GetNamedSet()->GetName());
-				Base = ReferenceValues->FindSetMap(Input->GetMappingKey());
+				Base = ReferenceValues->GetBoundValueMaps().Find(Input->GetMappingKey());
 			}
 
 			if (Output == nullptr)
@@ -647,7 +667,7 @@ namespace UE::UAF::Transformers
 				// Our base is missing, use the reference pose
 				const FAttributeTypedSetPtr& TypedSet = Additive->GetTypedSet();
 				const FAttributeCollectionPtr ReferenceValues = TypedSet->GetBindingData()->GetDefaultAttributeValues(TypedSet->GetNamedSet()->GetName());
-				Base = ReferenceValues->FindSetMap(Additive->GetMappingKey());
+				Base = ReferenceValues->GetBoundValueMaps().Find(Additive->GetMappingKey());
 			}
 
 			if (Output == nullptr)
@@ -805,7 +825,7 @@ namespace UE::UAF::Transformers
 				// Our base is missing, use the user authored default values
 				const FAttributeTypedSetPtr& TypedSet = Additive->GetTypedSet();
 				const FAttributeCollectionPtr ReferenceValues = TypedSet->GetBindingData()->GetDefaultAttributeValues(TypedSet->GetNamedSet()->GetName());
-				Base = ReferenceValues->FindSetMap(Additive->GetMappingKey());
+				Base = ReferenceValues->GetBoundValueMaps().Find(Additive->GetMappingKey());
 			}
 
 			if (Output == nullptr)
diff --git a/Engine/Plugins/Experimental/UAF/UAF/Source/UAF/Private/UAF/ValueRuntime/Transformers/BoneSpace.cpp b/Engine/Plugins/Experimental/UAF/UAF/Source/UAF/Private/UAF/ValueRuntime/Transformers/BoneSpace.cpp
--- a/Engine/Plugins/Experimental/UAF/UAF/Source/UAF/Private/UAF/ValueRuntime/Transformers/BoneSpace.cpp
+++ b/Engine/Plugins/Experimental/UAF/UAF/Source/UAF/Private/UAF/ValueRuntime/Transformers/BoneSpace.cpp
@@ -2,12 +2,12 @@
 
 #include "UAF/ValueRuntime/Transformers/BoneSpace.h"
 
-#include "UAF/ValueRuntime/AttributeCollection.h"
+#include "UAF/ValueRuntime/ValueBundle.h"
 #include "UAF/Attributes/EngineAttributes.h"
 
 namespace UE::UAF::Transformers
 {
-	void FBoneSpace::LocalToComponent(const FAttributeCollection& InputValues, FAttributeCollection& OutputValues)
+	void FBoneSpace::LocalToComponent(const FValueBundle& InputValues, FValueBundle& OutputValues)
 	{
 		checkf(InputValues.GetValueSpace().GetType() == EValueSpaceType::Local, TEXT("Input must be in local space"));
 		checkf(!InputValues.GetValueSpace().IsAdditive(), TEXT("Input space cannot be additive"));
@@ -21,7 +21,7 @@ namespace UE::UAF::Transformers
 		OutputValues.SetValueSpace(FValueSpace(EValueSpaceType::Component));
 
 		// Modify in place
-		TBoundValueMap<FBoneTransformAnimationAttribute>* BoneTransforms = OutputValues.FindSetMap<FBoneTransformAnimationAttribute>(FAttributeMappingKey::MakeFromTo<FBoneTransformAnimationAttribute>());
+		TBoundValueMap<FBoneTransformAnimationAttribute>* BoneTransforms = OutputValues.GetBoundValueMaps().Find<FBoneTransformAnimationAttribute>(FAttributeMappingKey::MakeFromTo<FBoneTransformAnimationAttribute>());
 		if (BoneTransforms == nullptr)
 		{
 			// No bone transforms, nothing to convert
@@ -52,7 +52,7 @@ namespace UE::UAF::Transformers
 		}
 	}
 
-	void FBoneSpace::ComponentToLocal(const FAttributeCollection& InputValues, FAttributeCollection& OutputValues)
+	void FBoneSpace::ComponentToLocal(const FValueBundle& InputValues, FValueBundle& OutputValues)
 	{
 		checkf(InputValues.GetValueSpace().GetType() == EValueSpaceType::Component, TEXT("Input must be in component space"));
 		checkf(!InputValues.GetValueSpace().IsAdditive(), TEXT("Input space cannot be additive"));
@@ -66,7 +66,7 @@ namespace UE::UAF::Transformers
 		OutputValues.SetValueSpace(FValueSpace(EValueSpaceType::Local));
 
 		// Modify in place
-		TBoundValueMap<FBoneTransformAnimationAttribute>* BoneTransforms = OutputValues.FindSetMap<FBoneTransformAnimationAttribute>(FAttributeMappingKey::MakeFromTo<FBoneTransformAnimationAttribute>());
+		TBoundValueMap<FBoneTransformAnimationAttribute>* BoneTransforms = OutputValues.GetBoundValueMaps().Find<FBoneTransformAnimationAttribute>(FAttributeMappingKey::MakeFromTo<FBoneTransformAnimationAttribute>());
 		if (BoneTransforms == nullptr)
 		{
 			// No bone transforms, nothing to convert
@@ -97,7 +97,7 @@ namespace UE::UAF::Transformers
 		}
 	}
 
-	void FBoneSpace::LocalToMeshRotation(const FAttributeCollection& InputValues, FAttributeCollection& OutputValues)
+	void FBoneSpace::LocalToMeshRotation(const FValueBundle& InputValues, FValueBundle& OutputValues)
 	{
 		const FValueSpace InputSpace = InputValues.GetValueSpace();
 		checkf(InputSpace.GetType() == EValueSpaceType::Local || InputSpace.GetType() == EValueSpaceType::Mixed, TEXT("Input must be in local or mixed space"));
@@ -116,7 +116,7 @@ namespace UE::UAF::Transformers
 		OutputValues.SetValueSpace(FValueSpace(MixedSpaceFlags));
 
 		// Modify in place
-		TBoundValueMap<FBoneTransformAnimationAttribute>* BoneTransforms = OutputValues.FindSetMap<FBoneTransformAnimationAttribute>(FAttributeMappingKey::MakeFromTo<FBoneTransformAnimationAttribute>());
+		TBoundValueMap<FBoneTransformAnimationAttribute>* BoneTransforms = OutputValues.GetBoundValueMaps().Find<FBoneTransformAnimationAttribute>(FAttributeMappingKey::MakeFromTo<FBoneTransformAnimationAttribute>());
 		if (BoneTransforms == nullptr)
 		{
 			// No bone transforms, nothing to convert
@@ -146,7 +146,7 @@ namespace UE::UAF::Transformers
 		}
 	}
 
-	void FBoneSpace::MeshRotationToLocal(const FAttributeCollection& InputValues, FAttributeCollection& OutputValues)
+	void FBoneSpace::MeshRotationToLocal(const FValueBundle& InputValues, FValueBundle& OutputValues)
 	{
 		const FValueSpace InputSpace = InputValues.GetValueSpace();
 		EMixedSpaceFlags MixedSpaceFlags = InputSpace.GetMixedSpaceFlags();
@@ -174,7 +174,7 @@ namespace UE::UAF::Transformers
 		OutputValues.SetValueSpace(OutputSpace);
 
 		// Modify in place
-		TBoundValueMap<FBoneTransformAnimationAttribute>* BoneTransforms = OutputValues.FindSetMap<FBoneTransformAnimationAttribute>(FAttributeMappingKey::MakeFromTo<FBoneTransformAnimationAttribute>());
+		TBoundValueMap<FBoneTransformAnimationAttribute>* BoneTransforms = OutputValues.GetBoundValueMaps().Find<FBoneTransformAnimationAttribute>(FAttributeMappingKey::MakeFromTo<FBoneTransformAnimationAttribute>());
 		if (BoneTransforms == nullptr)
 		{
 			// No bone transforms, nothing to convert
@@ -205,7 +205,7 @@ namespace UE::UAF::Transformers
 		}
 	}
 
-	void FBoneSpace::LocalToRootRotation(const FAttributeCollection& InputValues, FAttributeCollection& OutputValues)
+	void FBoneSpace::LocalToRootRotation(const FValueBundle& InputValues, FValueBundle& OutputValues)
 	{
 		const FValueSpace InputSpace = InputValues.GetValueSpace();
 		checkf(InputSpace.GetType() == EValueSpaceType::Local || InputSpace.GetType() == EValueSpaceType::Mixed, TEXT("Input must be in local or mixed space"));
@@ -224,7 +224,7 @@ namespace UE::UAF::Transformers
 		OutputValues.SetValueSpace(FValueSpace(MixedSpaceFlags));
 
 		// Modify in place
-		TBoundValueMap<FBoneTransformAnimationAttribute>* BoneTransforms = OutputValues.FindSetMap<FBoneTransformAnimationAttribute>(FAttributeMappingKey::MakeFromTo<FBoneTransformAnimationAttribute>());
+		TBoundValueMap<FBoneTransformAnimationAttribute>* BoneTransforms = OutputValues.GetBoundValueMaps().Find<FBoneTransformAnimationAttribute>(FAttributeMappingKey::MakeFromTo<FBoneTransformAnimationAttribute>());
 		if (BoneTransforms == nullptr)
 		{
 			// No bone transforms, nothing to convert
@@ -256,7 +256,7 @@ namespace UE::UAF::Transformers
 		}
 	}
 
-	void FBoneSpace::RootRotationToLocal(const FAttributeCollection& InputValues, FAttributeCollection& OutputValues)
+	void FBoneSpace::RootRotationToLocal(const FValueBundle& InputValues, FValueBundle& OutputValues)
 	{
 		const FValueSpace InputSpace = InputValues.GetValueSpace();
 		EMixedSpaceFlags MixedSpaceFlags = InputSpace.GetMixedSpaceFlags();
@@ -284,7 +284,7 @@ namespace UE::UAF::Transformers
 		OutputValues.SetValueSpace(OutputSpace);
 
 		// Modify in place
-		TBoundValueMap<FBoneTransformAnimationAttribute>* BoneTransforms = OutputValues.FindSetMap<FBoneTransformAnimationAttribute>(FAttributeMappingKey::MakeFromTo<FBoneTransformAnimationAttribute>());
+		TBoundValueMap<FBoneTransformAnimationAttribute>* BoneTransforms = OutputValues.GetBoundValueMaps().Find<FBoneTransformAnimationAttribute>(FAttributeMappingKey::MakeFromTo<FBoneTransformAnimationAttribute>());
 		if (BoneTransforms == nullptr)
 		{
 			// No bone transforms, nothing to convert
@@ -317,7 +317,7 @@ namespace UE::UAF::Transformers
 		}
 	}
 
-	void FBoneSpace::LocalToMeshScale(const FAttributeCollection& InputValues, FAttributeCollection& OutputValues)
+	void FBoneSpace::LocalToMeshScale(const FValueBundle& InputValues, FValueBundle& OutputValues)
 	{
 		const FValueSpace InputSpace = InputValues.GetValueSpace();
 		checkf(InputSpace.GetType() == EValueSpaceType::Local || InputSpace.GetType() == EValueSpaceType::Mixed, TEXT("Input must be in local or mixed space"));
@@ -335,7 +335,7 @@ namespace UE::UAF::Transformers
 		OutputValues.SetValueSpace(FValueSpace(MixedSpaceFlags));
 
 		// Modify in place
-		TBoundValueMap<FBoneTransformAnimationAttribute>* BoneTransforms = OutputValues.FindSetMap<FBoneTransformAnimationAttribute>(FAttributeMappingKey::MakeFromTo<FBoneTransformAnimationAttribute>());
+		TBoundValueMap<FBoneTransformAnimationAttribute>* BoneTransforms = OutputValues.GetBoundValueMaps().Find<FBoneTransformAnimationAttribute>(FAttributeMappingKey::MakeFromTo<FBoneTransformAnimationAttribute>());
 		if (BoneTransforms == nullptr)
 		{
 			// No bone transforms, nothing to convert
@@ -366,7 +366,7 @@ namespace UE::UAF::Transformers
 		}
 	}
 
-	void FBoneSpace::MeshScaleToLocal(const FAttributeCollection& InputValues, FAttributeCollection& OutputValues)
+	void FBoneSpace::MeshScaleToLocal(const FValueBundle& InputValues, FValueBundle& OutputValues)
 	{
 		const FValueSpace InputSpace = InputValues.GetValueSpace();
 		EMixedSpaceFlags MixedSpaceFlags = InputSpace.GetMixedSpaceFlags();
@@ -394,7 +394,7 @@ namespace UE::UAF::Transformers
 		OutputValues.SetValueSpace(OutputSpace);
 
 		// Modify in place
-		TBoundValueMap<FBoneTransformAnimationAttribute>* BoneTransforms = OutputValues.FindSetMap<FBoneTransformAnimationAttribute>(FAttributeMappingKey::MakeFromTo<FBoneTransformAnimationAttribute>());
+		TBoundValueMap<FBoneTransformAnimationAttribute>* BoneTransforms = OutputValues.GetBoundValueMaps().Find<FBoneTransformAnimationAttribute>(FAttributeMappingKey::MakeFromTo<FBoneTransformAnimationAttribute>());
 		if (BoneTransforms == nullptr)
 		{
 			// No bone transforms, nothing to convert
diff --git a/Engine/Plugins/Experimental/UAF/UAF/Source/UAF/Private/UAF/ValueRuntime/Transformers/Interpolate.cpp b/Engine/Plugins/Experimental/UAF/UAF/Source/UAF/Private/UAF/ValueRuntime/Transformers/Interpolate.cpp
--- a/Engine/Plugins/Experimental/UAF/UAF/Source/UAF/Private/UAF/ValueRuntime/Transformers/Interpolate.cpp
+++ b/Engine/Plugins/Experimental/UAF/UAF/Source/UAF/Private/UAF/ValueRuntime/Transformers/Interpolate.cpp
@@ -10,10 +10,10 @@ namespace UE::UAF::Transformers
 
 	void FInterpolate::Apply(
 		const FValueTransformerMapPtr& TransformerMap,
-		const FAttributeCollection& InputA,
-		const FAttributeCollection& InputB,
+		const FValueBundle& InputA,
+		const FValueBundle& InputB,
 		float WeightB,
-		FAttributeCollection& Output)
+		FValueBundle& Output)
 	{
 		// Named sets must match as it ensures that our inputs have the same sizes/shapes
 		checkf(InputA.GetNamedSet() == InputB.GetNamedSet(), TEXT("Both input collections must use the same named set"));
@@ -34,17 +34,17 @@ namespace UE::UAF::Transformers
 		}
 
 		// Allocator is irrelevant since it will remain empty
-		FAttributeCollectionHeap EmptyCollection(InputA.GetNamedSet());
+		FValueBundleHeap EmptyCollection(InputA.GetNamedSet());
 		FInterpolate::Apply(TransformerMap, InputA, InputB, EmptyCollection, WeightB, Output);
 	}
 
 	void FInterpolate::Apply(
 		const FValueTransformerMapPtr& TransformerMap,
-		const FAttributeCollection& InputA,
-		const FAttributeCollection& InputB,
-		const FAttributeCollection& PerValueWeightsB,
+		const FValueBundle& InputA,
+		const FValueBundle& InputB,
+		const FValueBundle& PerValueWeightsB,
 		float DefaultWeightB,
-		FAttributeCollection& Output)
+		FValueBundle& Output)
 	{
 		// Named sets must match as it ensures that our inputs have the same sizes/shapes
 		checkf(InputA.GetNamedSet() == InputB.GetNamedSet(), TEXT("Both input collections must use the same named set"));
@@ -69,12 +69,14 @@ namespace UE::UAF::Transformers
 
 			// Bound value maps
 			{
+				const FBoundMapCollection& PerValueWeightBMaps = PerValueWeightsB.GetBoundValueMaps();
+
 				// Keep a list of set maps that needs to be appended to our output
 				TArray<FBoundValueMap*, TInlineAllocator<8>> PendingMaps;
 
 				OuterJoinBy(
 					FValueTransformerListWithBoundMapCollectionJoinOp(),
-					[&PerValueWeightsB, DefaultWeightB, bIsOutputA, FloatValueType, OutputAllocator, &PendingMaps](FRawTransformerFunc Transformer, FBoundValueMap* MapA, FBoundValueMap* MapB)
+					[&PerValueWeightBMaps, DefaultWeightB, bIsOutputA, FloatValueType, OutputAllocator, &PendingMaps](FRawTransformerFunc Transformer, FBoundValueMap* MapA, FBoundValueMap* MapB)
 					{
 						if (MapA == nullptr)
 						{
@@ -150,7 +152,7 @@ namespace UE::UAF::Transformers
 
 									// We cannot join our weights because their value type differs
 									const FAttributeMappingKey MappingKey = MapA->GetMappingKey().To(FloatValueType);
-									const FBoundValueMap* PerValueWeightMapB = PerValueWeightsB.FindSetMap(MappingKey);
+									const FBoundValueMap* PerValueWeightMapB = PerValueWeightBMaps.Find(MappingKey);
 
 									FTransformBoundValueMapFunc TypedTransformer = reinterpret_cast<FTransformBoundValueMapFunc>(Transformer);
 									(*TypedTransformer)(MapA, MapB, PerValueWeightMapB, DefaultWeightB, OutputMap, nullptr);
@@ -160,13 +162,14 @@ namespace UE::UAF::Transformers
 					},
 					TransformerList->CreateBoundValueMapTransformerIterator(),
 					// Remove the const qualifier since we'll mutate one in place as our output
-					const_cast<FAttributeCollection&>(InputA).CreateSetMapIterator(),
-					const_cast<FAttributeCollection&>(InputB).CreateSetMapIterator());
+					const_cast<FValueBundle&>(InputA).GetBoundValueMaps().CreateIterator(),
+					const_cast<FValueBundle&>(InputB).GetBoundValueMaps().CreateIterator());
 
 				// Append our pending entries
+				FBoundMapCollection& OutputMaps = Output.GetBoundValueMaps();
 				for (FBoundValueMap* Map : PendingMaps)
 				{
-					Output.AddSetMap(Map);
+					OutputMaps.Add(Map);
 				}
 			}
 
@@ -268,13 +271,14 @@ namespace UE::UAF::Transformers
 					},
 					TransformerList->CreateUnboundValueMapTransformerIterator(),
 					// Remove the const qualifier since we'll mutate one in place as our output
-					const_cast<FAttributeCollection&>(InputA).CreateMapIterator(),
-					const_cast<FAttributeCollection&>(InputB).CreateMapIterator());
+					const_cast<FValueBundle&>(InputA).GetUnboundValueMaps().CreateIterator(),
+					const_cast<FValueBundle&>(InputB).GetUnboundValueMaps().CreateIterator());
 
 				// Append our pending entries
+				FUnboundMapCollection& OutputMaps = Output.GetUnboundValueMaps();
 				for (FUnboundValueMap* Map : PendingMaps)
 				{
-					Output.AddMap(Map);
+					OutputMaps.Add(Map);
 				}
 			}
 		}
@@ -285,128 +289,137 @@ namespace UE::UAF::Transformers
 			Output.SetValueSpace(InputA.GetValueSpace());
 
 			// Bound value maps
-			OuterJoinBy(
-				FValueTransformerListWithBoundMapCollectionJoinOp(),
-				[&PerValueWeightsB, DefaultWeightB, FloatValueType, &Output](FRawTransformerFunc Transformer, const FBoundValueMap* MapA, const FBoundValueMap* MapB)
-				{
-					if (MapA == nullptr)
-					{
-						if (MapB == nullptr)
-						{
-							// No maps found for the transformer type
-						}
-						else
-						{
-							// This set map doesn't have a matching pair, duplicate it
-							Output.AppendSetMap(MapB->Duplicate(Output.GetAllocator()));
-						}
-					}
-					else
+			{
+				const FBoundMapCollection& PerValueWeightBMaps = PerValueWeightsB.GetBoundValueMaps();
+				FBoundMapCollection& OutputMaps = Output.GetBoundValueMaps();
+
+				OuterJoinBy(
+					FValueTransformerListWithBoundMapCollectionJoinOp(),
+					[&PerValueWeightBMaps, DefaultWeightB, FloatValueType, &OutputMaps](FRawTransformerFunc Transformer, const FBoundValueMap* MapA, const FBoundValueMap* MapB)
 					{
-						if (MapB == nullptr)
+						if (MapA == nullptr)
 						{
-							// This set map doesn't have a matching pair, duplicate it
-							Output.AppendSetMap(MapA->Duplicate(Output.GetAllocator()));
+							if (MapB == nullptr)
+							{
+								// No maps found for the transformer type
+							}
+							else
+							{
+								// This set map doesn't have a matching pair, duplicate it
+								OutputMaps.Append(MapB->Duplicate(OutputMaps.GetAllocator()));
+							}
 						}
 						else
 						{
-							if (Transformer == nullptr)
+							if (MapB == nullptr)
+							{
+								// This set map doesn't have a matching pair, duplicate it
+								OutputMaps.Append(MapA->Duplicate(OutputMaps.GetAllocator()));
+							}
+							else
 							{
-								// No transformer was found for this value type, retain the one with the largest weight
-								FBoundValueMap* MapCopy;
-								if (DefaultWeightB >= 0.5f)
+								if (Transformer == nullptr)
 								{
-									MapCopy = MapB->Duplicate(Output.GetAllocator());
+									// No transformer was found for this value type, retain the one with the largest weight
+									FBoundValueMap* MapCopy;
+									if (DefaultWeightB >= 0.5f)
+									{
+										MapCopy = MapB->Duplicate(OutputMaps.GetAllocator());
+									}
+									else
+									{
+										MapCopy = MapA->Duplicate(OutputMaps.GetAllocator());
+									}
+
+									OutputMaps.Append(MapCopy);
 								}
 								else
 								{
-									MapCopy = MapA->Duplicate(Output.GetAllocator());
-								}
-
-								Output.AppendSetMap(MapCopy);
-							}
-							else
-							{
-								// Transform our matching pair
-								FBoundValueMap* OutputMap = nullptr;
+									// Transform our matching pair
+									FBoundValueMap* OutputMap = nullptr;
 
-								// We cannot join our weights because their value type differs
-								const FAttributeMappingKey MappingKey = MapA->GetMappingKey().To(FloatValueType);
-								const FBoundValueMap* PerValueWeightMapB = PerValueWeightsB.FindSetMap(MappingKey);
+									// We cannot join our weights because their value type differs
+									const FAttributeMappingKey MappingKey = MapA->GetMappingKey().To(FloatValueType);
+									const FBoundValueMap* PerValueWeightMapB = PerValueWeightBMaps.Find(MappingKey);
 
-								FTransformBoundValueMapFunc TypedTransformer = reinterpret_cast<FTransformBoundValueMapFunc>(Transformer);
-								(*TypedTransformer)(MapA, MapB, PerValueWeightMapB, DefaultWeightB, OutputMap, Output.GetAllocator());
+									FTransformBoundValueMapFunc TypedTransformer = reinterpret_cast<FTransformBoundValueMapFunc>(Transformer);
+									(*TypedTransformer)(MapA, MapB, PerValueWeightMapB, DefaultWeightB, OutputMap, OutputMaps.GetAllocator());
 
-								checkf(OutputMap != nullptr, TEXT("Transformer function is responsible for allocating an output buffer"));
-								checkf(OutputMap->GetAllocator() == Output.GetAllocator(), TEXT("Transformer function is responsible for allocating using the output allocator"));
-								Output.AppendSetMap(OutputMap);
+									checkf(OutputMap != nullptr, TEXT("Transformer function is responsible for allocating an output buffer"));
+									checkf(OutputMap->GetAllocator() == OutputMaps.GetAllocator(), TEXT("Transformer function is responsible for allocating using the output allocator"));
+									OutputMaps.Append(OutputMap);
+								}
 							}
 						}
-					}
-				},
-				TransformerList->CreateBoundValueMapTransformerIterator(),
-				InputA.CreateConstSetMapIterator(),
-				InputB.CreateConstSetMapIterator());
+					},
+					TransformerList->CreateBoundValueMapTransformerIterator(),
+					InputA.GetBoundValueMaps().CreateConstIterator(),
+					InputB.GetBoundValueMaps().CreateConstIterator());
+			}
 
 			// Name/value maps
-			OuterJoinBy(
-				FValueTransformerListWithUnboundMapCollectionJoinOp(),
-				[DefaultWeightB, &Output](FRawTransformerFunc Transformer, const FUnboundValueMap* MapA, const FUnboundValueMap* MapB)
-				{
-					if (MapA == nullptr)
-					{
-						if (MapB == nullptr)
-						{
-							// No maps found for the transformer type
-						}
-						else
-						{
-							// This map doesn't have a matching pair, duplicate it
-							Output.AppendMap(MapB->Duplicate(Output.GetAllocator()));
-						}
-					}
-					else
+			{
+				FUnboundMapCollection& OutputMaps = Output.GetUnboundValueMaps();
+
+				OuterJoinBy(
+					FValueTransformerListWithUnboundMapCollectionJoinOp(),
+					[DefaultWeightB, &OutputMaps](FRawTransformerFunc Transformer, const FUnboundValueMap* MapA, const FUnboundValueMap* MapB)
 					{
-						if (MapB == nullptr)
+						if (MapA == nullptr)
 						{
-							// This map doesn't have a matching pair
-							// Our output is A and thus already contains this
+							if (MapB == nullptr)
+							{
+								// No maps found for the transformer type
+							}
+							else
+							{
+								// This map doesn't have a matching pair, duplicate it
+								OutputMaps.Append(MapB->Duplicate(OutputMaps.GetAllocator()));
+							}
 						}
 						else
 						{
-							if (Transformer == nullptr)
+							if (MapB == nullptr)
+							{
+								// This map doesn't have a matching pair
+								// Our output is A and thus already contains this
+							}
+							else
 							{
-								// No transformer was found for this value type, retain the one with the largest weight
-								FUnboundValueMap* MapCopy;
-								if (DefaultWeightB >= 0.5f)
+								if (Transformer == nullptr)
 								{
-									MapCopy = MapB->Duplicate(Output.GetAllocator());
+									// No transformer was found for this value type, retain the one with the largest weight
+									FUnboundValueMap* MapCopy;
+									if (DefaultWeightB >= 0.5f)
+									{
+										MapCopy = MapB->Duplicate(OutputMaps.GetAllocator());
+									}
+									else
+									{
+										MapCopy = MapA->Duplicate(OutputMaps.GetAllocator());
+									}
+
+									OutputMaps.Append(MapCopy);
 								}
 								else
 								{
-									MapCopy = MapA->Duplicate(Output.GetAllocator());
-								}
-
-								Output.AppendMap(MapCopy);
-							}
-							else
-							{
-								// Transform our matching pair in place
-								FUnboundValueMap* OutputMap = nullptr;
+									// Transform our matching pair in place
+									FUnboundValueMap* OutputMap = nullptr;
 
-								FTransformUnboundValueMapFunc TypedTransformer = reinterpret_cast<FTransformUnboundValueMapFunc>(Transformer);
-								(*TypedTransformer)(MapA, MapB, DefaultWeightB, OutputMap, Output.GetAllocator());
+									FTransformUnboundValueMapFunc TypedTransformer = reinterpret_cast<FTransformUnboundValueMapFunc>(Transformer);
+									(*TypedTransformer)(MapA, MapB, DefaultWeightB, OutputMap, OutputMaps.GetAllocator());
 
-								checkf(OutputMap != nullptr, TEXT("Transformer function is responsible for allocating an output buffer"));
-								checkf(OutputMap->GetAllocator() == Output.GetAllocator(), TEXT("Transformer function is responsible for allocating using the output allocator"));
-								Output.AppendMap(OutputMap);
+									checkf(OutputMap != nullptr, TEXT("Transformer function is responsible for allocating an output buffer"));
+									checkf(OutputMap->GetAllocator() == OutputMaps.GetAllocator(), TEXT("Transformer function is responsible for allocating using the output allocator"));
+									OutputMaps.Append(OutputMap);
+								}
 							}
 						}
-					}
-				},
-				TransformerList->CreateUnboundValueMapTransformerIterator(),
-				InputA.CreateConstMapIterator(),
-				InputB.CreateConstMapIterator());
+					},
+					TransformerList->CreateUnboundValueMapTransformerIterator(),
+					InputA.GetUnboundValueMaps().CreateConstIterator(),
+					InputB.GetUnboundValueMaps().CreateConstIterator());
+			}
 		}
 	}
 
diff --git a/Engine/Plugins/Experimental/UAF/UAF/Source/UAF/Private/UAF/ValueRuntime/Transformers/Layer.cpp b/Engine/Plugins/Experimental/UAF/UAF/Source/UAF/Private/UAF/ValueRuntime/Transformers/Layer.cpp
--- a/Engine/Plugins/Experimental/UAF/UAF/Source/UAF/Private/UAF/ValueRuntime/Transformers/Layer.cpp
+++ b/Engine/Plugins/Experimental/UAF/UAF/Source/UAF/Private/UAF/ValueRuntime/Transformers/Layer.cpp
@@ -10,10 +10,10 @@ namespace UE::UAF::Transformers
 
 	void FLayer::Apply(
 		const FValueTransformerMapPtr& TransformerMap,
-		const FAttributeCollection& Base,
-		const FAttributeCollection& Layer,
+		const FValueBundle& Base,
+		const FValueBundle& Layer,
 		float LayerWeight,
-		FAttributeCollection& Output)
+		FValueBundle& Output)
 	{
 		const bool bIsOutputBase = &Base == &Output;
 		if (bIsOutputBase && !FAnimWeight::IsRelevant(LayerWeight))
@@ -23,17 +23,17 @@ namespace UE::UAF::Transformers
 		}
 
 		// Allocator is irrelevant since it will remain empty
-		FAttributeCollectionHeap EmptyCollection(Base.GetNamedSet());
+		FValueBundleHeap EmptyCollection(Base.GetNamedSet());
 		FLayer::Apply(TransformerMap, Base, Layer, EmptyCollection, LayerWeight, Output);
 	}
 
 	void FLayer::Apply(
 		const FValueTransformerMapPtr& TransformerMap,
-		const FAttributeCollection& Base,
-		const FAttributeCollection& Layer,
-		const FAttributeCollection& PerValueLayerWeights,
+		const FValueBundle& Base,
+		const FValueBundle& Layer,
+		const FValueBundle& PerValueLayerWeights,
 		float DefaultLayerWeight,
-		FAttributeCollection& Output)
+		FValueBundle& Output)
 	{
 		const FValueTransformerList* TransformerList = TransformerMap->Find(TransformerName);
 		if (TransformerList == nullptr)
@@ -44,8 +44,8 @@ namespace UE::UAF::Transformers
 
 		const float LayerWeight = FMath::Clamp(DefaultLayerWeight, 0.0f, 1.0f);
 
-		FAttributeCollectionStack LayerCopy;
-		const FAttributeCollection* LayerPtr = &Layer;
+		FValueBundleStack LayerCopy;
+		const FValueBundle* LayerPtr = &Layer;
 
 		const bool bIsOutputBase = &Base == &Output;
 		bool bIsOutputLayer = &Layer == &Output;
@@ -72,10 +72,12 @@ namespace UE::UAF::Transformers
 			// Our output collection matches one of our inputs, we'll modify it in place
 			FReallocFun OutputAllocator = Output.GetAllocator();
 
+			const FBoundMapCollection& PerValueLayerWeightMaps = PerValueLayerWeights.GetBoundValueMaps();
+
 			// Bound value maps
 			OuterJoinBy(
 				FValueTransformerListWithBoundMapCollectionJoinOp(),
-				[&PerValueLayerWeights, LayerWeight, bIsOutputBase, FloatValueType, OutputAllocator](FRawTransformerFunc Transformer, FBoundValueMap* BaseMap, FBoundValueMap* LayerMap)
+				[&PerValueLayerWeightMaps, LayerWeight, bIsOutputBase, FloatValueType, OutputAllocator](FRawTransformerFunc Transformer, FBoundValueMap* BaseMap, FBoundValueMap* LayerMap)
 				{
 					if (BaseMap == nullptr)
 					{
@@ -133,7 +135,7 @@ namespace UE::UAF::Transformers
 
 								// We cannot join our weights because their value type differs
 								const FAttributeMappingKey MappingKey = BaseMap->GetMappingKey().To(FloatValueType);
-								const FBoundValueMap* PerValueWeightMap = PerValueLayerWeights.FindSetMap(MappingKey);
+								const FBoundValueMap* PerValueWeightMap = PerValueLayerWeightMaps.Find(MappingKey);
 
 								FTransformBoundValueMapFunc TypedTransformer = reinterpret_cast<FTransformBoundValueMapFunc>(Transformer);
 								(*TypedTransformer)(BaseMap, LayerMap, PerValueWeightMap, LayerWeight, OutputMap, nullptr);
@@ -143,8 +145,8 @@ namespace UE::UAF::Transformers
 				},
 				TransformerList->CreateBoundValueMapTransformerIterator(),
 				// Remove the const qualifier since we'll mutate one in place as our output
-				const_cast<FAttributeCollection&>(Base).CreateSetMapIterator(),
-				const_cast<FAttributeCollection&>(Layer).CreateSetMapIterator());
+				const_cast<FValueBundle&>(Base).GetBoundValueMaps().CreateIterator(),
+				const_cast<FValueBundle&>(Layer).GetBoundValueMaps().CreateIterator());
 
 			// Name/value maps
 			OuterJoinBy(
@@ -222,8 +224,8 @@ namespace UE::UAF::Transformers
 				},
 				TransformerList->CreateUnboundValueMapTransformerIterator(),
 				// Remove the const qualifier since we'll mutate one in place as our output
-				const_cast<FAttributeCollection&>(Base).CreateMapIterator(),
-				const_cast<FAttributeCollection&>(Layer).CreateMapIterator());
+				const_cast<FValueBundle&>(Base).GetUnboundValueMaps().CreateIterator(),
+				const_cast<FValueBundle&>(Layer).GetUnboundValueMaps().CreateIterator());
 		}
 		else
 		{
@@ -232,128 +234,137 @@ namespace UE::UAF::Transformers
 			Output.SetValueSpace(Base.GetValueSpace());
 
 			// Bound value maps
-			OuterJoinBy(
-				FValueTransformerListWithBoundMapCollectionJoinOp(),
-				[&PerValueLayerWeights, LayerWeight, FloatValueType, &Output](FRawTransformerFunc Transformer, const FBoundValueMap* BaseMap, const FBoundValueMap* LayerMap)
-				{
-					if (BaseMap == nullptr)
-					{
-						if (LayerMap == nullptr)
-						{
-							// No maps found for the transformer type
-						}
-						else
-						{
-							// This set map doesn't have a matching pair, duplicate it
-							Output.AppendSetMap(LayerMap->Duplicate(Output.GetAllocator()));
-						}
-					}
-					else
+			{
+				const FBoundMapCollection& PerValueLayerWeightMaps = PerValueLayerWeights.GetBoundValueMaps();
+				FBoundMapCollection& OutputMaps = Output.GetBoundValueMaps();
+
+				OuterJoinBy(
+					FValueTransformerListWithBoundMapCollectionJoinOp(),
+					[&PerValueLayerWeightMaps, LayerWeight, FloatValueType, &OutputMaps](FRawTransformerFunc Transformer, const FBoundValueMap* BaseMap, const FBoundValueMap* LayerMap)
 					{
-						if (LayerMap == nullptr)
+						if (BaseMap == nullptr)
 						{
-							// This set map doesn't have a matching pair, duplicate it
-							Output.AppendSetMap(BaseMap->Duplicate(Output.GetAllocator()));
+							if (LayerMap == nullptr)
+							{
+								// No maps found for the transformer type
+							}
+							else
+							{
+								// This set map doesn't have a matching pair, duplicate it
+								OutputMaps.Append(LayerMap->Duplicate(OutputMaps.GetAllocator()));
+							}
 						}
 						else
 						{
-							if (Transformer == nullptr)
+							if (LayerMap == nullptr)
 							{
-								// No transformer was found for this value type, retain the one with the largest weight
-								FBoundValueMap* MapCopy;
-								if (LayerWeight >= 0.5f)
+								// This set map doesn't have a matching pair, duplicate it
+								OutputMaps.Append(BaseMap->Duplicate(OutputMaps.GetAllocator()));
+							}
+							else
+							{
+								if (Transformer == nullptr)
 								{
-									MapCopy = LayerMap->Duplicate(Output.GetAllocator());
+									// No transformer was found for this value type, retain the one with the largest weight
+									FBoundValueMap* MapCopy;
+									if (LayerWeight >= 0.5f)
+									{
+										MapCopy = LayerMap->Duplicate(OutputMaps.GetAllocator());
+									}
+									else
+									{
+										MapCopy = BaseMap->Duplicate(OutputMaps.GetAllocator());
+									}
+
+									OutputMaps.Append(MapCopy);
 								}
 								else
 								{
-									MapCopy = BaseMap->Duplicate(Output.GetAllocator());
-								}
-
-								Output.AppendSetMap(MapCopy);
-							}
-							else
-							{
-								// Transform our matching pair
-								FBoundValueMap* OutputMap = nullptr;
+									// Transform our matching pair
+									FBoundValueMap* OutputMap = nullptr;
 
-								// We cannot join our weights because their value type differs
-								const FAttributeMappingKey MappingKey = BaseMap->GetMappingKey().To(FloatValueType);
-								const FBoundValueMap* PerValueWeightMap = PerValueLayerWeights.FindSetMap(MappingKey);
+									// We cannot join our weights because their value type differs
+									const FAttributeMappingKey MappingKey = BaseMap->GetMappingKey().To(FloatValueType);
+									const FBoundValueMap* PerValueWeightMap = PerValueLayerWeightMaps.Find(MappingKey);
 
-								FTransformBoundValueMapFunc TypedTransformer = reinterpret_cast<FTransformBoundValueMapFunc>(Transformer);
-								(*TypedTransformer)(BaseMap, LayerMap, PerValueWeightMap, LayerWeight, OutputMap, Output.GetAllocator());
+									FTransformBoundValueMapFunc TypedTransformer = reinterpret_cast<FTransformBoundValueMapFunc>(Transformer);
+									(*TypedTransformer)(BaseMap, LayerMap, PerValueWeightMap, LayerWeight, OutputMap, OutputMaps.GetAllocator());
 
-								checkf(OutputMap != nullptr, TEXT("Transformer function is responsible for allocating an output buffer"));
-								checkf(OutputMap->GetAllocator() == Output.GetAllocator(), TEXT("Transformer function is responsible for allocating using the output allocator"));
-								Output.AppendSetMap(OutputMap);
+									checkf(OutputMap != nullptr, TEXT("Transformer function is responsible for allocating an output buffer"));
+									checkf(OutputMap->GetAllocator() == OutputMaps.GetAllocator(), TEXT("Transformer function is responsible for allocating using the output allocator"));
+									OutputMaps.Append(OutputMap);
+								}
 							}
 						}
-					}
-				},
-				TransformerList->CreateBoundValueMapTransformerIterator(),
-				Base.CreateConstSetMapIterator(),
-				LayerPtr->CreateConstSetMapIterator());
+					},
+					TransformerList->CreateBoundValueMapTransformerIterator(),
+					Base.GetBoundValueMaps().CreateConstIterator(),
+					LayerPtr->GetBoundValueMaps().CreateConstIterator());
+			}
 
 			// Name/value maps
-			OuterJoinBy(
-				FValueTransformerListWithUnboundMapCollectionJoinOp(),
-				[LayerWeight, &Output](FRawTransformerFunc Transformer, const FUnboundValueMap* BaseMap, const FUnboundValueMap* LayerMap)
-				{
-					if (BaseMap == nullptr)
-					{
-						if (LayerMap == nullptr)
-						{
-							// No maps found for the transformer type
-						}
-						else
-						{
-							// This map doesn't have a matching pair, duplicate it
-							Output.AppendMap(LayerMap->Duplicate(Output.GetAllocator()));
-						}
-					}
-					else
+			{
+				FUnboundMapCollection& OutputMaps = Output.GetUnboundValueMaps();
+
+				OuterJoinBy(
+					FValueTransformerListWithUnboundMapCollectionJoinOp(),
+					[LayerWeight, &OutputMaps](FRawTransformerFunc Transformer, const FUnboundValueMap* BaseMap, const FUnboundValueMap* LayerMap)
 					{
-						if (LayerMap == nullptr)
+						if (BaseMap == nullptr)
 						{
-							// This map doesn't have a matching pair
-							// Our output is A and thus already contains this
+							if (LayerMap == nullptr)
+							{
+								// No maps found for the transformer type
+							}
+							else
+							{
+								// This map doesn't have a matching pair, duplicate it
+								OutputMaps.Append(LayerMap->Duplicate(OutputMaps.GetAllocator()));
+							}
 						}
 						else
 						{
-							if (Transformer == nullptr)
+							if (LayerMap == nullptr)
 							{
-								// No transformer was found for this value type, retain the one with the largest weight
-								FUnboundValueMap* MapCopy;
-								if (LayerWeight >= 0.5f)
+								// This map doesn't have a matching pair
+								// Our output is A and thus already contains this
+							}
+							else
+							{
+								if (Transformer == nullptr)
 								{
-									MapCopy = LayerMap->Duplicate(Output.GetAllocator());
+									// No transformer was found for this value type, retain the one with the largest weight
+									FUnboundValueMap* MapCopy;
+									if (LayerWeight >= 0.5f)
+									{
+										MapCopy = LayerMap->Duplicate(OutputMaps.GetAllocator());
+									}
+									else
+									{
+										MapCopy = BaseMap->Duplicate(OutputMaps.GetAllocator());
+									}
+
+									OutputMaps.Append(MapCopy);
 								}
 								else
 								{
-									MapCopy = BaseMap->Duplicate(Output.GetAllocator());
-								}
-
-								Output.AppendMap(MapCopy);
-							}
-							else
-							{
-								// Transform our matching pair in place
-								FUnboundValueMap* OutputMap = nullptr;
+									// Transform our matching pair in place
+									FUnboundValueMap* OutputMap = nullptr;
 
-								FTransformUnboundValueMapFunc TypedTransformer = reinterpret_cast<FTransformUnboundValueMapFunc>(Transformer);
-								(*TypedTransformer)(BaseMap, LayerMap, LayerWeight, OutputMap, Output.GetAllocator());
+									FTransformUnboundValueMapFunc TypedTransformer = reinterpret_cast<FTransformUnboundValueMapFunc>(Transformer);
+									(*TypedTransformer)(BaseMap, LayerMap, LayerWeight, OutputMap, OutputMaps.GetAllocator());
 
-								checkf(OutputMap != nullptr, TEXT("Transformer function is responsible for allocating an output buffer"));
-								checkf(OutputMap->GetAllocator() == Output.GetAllocator(), TEXT("Transformer function is responsible for allocating using the output allocator"));
-								Output.AppendMap(OutputMap);
+									checkf(OutputMap != nullptr, TEXT("Transformer function is responsible for allocating an output buffer"));
+									checkf(OutputMap->GetAllocator() == OutputMaps.GetAllocator(), TEXT("Transformer function is responsible for allocating using the output allocator"));
+									OutputMaps.Append(OutputMap);
+								}
 							}
 						}
-					}
-				},
-				TransformerList->CreateUnboundValueMapTransformerIterator(),
-				Base.CreateConstMapIterator(),
-				LayerPtr->CreateConstMapIterator());
+					},
+					TransformerList->CreateUnboundValueMapTransformerIterator(),
+					Base.GetUnboundValueMaps().CreateConstIterator(),
+					LayerPtr->GetUnboundValueMaps().CreateConstIterator());
+			}
 		}
 	}
 
diff --git a/Engine/Plugins/Experimental/UAF/UAF/Source/UAF/Private/UAF/ValueRuntime/Transformers/Overwrite.cpp b/Engine/Plugins/Experimental/UAF/UAF/Source/UAF/Private/UAF/ValueRuntime/Transformers/Overwrite.cpp
--- a/Engine/Plugins/Experimental/UAF/UAF/Source/UAF/Private/UAF/ValueRuntime/Transformers/Overwrite.cpp
+++ b/Engine/Plugins/Experimental/UAF/UAF/Source/UAF/Private/UAF/ValueRuntime/Transformers/Overwrite.cpp
@@ -8,7 +8,7 @@ namespace UE::UAF::Transformers
 {
 	FName FOverwrite::TransformerName = TEXT("UAF::Transformers::Overwrite");
 
-	void FOverwrite::Apply(const FValueTransformerMapPtr& TransformerMap, const FAttributeCollection& Input, float Weight, FAttributeCollection& Output)
+	void FOverwrite::Apply(const FValueTransformerMapPtr& TransformerMap, const FValueBundle& Input, float Weight, FValueBundle& Output)
 	{
 		const bool bIsInputOutput = &Input == &Output;
 		if (bIsInputOutput && FAnimWeight::IsFullWeight(Weight))
@@ -18,11 +18,11 @@ namespace UE::UAF::Transformers
 		}
 		
 		// Allocator is irrelevant since it will remain empty
-		FAttributeCollectionHeap EmptyCollection(Input.GetNamedSet());
+		FValueBundleHeap EmptyCollection(Input.GetNamedSet());
 		FOverwrite::Apply(TransformerMap, Input, EmptyCollection, Weight, Output);
 	}
 
-	void FOverwrite::Apply(const FValueTransformerMapPtr& TransformerMap, const FAttributeCollection& Input, const FAttributeCollection& PerValueWeights, float DefaultWeight, FAttributeCollection& Output)
+	void FOverwrite::Apply(const FValueTransformerMapPtr& TransformerMap, const FValueBundle& Input, const FValueBundle& PerValueWeights, float DefaultWeight, FValueBundle& Output)
 	{
 		checkf(Input.GetNamedSet() == PerValueWeights.GetNamedSet(), TEXT("Input and per-value weight collections must use the same named set"));
 
@@ -41,11 +41,13 @@ namespace UE::UAF::Transformers
 			// Our output collection matches one of our inputs, we'll modify it in place
 			FReallocFun OutputAllocator = Output.GetAllocator();
 
+			const FBoundMapCollection& PerValueWeightMaps = PerValueWeights.GetBoundValueMaps();
+
 			// Bound value maps
 			{
 				OuterJoinBy(
 					FValueTransformerListWithBoundMapCollectionJoinOp(),
-					[&PerValueWeights, DefaultWeight, FloatValueType, OutputAllocator](FRawTransformerFunc Transformer, FBoundValueMap* InputMap)
+					[&PerValueWeightMaps, DefaultWeight, FloatValueType, OutputAllocator](FRawTransformerFunc Transformer, FBoundValueMap* InputMap)
 					{
 						if (InputMap == nullptr)
 						{
@@ -64,7 +66,7 @@ namespace UE::UAF::Transformers
 
 								// We cannot join our weights because their value type differs
 								const FAttributeMappingKey MappingKey = InputMap->GetMappingKey().To(FloatValueType);
-								const FBoundValueMap* PerValueWeightMap = PerValueWeights.FindSetMap(MappingKey);
+								const FBoundValueMap* PerValueWeightMap = PerValueWeightMaps.Find(MappingKey);
 
 								FTransformBoundValueMapFunc TypedTransformer = reinterpret_cast<FTransformBoundValueMapFunc>(Transformer);
 								(*TypedTransformer)(InputMap, PerValueWeightMap, DefaultWeight, OutputMap, nullptr);
@@ -73,10 +75,10 @@ namespace UE::UAF::Transformers
 					},
 					TransformerList->CreateBoundValueMapTransformerIterator(),
 					// Our input is our output
-					Output.CreateSetMapIterator());
+					Output.GetBoundValueMaps().CreateIterator());
 			}
 
-			// Name/value maps
+			// Unbound value maps
 			{
 				OuterJoinBy(
 					FValueTransformerListWithUnboundMapCollectionJoinOp(),
@@ -104,7 +106,7 @@ namespace UE::UAF::Transformers
 					},
 					TransformerList->CreateUnboundValueMapTransformerIterator(),
 					// Our input is our output
-					Output.CreateMapIterator());
+					Output.GetUnboundValueMaps().CreateIterator());
 			}
 		}
 		else
@@ -114,74 +116,83 @@ namespace UE::UAF::Transformers
 			Output.SetValueSpace(Input.GetValueSpace());
 
 			// Bound value maps
-			OuterJoinBy(
-				FValueTransformerListWithBoundMapCollectionJoinOp(),
-				[&PerValueWeights, DefaultWeight, FloatValueType, &Output](FRawTransformerFunc Transformer, const FBoundValueMap* InputMap)
-				{
-					if (InputMap == nullptr)
-					{
-						// No maps found for the transformer type
-					}
-					else
+			{
+				const FBoundMapCollection& PerValueWeightMaps = PerValueWeights.GetBoundValueMaps();
+				FBoundMapCollection& OutputMaps = Output.GetBoundValueMaps();
+
+				OuterJoinBy(
+					FValueTransformerListWithBoundMapCollectionJoinOp(),
+					[&PerValueWeightMaps, DefaultWeight, FloatValueType, &OutputMaps](FRawTransformerFunc Transformer, const FBoundValueMap* InputMap)
 					{
-						if (Transformer == nullptr)
+						if (InputMap == nullptr)
 						{
-							// No transformer was found for this value type, retain it as-is
-							Output.AppendSetMap(InputMap->Duplicate(Output.GetAllocator()));
+							// No maps found for the transformer type
 						}
 						else
 						{
-							// Transform our matching pair
-							FBoundValueMap* OutputMap = nullptr;
+							if (Transformer == nullptr)
+							{
+								// No transformer was found for this value type, retain it as-is
+								OutputMaps.Append(InputMap->Duplicate(OutputMaps.GetAllocator()));
+							}
+							else
+							{
+								// Transform our matching pair
+								FBoundValueMap* OutputMap = nullptr;
 
-							// We cannot join our weights because their value type differs
-							const FAttributeMappingKey MappingKey = InputMap->GetMappingKey().To(FloatValueType);
-							const FBoundValueMap* PerValueWeightMap = PerValueWeights.FindSetMap(MappingKey);
+								// We cannot join our weights because their value type differs
+								const FAttributeMappingKey MappingKey = InputMap->GetMappingKey().To(FloatValueType);
+								const FBoundValueMap* PerValueWeightMap = PerValueWeightMaps.Find(MappingKey);
 
-							FTransformBoundValueMapFunc TypedTransformer = reinterpret_cast<FTransformBoundValueMapFunc>(Transformer);
-							(*TypedTransformer)(InputMap, PerValueWeightMap, DefaultWeight, OutputMap, Output.GetAllocator());
+								FTransformBoundValueMapFunc TypedTransformer = reinterpret_cast<FTransformBoundValueMapFunc>(Transformer);
+								(*TypedTransformer)(InputMap, PerValueWeightMap, DefaultWeight, OutputMap, OutputMaps.GetAllocator());
 
-							checkf(OutputMap != nullptr, TEXT("Transformer function is responsible for allocating an output buffer"));
-							checkf(OutputMap->GetAllocator() == Output.GetAllocator(), TEXT("Transformer function is responsible for allocating using the output allocator"));
-							Output.AppendSetMap(OutputMap);
+								checkf(OutputMap != nullptr, TEXT("Transformer function is responsible for allocating an output buffer"));
+								checkf(OutputMap->GetAllocator() == OutputMaps.GetAllocator(), TEXT("Transformer function is responsible for allocating using the output allocator"));
+								OutputMaps.Append(OutputMap);
+							}
 						}
-					}
-				},
-				TransformerList->CreateBoundValueMapTransformerIterator(),
-				Input.CreateConstSetMapIterator());
-
-			// Name/value maps
-			OuterJoinBy(
-				FValueTransformerListWithUnboundMapCollectionJoinOp(),
-				[DefaultWeight, &Output](FRawTransformerFunc Transformer, const FUnboundValueMap* InputMap)
-				{
-					if (InputMap == nullptr)
-					{
-						// No maps found for the transformer type
-					}
-					else
+					},
+					TransformerList->CreateBoundValueMapTransformerIterator(),
+					Input.GetBoundValueMaps().CreateConstIterator());
+			}
+
+			// Unbound value maps
+			{
+				FUnboundMapCollection& OutputMaps = Output.GetUnboundValueMaps();
+
+				OuterJoinBy(
+					FValueTransformerListWithUnboundMapCollectionJoinOp(),
+					[DefaultWeight, &OutputMaps](FRawTransformerFunc Transformer, const FUnboundValueMap* InputMap)
 					{
-						if (Transformer == nullptr)
+						if (InputMap == nullptr)
 						{
-							// No transformer was found for this value type, retain it as-is
-							Output.AppendMap(InputMap->Duplicate(Output.GetAllocator()));
+							// No maps found for the transformer type
 						}
 						else
 						{
-							// Transform our matching pair
-							FUnboundValueMap* OutputMap = nullptr;
+							if (Transformer == nullptr)
+							{
+								// No transformer was found for this value type, retain it as-is
+								OutputMaps.Append(InputMap->Duplicate(OutputMaps.GetAllocator()));
+							}
+							else
+							{
+								// Transform our matching pair
+								FUnboundValueMap* OutputMap = nullptr;
 
-							FTransformUnboundValueMapFunc TypedTransformer = reinterpret_cast<FTransformUnboundValueMapFunc>(Transformer);
-							(*TypedTransformer)(InputMap, DefaultWeight, OutputMap, Output.GetAllocator());
+								FTransformUnboundValueMapFunc TypedTransformer = reinterpret_cast<FTransformUnboundValueMapFunc>(Transformer);
+								(*TypedTransformer)(InputMap, DefaultWeight, OutputMap, OutputMaps.GetAllocator());
 
-							checkf(OutputMap != nullptr, TEXT("Transformer function is responsible for allocating an output buffer"));
-							checkf(OutputMap->GetAllocator() == Output.GetAllocator(), TEXT("Transformer function is responsible for allocating using the output allocator"));
-							Output.AppendMap(OutputMap);
+								checkf(OutputMap != nullptr, TEXT("Transformer function is responsible for allocating an output buffer"));
+								checkf(OutputMap->GetAllocator() == OutputMaps.GetAllocator(), TEXT("Transformer function is responsible for allocating using the output allocator"));
+								OutputMaps.Append(OutputMap);
+							}
 						}
-					}
-				},
-				TransformerList->CreateUnboundValueMapTransformerIterator(),
-				Input.CreateConstMapIterator());
+					},
+					TransformerList->CreateUnboundValueMapTransformerIterator(),
+					Input.GetUnboundValueMaps().CreateConstIterator());
+			}
 		}
 	}
 
diff --git a/Engine/Plugins/Experimental/UAF/UAF/Source/UAF/Private/UAF/ValueRuntime/Transformers/Sanitize.cpp b/Engine/Plugins/Experimental/UAF/UAF/Source/UAF/Private/UAF/ValueRuntime/Transformers/Sanitize.cpp
--- a/Engine/Plugins/Experimental/UAF/UAF/Source/UAF/Private/UAF/ValueRuntime/Transformers/Sanitize.cpp
+++ b/Engine/Plugins/Experimental/UAF/UAF/Source/UAF/Private/UAF/ValueRuntime/Transformers/Sanitize.cpp
@@ -8,7 +8,7 @@ namespace UE::UAF::Transformers
 {
 	FName FSanitize::TransformerName = TEXT("UAF::Transformers::Sanitize");
 
-	void FSanitize::Apply(const FValueTransformerMapPtr& TransformerMap, FAttributeCollection& Collection)
+	void FSanitize::Apply(const FValueTransformerMapPtr& TransformerMap, FValueBundle& Collection)
 	{
 		const FValueTransformerList* TransformerList = TransformerMap->Find(TransformerName);
 		if (TransformerList == nullptr)
@@ -26,7 +26,7 @@ namespace UE::UAF::Transformers
 				(*TypedTransformer)(Map);
 			},
 			TransformerList->CreateBoundValueMapTransformerIterator(),
-			Collection.CreateSetMapIterator());
+			Collection.GetBoundValueMaps().CreateIterator());
 
 		// Name/value maps
 		InnerJoinBy(
@@ -37,7 +37,7 @@ namespace UE::UAF::Transformers
 				(*TypedTransformer)(Map);
 			},
 			TransformerList->CreateUnboundValueMapTransformerIterator(),
-			Collection.CreateMapIterator());
+			Collection.GetUnboundValueMaps().CreateIterator());
 	}
 
 	namespace Private
diff --git a/Engine/Plugins/Experimental/UAF/UAF/Source/UAF/Public/UAF/AbstractSkeleton/AbstractSkeletonSetBinding.h b/Engine/Plugins/Experimental/UAF/UAF/Source/UAF/Public/UAF/AbstractSkeleton/AbstractSkeletonSetBinding.h
--- a/Engine/Plugins/Experimental/UAF/UAF/Source/UAF/Public/UAF/AbstractSkeleton/AbstractSkeletonSetBinding.h
+++ b/Engine/Plugins/Experimental/UAF/UAF/Source/UAF/Public/UAF/AbstractSkeleton/AbstractSkeletonSetBinding.h
@@ -7,7 +7,7 @@
 #include "Animation/BuiltInAttributeTypes.h"
 #include "UAF/Attributes/AttributeBindingIndex.h"
 #include "UAF/Attributes/AttributeBindingData.h"
-#include "UAF/ValueRuntime/AttributeCollection.h"
+#include "UAF/ValueRuntime/ValueBundle.h"
 
 #include "AbstractSkeletonSetBinding.generated.h"
 
diff --git a/Engine/Plugins/Experimental/UAF/UAF/Source/UAF/Public/UAF/Attributes/AttributeBindingData.h b/Engine/Plugins/Experimental/UAF/UAF/Source/UAF/Public/UAF/Attributes/AttributeBindingData.h
--- a/Engine/Plugins/Experimental/UAF/UAF/Source/UAF/Public/UAF/Attributes/AttributeBindingData.h
+++ b/Engine/Plugins/Experimental/UAF/UAF/Source/UAF/Public/UAF/Attributes/AttributeBindingData.h
@@ -14,7 +14,7 @@ class UAbstractSkeletonSetBinding;
 
 namespace UE::UAF
 {
-	using FAttributeCollectionPtr = TSharedPtr<class FAttributeCollection, ESPMode::ThreadSafe>;
+	using FAttributeCollectionPtr = TSharedPtr<class FValueBundle, ESPMode::ThreadSafe>;
 	using FAttributeBindingDataPtr = TRefCountPtr<const class FAttributeBindingData>;
 
 	// Our baked abstract skeleton binding data
diff --git a/Engine/Plugins/Experimental/UAF/UAF/Source/UAF/Public/UAF/Attributes/AttributeNamedSet.h b/Engine/Plugins/Experimental/UAF/UAF/Source/UAF/Public/UAF/Attributes/AttributeNamedSet.h
--- a/Engine/Plugins/Experimental/UAF/UAF/Source/UAF/Public/UAF/Attributes/AttributeNamedSet.h
+++ b/Engine/Plugins/Experimental/UAF/UAF/Source/UAF/Public/UAF/Attributes/AttributeNamedSet.h
@@ -13,7 +13,7 @@ namespace UE::UAF
 	class FAttributeBindingData;
 
 	using FAttributeBindingDataPtr = TRefCountPtr<const class FAttributeBindingData>;
-	using FAttributeCollectionPtr = TSharedPtr<class FAttributeCollection, ESPMode::ThreadSafe>;
+	using FAttributeCollectionPtr = TSharedPtr<class FValueBundle, ESPMode::ThreadSafe>;
 	using FAttributeNamedSetPtr = TRefCountPtr<const class FAttributeNamedSet>;
 
 	class FAttributeNamedSet
diff --git a/Engine/Plugins/Experimental/UAF/UAF/Source/UAF/Public/UAF/ValueRuntime/BoundValueMap.h b/Engine/Plugins/Experimental/UAF/UAF/Source/UAF/Public/UAF/ValueRuntime/BoundValueMap.h
--- a/Engine/Plugins/Experimental/UAF/UAF/Source/UAF/Public/UAF/ValueRuntime/BoundValueMap.h
+++ b/Engine/Plugins/Experimental/UAF/UAF/Source/UAF/Public/UAF/ValueRuntime/BoundValueMap.h
@@ -121,7 +121,7 @@ namespace UE::UAF
 		// The indirect allocator function pointer
 		FReallocFun ReallocFun = nullptr;
 
-		friend class FAttributeCollection;
+		friend class FValueBundle;
 		friend class FValueRuntimeRegistry;
 		friend void ReleaseBoundValueMap(FBoundValueMap* Map);
 	};
diff --git a/Engine/Plugins/Experimental/UAF/UAF/Source/UAF/Public/UAF/ValueRuntime/Transformers/Accumulate.h b/Engine/Plugins/Experimental/UAF/UAF/Source/UAF/Public/UAF/ValueRuntime/Transformers/Accumulate.h
--- a/Engine/Plugins/Experimental/UAF/UAF/Source/UAF/Public/UAF/ValueRuntime/Transformers/Accumulate.h
+++ b/Engine/Plugins/Experimental/UAF/UAF/Source/UAF/Public/UAF/ValueRuntime/Transformers/Accumulate.h
@@ -27,13 +27,13 @@ namespace UE::UAF::Transformers
 		// Accumulate attribute values
 		// Output = A + (B * WeightB)
 		// This function supports outputting to either input
-		static UE_API void Apply(const FValueTransformerMapPtr& TransformerMap, const FAttributeCollection& InputA, const FAttributeCollection& InputB, float WeightB, FAttributeCollection& Output);
+		static UE_API void Apply(const FValueTransformerMapPtr& TransformerMap, const FValueBundle& InputA, const FValueBundle& InputB, float WeightB, FValueBundle& Output);
 
 		// Accumulate attribute values
 		// If a weight isn't specified, the default weight will be used
 		// Output = A + (B * (HasValueWeightB ? ValueWeightB : DefaultWeightB))
 		// This function supports outputting to either input
-		static UE_API void Apply(const FValueTransformerMapPtr& TransformerMap, const FAttributeCollection& InputA, const FAttributeCollection& InputB, const FAttributeCollection& PerValueWeightsB, float DefaultWeightB, FAttributeCollection& Output);
+		static UE_API void Apply(const FValueTransformerMapPtr& TransformerMap, const FValueBundle& InputA, const FValueBundle& InputB, const FValueBundle& PerValueWeightsB, float DefaultWeightB, FValueBundle& Output);
 	};
 
 	//////////////////////////////////////////////////////////////////////////
diff --git a/Engine/Plugins/Experimental/UAF/UAF/Source/UAF/Public/UAF/ValueRuntime/Transformers/AdditiveSpace.h b/Engine/Plugins/Experimental/UAF/UAF/Source/UAF/Public/UAF/ValueRuntime/Transformers/AdditiveSpace.h
--- a/Engine/Plugins/Experimental/UAF/UAF/Source/UAF/Public/UAF/ValueRuntime/Transformers/AdditiveSpace.h
+++ b/Engine/Plugins/Experimental/UAF/UAF/Source/UAF/Public/UAF/ValueRuntime/Transformers/AdditiveSpace.h
@@ -24,9 +24,9 @@ namespace UE::UAF::Transformers
 		// Creates a new additive collection by subtracting the base from the specified Input
 		static UE_API void Apply(
 			const FValueTransformerMapPtr& TransformerMap,
-			const FAttributeCollection& Base,
-			const FAttributeCollection& Input,
-			FAttributeCollection& AdditiveOutput);
+			const FValueBundle& Base,
+			const FValueBundle& Input,
+			FValueBundle& AdditiveOutput);
 	};
 
 	// Attribute value transformer to apply additive values onto a base
@@ -46,10 +46,10 @@ namespace UE::UAF::Transformers
 		// This function supports outputting to either input
 		static UE_API void Apply(
 			const FValueTransformerMapPtr& TransformerMap,
-			const FAttributeCollection& Base,
-			const FAttributeCollection& Additive,
+			const FValueBundle& Base,
+			const FValueBundle& Additive,
 			float AdditiveWeight,
-			FAttributeCollection& Output);
+			FValueBundle& Output);
 
 		// Applies an additive collection onto a base by adding it using per value weights (must be of type FFloatAnimationAttribute)
 		// If a weight isn't specified, the default weight will be used
@@ -57,11 +57,11 @@ namespace UE::UAF::Transformers
 		// This function supports outputting to either input
 		static UE_API void Apply(
 			const FValueTransformerMapPtr& TransformerMap,
-			const FAttributeCollection& Base,
-			const FAttributeCollection& Additive,
-			const FAttributeCollection& PerValueWeights,
+			const FValueBundle& Base,
+			const FValueBundle& Additive,
+			const FValueBundle& PerValueWeights,
 			float DefaultWeight,
-			FAttributeCollection& Output);
+			FValueBundle& Output);
 	};
 
 	//////////////////////////////////////////////////////////////////////////
diff --git a/Engine/Plugins/Experimental/UAF/UAF/Source/UAF/Public/UAF/ValueRuntime/Transformers/BoneSpace.h b/Engine/Plugins/Experimental/UAF/UAF/Source/UAF/Public/UAF/ValueRuntime/Transformers/BoneSpace.h
--- a/Engine/Plugins/Experimental/UAF/UAF/Source/UAF/Public/UAF/ValueRuntime/Transformers/BoneSpace.h
+++ b/Engine/Plugins/Experimental/UAF/UAF/Source/UAF/Public/UAF/ValueRuntime/Transformers/BoneSpace.h
@@ -8,7 +8,7 @@
 
 namespace UE::UAF
 {
-	class FAttributeCollection;
+	class FValueBundle;
 }
 
 namespace UE::UAF::Transformers
@@ -19,31 +19,31 @@ namespace UE::UAF::Transformers
 	{
 		// Converts bone transforms between Local Space and Component Space
 		// These functions support outputting to their input in place
-		static UE_API void LocalToComponent(const FAttributeCollection& InputValues, FAttributeCollection& OutputValues);
-		static UE_API void ComponentToLocal(const FAttributeCollection& InputValues, FAttributeCollection& OutputValues);
+		static UE_API void LocalToComponent(const FValueBundle& InputValues, FValueBundle& OutputValues);
+		static UE_API void ComponentToLocal(const FValueBundle& InputValues, FValueBundle& OutputValues);
 
 		// Converts bone transforms between Local Space and Mesh Rotation Space
 		// Mesh Rotation Space assumes that bone transform scales are uniform and will yield incorrect results
 		// if non-uniform (aka 3D) scale is used. It transforms the rotations into component space by ignoring
 		// translation and scale. The output values will thus be in Mixed Space.
 		// These functions support outputting to their input in place
-		static UE_API void LocalToMeshRotation(const FAttributeCollection& InputValues, FAttributeCollection& OutputValues);
-		static UE_API void MeshRotationToLocal(const FAttributeCollection& InputValues, FAttributeCollection& OutputValues);
+		static UE_API void LocalToMeshRotation(const FValueBundle& InputValues, FValueBundle& OutputValues);
+		static UE_API void MeshRotationToLocal(const FValueBundle& InputValues, FValueBundle& OutputValues);
 
 		// Converts bone transforms between Local Space and Root Rotation Space
 		// Root Rotation Space is similar to Mesh Rotation space except that the root bone transform is ignored
 		// and assumed to be the identity transform (root corrections are ignored).
 		// These functions support outputting to their input in place
-		static UE_API void LocalToRootRotation(const FAttributeCollection& InputValues, FAttributeCollection& OutputValues);
-		static UE_API void RootRotationToLocal(const FAttributeCollection& InputValues, FAttributeCollection& OutputValues);
+		static UE_API void LocalToRootRotation(const FValueBundle& InputValues, FValueBundle& OutputValues);
+		static UE_API void RootRotationToLocal(const FValueBundle& InputValues, FValueBundle& OutputValues);
 
 		// Converts bone transforms between Local Space and Mesh Scale Space
 		// Mesh Scale Space assumes that bone transform scales are uniform and will yield incorrect results
 		// if non-uniform (aka 3D) scale is used. It transforms the scales into component space by ignoring
 		// rotation and translation. The output values will thus be in a Mixed Space.
 		// These functions support outputting to their input in place
-		static UE_API void LocalToMeshScale(const FAttributeCollection& InputValues, FAttributeCollection& OutputValues);
-		static UE_API void MeshScaleToLocal(const FAttributeCollection& InputValues, FAttributeCollection& OutputValues);
+		static UE_API void LocalToMeshScale(const FValueBundle& InputValues, FValueBundle& OutputValues);
+		static UE_API void MeshScaleToLocal(const FValueBundle& InputValues, FValueBundle& OutputValues);
 	};
 }
 
