UAF: Mark newly added variables and functions as public per default

#Synced-CL 49556441

diff --git a/Engine/Plugins/Experimental/UAF/UAF/Source/UAFUncookedOnly/Private/AnimNextRigVMAssetEditorData.cpp b/Engine/Plugins/Experimental/UAF/UAF/Source/UAFUncookedOnly/Private/AnimNextRigVMAssetEditorData.cpp
--- a/Engine/Plugins/Experimental/UAF/UAF/Source/UAFUncookedOnly/Private/AnimNextRigVMAssetEditorData.cpp
+++ b/Engine/Plugins/Experimental/UAF/UAF/Source/UAFUncookedOnly/Private/AnimNextRigVMAssetEditorData.cpp
@@ -2392,6 +2392,7 @@ UAnimNextVariableEntry* UAnimNextRigVMAssetEditorData::AddVariable(FName InName,
 
 		NewEntry->SetVariableName(NewParameterName, false);
 		NewEntry->SetType(InType, false);
+		NewEntry->SetExportAccessSpecifier(EAnimNextExportAccessSpecifier::Public, false);
 		if(InDefaultValue.Len() > 0)
 		{
 			NewEntry->SetDefaultValueFromString(InDefaultValue, false);
@@ -2727,7 +2728,11 @@ URigVMLibraryNode* UAnimNextRigVMAssetEditorData::AddFunction(FName InFunctionNa
 {
 	URigVMController* Controller = RigVMClient.GetOrCreateController(GetLocalFunctionLibrary());
 	URigVMLibraryNode* Node = Controller->AddFunctionToLibrary(InFunctionName, bInMutable, FVector2D::ZeroVector, bInSetupUndoRedo, false);
-
+	if (Node)
+	{
+		MarkFunctionPublic(Node->GetFName());
+	}
+	
 	if (bPrintPythonCommand)
 	{
 		RigVMPythonUtils::Print(GetName(),
