Enable HighlightParentNodesForSelection on skeleton tree inside set binding editor

#Synced-CL 49272195

diff --git a/Engine/Plugins/Experimental/UAF/UAF/Source/UAFEditor/Private/UAF/AbstractSkeleton/Sets/SSetSkeletonTree.cpp b/Engine/Plugins/Experimental/UAF/UAF/Source/UAFEditor/Private/UAF/AbstractSkeleton/Sets/SSetSkeletonTree.cpp
--- a/Engine/Plugins/Experimental/UAF/UAF/Source/UAFEditor/Private/UAF/AbstractSkeleton/Sets/SSetSkeletonTree.cpp
+++ b/Engine/Plugins/Experimental/UAF/UAF/Source/UAFEditor/Private/UAF/AbstractSkeleton/Sets/SSetSkeletonTree.cpp
@@ -53,6 +53,7 @@ namespace UE::UAF::Sets
 							.OnGenerateRow(this, &SSetsSkeletonTree::TreeView_OnGenerateRow)
 							.OnGetChildren(this, &SSetsSkeletonTree::TreeView_OnGetChildren)
 							.OnContextMenuOpening(this, &SSetsSkeletonTree::TreeView_OnContextMenuOpening)
+							.HighlightParentNodesForSelection(true)
 							.HeaderRow(
 								SNew(SHeaderRow)
 								+ SHeaderRow::Column("Set")
