Enable multi-select in set binding tree view and merge context menu action

#Synced-CL 49272195

diff --git a/Engine/Plugins/Experimental/UAF/UAF/Source/UAF/Private/UAF/AbstractSkeleton/AbstractSkeletonSetBinding.cpp b/Engine/Plugins/Experimental/UAF/UAF/Source/UAF/Private/UAF/AbstractSkeleton/AbstractSkeletonSetBinding.cpp
--- a/Engine/Plugins/Experimental/UAF/UAF/Source/UAF/Private/UAF/AbstractSkeleton/AbstractSkeletonSetBinding.cpp
+++ b/Engine/Plugins/Experimental/UAF/UAF/Source/UAF/Private/UAF/AbstractSkeleton/AbstractSkeletonSetBinding.cpp
@@ -65,6 +65,8 @@ bool UAbstractSkeletonSetBinding::AddBoneToSet(const FName BoneName, const FName
 		return false; // Already exists in a set
 	}
 
+	Modify();
+
 	FAbstractSkeleton_BoneBinding& Binding = BoneBindings.AddDefaulted_GetRef();
 	Binding.BoneName = BoneName;
 	Binding.SetName = SetName;
@@ -82,6 +84,8 @@ bool UAbstractSkeletonSetBinding::RemoveBoneFromSet(const FName BoneName)
 
 	if (BindingIndex != INDEX_NONE)
 	{
+		Modify();
+
 		BoneBindings.RemoveAtSwap(BindingIndex);
 		OnBindingsChanged();
 		return true;
@@ -151,6 +155,8 @@ bool UAbstractSkeletonSetBinding::AddAttributeToSet(const FAnimationAttributeIde
 
 		if (ensure(UnboundAttribute))
 		{
+			Modify();
+
 			UnboundAttribute->SetName = SetName;
 			OnBindingsChanged();
 			return true;
@@ -161,6 +167,8 @@ bool UAbstractSkeletonSetBinding::AddAttributeToSet(const FAnimationAttributeIde
 		}
 	}
 
+	Modify();
+
 	FAbstractSkeleton_AttributeBinding& Binding = AttributeBindings.AddDefaulted_GetRef();
 	Binding.Attribute = Attribute;
 	Binding.SetName = SetName;
@@ -178,6 +186,8 @@ bool UAbstractSkeletonSetBinding::RemoveAttributeFromSet(const FAnimationAttribu
 
 	if (BindingIndex != INDEX_NONE)
 	{
+		Modify();
+
 		FAbstractSkeleton_AttributeBinding& Binding = AttributeBindings[BindingIndex];
 
 		if (Binding.SetName.IsNone())
@@ -200,6 +210,8 @@ bool UAbstractSkeletonSetBinding::RemoveAttributeFromSet(const FAnimationAttribu
 
 bool UAbstractSkeletonSetBinding::RemoveAllFromSet(const FName SetName)
 {
+	Modify();
+
 	int32 bCountRemoved = 0;
 
 	bCountRemoved += BoneBindings.RemoveAll([&](const FAbstractSkeleton_BoneBinding& Binding)
diff --git a/Engine/Plugins/Experimental/UAF/UAF/Source/UAFEditor/Private/UAF/AbstractSkeleton/Sets/SSetBinding.cpp b/Engine/Plugins/Experimental/UAF/UAF/Source/UAFEditor/Private/UAF/AbstractSkeleton/Sets/SSetBinding.cpp
--- a/Engine/Plugins/Experimental/UAF/UAF/Source/UAFEditor/Private/UAF/AbstractSkeleton/Sets/SSetBinding.cpp
+++ b/Engine/Plugins/Experimental/UAF/UAF/Source/UAFEditor/Private/UAF/AbstractSkeleton/Sets/SSetBinding.cpp
@@ -10,6 +10,7 @@
 #include "Styling/StyleColors.h" 
 #include "UAF/AbstractSkeleton/AbstractSkeletonSetBinding.h"
 #include "UAF/AbstractSkeleton/AbstractSkeletonSetCollection.h"
+#include "UAF/AbstractSkeleton/Sets/Commands.h"
 #include "UAF/AbstractSkeleton/Sets/SetDragDrop.h"
 #include "Widgets/Input/SSearchBox.h"
 #include "Widgets/Layout/SWidgetSwitcher.h"
@@ -74,7 +75,7 @@ namespace UE::UAF::Sets
 					.Padding(FMargin(1.0f, 1.0f))
 					[
 						SAssignNew(TreeView, STreeView<TSharedPtr<ITreeItem>>)
-							.SelectionMode(ESelectionMode::Single)
+							.SelectionMode(ESelectionMode::Multi)
 							.HighlightParentNodesForSelection(true)
 							.TreeItemsSource(&RootItems)
 							.OnGenerateRow(this, &SSetBinding::TreeView_OnGenerateRow)
@@ -127,22 +128,68 @@ namespace UE::UAF::Sets
 
 	TSharedPtr<SWidget> SSetBinding::TreeView_OnContextMenuOpening()
 	{
+		const Editor::FSetBindingEditorCommands& Commands = Editor::FSetBindingEditorCommands::Get();
+
 		FMenuBuilder MenuBuilder(true, CommandList.ToSharedRef());
 
+		MenuBuilder.AddMenuEntry(Commands.UnbindSelection);
+
+		return MenuBuilder.MakeWidget();
+	}
+
+	bool SSetBinding::CanUnbindSelection() const
+	{
 		TArray<TSharedPtr<ITreeItem>> Selection = TreeView->GetSelectedItems();
-		if (Selection.Num() > 0)
+		
+		for (const TSharedPtr<ITreeItem>& Selected : Selection)
 		{
-			TSharedPtr<ITreeItem> Selected = Selection[0]; // Single selection only
+			if (Selected->GetType() == FTreeItem_Bone::StaticGetType() || Selected->GetType() == FTreeItem_Attribute::StaticGetType())
+			{
+				return true;
+			}
+			else
+			{
+				TSharedPtr<FTreeItem_Set> SetItem = StaticCastSharedPtr<FTreeItem_Set>(Selected);
+				if (!SetItem->Children.IsEmpty())
+				{
+					return true;
+				}
+			}
+		}
 
-			const bool bShowMenu = Selected->GenerateContextMenu(MenuBuilder);
+		return false;
+	}
+
+	void SSetBinding::HandleUnbindSelection()
+	{
+		const FScopedTransaction Transaction(LOCTEXT("UnbindSelection", "Unbind Selection"));
+		
+		TArray<TSharedPtr<ITreeItem>> Selection = TreeView->GetSelectedItems();
 
-			if (bShowMenu)
+		// Unbind all bones/attributes first
+		for (const TSharedPtr<ITreeItem>& Selected : Selection)
+		{
+			if (Selected->GetType() == FTreeItem_Bone::StaticGetType())
+			{
+				TSharedPtr<FTreeItem_Bone> BoneItem = StaticCastSharedPtr<FTreeItem_Bone>(Selected);
+				SetBinding->RemoveBoneFromSet(BoneItem->Binding.BoneName);
+			}
+			else if (Selected->GetType() == FTreeItem_Attribute::StaticGetType())
 			{
-				return MenuBuilder.MakeWidget();
+				TSharedPtr<FTreeItem_Attribute> AttributeItem = StaticCastSharedPtr<FTreeItem_Attribute>(Selected);
+				SetBinding->RemoveAttributeFromSet(AttributeItem->Binding.Attribute);
 			}
 		}
 
-		return SNullWidget::NullWidget;
+		// In a second pass, unbind all sets
+		for (const TSharedPtr<ITreeItem>& Selected : Selection)
+		{
+			if (Selected->GetType() == FTreeItem_Set::StaticGetType())
+			{
+				TSharedPtr<FTreeItem_Set> SetItem = StaticCastSharedPtr<FTreeItem_Set>(Selected);
+				SetBinding->RemoveAllFromSet(SetItem->Set.SetName);
+			}
+		}
 	}
 
 	void SSetBinding::PostUndo(bool bSuccess)
@@ -298,6 +345,12 @@ namespace UE::UAF::Sets
 	void SSetBinding::BindCommands()
 	{
 		CommandList = MakeShareable(new FUICommandList);
+
+		const Editor::FSetBindingEditorCommands& Commands = Editor::FSetBindingEditorCommands::Get();
+
+		CommandList->MapAction(Commands.UnbindSelection,
+			FExecuteAction::CreateSP(this, &SSetBinding::HandleUnbindSelection),
+			FCanExecuteAction::CreateSP(this, &SSetBinding::CanUnbindSelection));
 	}
 
 	TArray<TSharedPtr<SSetBinding::ITreeItem>> SSetBinding::GetAllTreeItems()
@@ -385,7 +438,6 @@ namespace UE::UAF::Sets
 						TSharedPtr<FTreeItem_Set> TargetSetTreeItem = StaticCastSharedPtr<FTreeItem_Set>(TargetTreeItem);
 
 						const FScopedTransaction Transaction(LOCTEXT("BindAttributeToSet", "Bind Attribute to Set"));
-						SetBindingWidget.SetBinding->Modify();
 
 						for (const TSharedPtr<SSetsSkeletonTree::FTreeItem>& Item : DragDropOp->Items)
 						{
@@ -401,7 +453,6 @@ namespace UE::UAF::Sets
 						TSharedPtr<FTreeItem_Set> TargetSetTreeItem = StaticCastSharedPtr<FTreeItem_Set>(TargetTreeItem);
 						
 						const FScopedTransaction Transaction(LOCTEXT("BindAttributeToSet", "Bind Attribute to Set"));
-						SetBindingWidget.SetBinding->Modify();
 
 						for (const FAnimationAttributeIdentifier& Attribute : DragDropOp->Attributes)
 						{
@@ -471,7 +522,6 @@ namespace UE::UAF::Sets
 					if (DragDropOp.IsValid())
 					{
 						const FScopedTransaction Transaction(LOCTEXT("AddAttributeToSet", "Add Attribute to Set"));
-						SetBindingWidget.SetBinding->Modify();
 
 						for (const TSharedPtr<SSetsSkeletonTree::FTreeItem>& Item : DragDropOp->Items)
 						{
@@ -532,7 +582,6 @@ namespace UE::UAF::Sets
 					if (DragDropOp.IsValid())
 					{
 						const FScopedTransaction Transaction(LOCTEXT("AddAttributeToSet", "Add Attribute to Set"));
-						SetBindingWidget.SetBinding->Modify();
 
 						for (const TSharedPtr<SSetsSkeletonTree::FTreeItem>& Item : DragDropOp->Items)
 						{
@@ -567,27 +616,6 @@ namespace UE::UAF::Sets
 			];
 	}
 
-	bool SSetBinding::FTreeItem_Set::GenerateContextMenu(FMenuBuilder& Menu)
-	{
-		FUIAction UnbindAll(
-			FExecuteAction::CreateLambda([&]()
-				{
-					FScopedTransaction Transaction(LOCTEXT("UnbindAllInSet", "Unbind All In Set"));
-					SetBindingWidget.SetBinding->Modify();
-
-					SetBindingWidget.SetBinding->RemoveAllFromSet(Set.SetName);
-				})
-		);
-
-		Menu.AddMenuEntry(
-			LOCTEXT("UnbindAllInSet_Label", "Unbind All In Set"),
-			LOCTEXT("UnbindAllInSet_Description", "Unbinds all attributes currently bound to this set."),
-			FSlateIcon(FAppStyle::GetAppStyleSetName(), "Icons.Unlink"),
-			UnbindAll);
-
-		return true;
-	}
-
 	void SSetBinding::FTreeItem_Set::SortChildren()
 	{
 		for (const auto& Child : Children)
@@ -631,48 +659,6 @@ namespace UE::UAF::Sets
 			});
 	}
 
-	bool SSetBinding::FTreeItem_Bone::GenerateContextMenu(FMenuBuilder& Menu)
-	{
-		FUIAction UnbindBone(
-			FExecuteAction::CreateLambda([&]()
-			{
-				FScopedTransaction Transaction(LOCTEXT("RemoveBoneFromSet", "Remove Bone From Set"));
-				SetBindingWidget.SetBinding->Modify();
-
-				SetBindingWidget.SetBinding->RemoveBoneFromSet(Binding.BoneName);
-			})
-		);
-
-		Menu.AddMenuEntry(
-			LOCTEXT("UnbindBone_Label", "Unbind Bone"),
-			LOCTEXT("UnbindBone_Description", "Unbind the selected bone."),
-			FSlateIcon(FAppStyle::GetAppStyleSetName(), "Icons.Unlink"),
-			UnbindBone);
-
-		return true;
-	}
-	
-	bool SSetBinding::FTreeItem_Attribute::GenerateContextMenu(FMenuBuilder& Menu)
-	{
-		FUIAction UnbindAttribute(
-			FExecuteAction::CreateLambda([&]()
-			{
-				FScopedTransaction Transaction(LOCTEXT("RemoveAttributeFromSet", "Remove Attribute From Set"));
-				SetBindingWidget.SetBinding->Modify();
-
-				SetBindingWidget.SetBinding->RemoveAttributeFromSet(Binding.Attribute);
-			})
-		);
-
-		Menu.AddMenuEntry(
-			LOCTEXT("UnbindAttribute_Label", "Unbind Attribute"),
-			LOCTEXT("UnbindAttribute_Description", "Unbind the selected attribute."),
-			FSlateIcon(FAppStyle::GetAppStyleSetName(), "Icons.Unlink"),
-			UnbindAttribute);
-
-		return true;
-	}
-
 }
 
 #undef LOCTEXT_NAMESPACE
\ No newline at end of file
diff --git a/Engine/Plugins/Experimental/UAF/UAF/Source/UAFEditor/Private/UAF/AbstractSkeleton/Sets/SSetBinding.h b/Engine/Plugins/Experimental/UAF/UAF/Source/UAFEditor/Private/UAF/AbstractSkeleton/Sets/SSetBinding.h
--- a/Engine/Plugins/Experimental/UAF/UAF/Source/UAFEditor/Private/UAF/AbstractSkeleton/Sets/SSetBinding.h
+++ b/Engine/Plugins/Experimental/UAF/UAF/Source/UAFEditor/Private/UAF/AbstractSkeleton/Sets/SSetBinding.h
@@ -43,7 +43,6 @@ namespace UE::UAF::Sets
 			virtual ETreeItem GetType() const = 0;
 			virtual void GetChildren(TArray<TSharedPtr<ITreeItem>>& OutChildren) = 0;
 			virtual TSharedRef<ITableRow> GenerateRow(const TSharedRef<STableViewBase>& OwnerTable) = 0;
-			virtual bool GenerateContextMenu(FMenuBuilder& Menu) = 0;
 			virtual void SortChildren() = 0;
 
 			SSetBinding& SetBindingWidget;
@@ -66,7 +65,6 @@ namespace UE::UAF::Sets
 			virtual ETreeItem GetType() const override { return StaticGetType(); }
 			virtual void GetChildren(TArray<TSharedPtr<ITreeItem>>& OutChildren) override;
 			virtual TSharedRef<ITableRow> GenerateRow(const TSharedRef<STableViewBase>& OwnerTable) override;
-			virtual bool GenerateContextMenu(FMenuBuilder& Menu) override;
 			virtual void SortChildren() override;
 
 			FAbstractSkeletonSet Set;
@@ -91,7 +89,6 @@ namespace UE::UAF::Sets
 			virtual ETreeItem GetType() const override { return StaticGetType(); }
 			virtual void GetChildren(TArray<TSharedPtr<ITreeItem>>& OutChildren) override;
 			virtual TSharedRef<ITableRow> GenerateRow(const TSharedRef<STableViewBase>& OwnerTable) override;
-			virtual bool GenerateContextMenu(FMenuBuilder& Menu) override;
 			virtual void SortChildren() override {};
 
 			FAbstractSkeleton_BoneBinding Binding;
@@ -113,7 +110,6 @@ namespace UE::UAF::Sets
 			virtual ETreeItem GetType() const override { return StaticGetType(); }
 			virtual void GetChildren(TArray<TSharedPtr<ITreeItem>>& OutChildren) override;
 			virtual TSharedRef<ITableRow> GenerateRow(const TSharedRef<STableViewBase>& OwnerTable) override;
-			virtual bool GenerateContextMenu(FMenuBuilder& Menu) override;
 			virtual void SortChildren() override {};
 
 			FAbstractSkeleton_AttributeBinding Binding;
@@ -149,6 +145,10 @@ namespace UE::UAF::Sets
 
 		TSharedPtr<SWidget> TreeView_OnContextMenuOpening();
 
+		bool CanUnbindSelection() const;
+
+		void HandleUnbindSelection();
+
 	public:
 		TWeakObjectPtr<UAbstractSkeletonSetBinding> SetBinding;
 
