Move abstract editor code to Editor namespace

#Synced-CL 49272195

diff --git a/Engine/Plugins/Experimental/UAF/UAF/Source/UAFEditor/Private/UAF/AbstractSkeleton/AbstractSkeletonEditor.cpp b/Engine/Plugins/Experimental/UAF/UAF/Source/UAFEditor/Private/UAF/AbstractSkeleton/AbstractSkeletonEditor.cpp
--- a/Engine/Plugins/Experimental/UAF/UAF/Source/UAFEditor/Private/UAF/AbstractSkeleton/AbstractSkeletonEditor.cpp
+++ b/Engine/Plugins/Experimental/UAF/UAF/Source/UAFEditor/Private/UAF/AbstractSkeleton/AbstractSkeletonEditor.cpp
@@ -22,7 +22,7 @@
 
 #define LOCTEXT_NAMESPACE "UE::UAF::FAbstractSkeletonEditor"
 
-namespace UE::UAF
+namespace UE::UAF::Editor
 {
 	namespace Tabs
 	{
@@ -168,7 +168,7 @@ namespace UE::UAF
 				TSharedRef<SDockTab> Tab = SNew(SDockTab);
 				SetsTab = Tab;
 
-				TSharedRef<SWidget> TabWidget = SNew(Sets::SSetsTab)
+				TSharedRef<SWidget> TabWidget = SNew(SSetsTab)
 					.ParentTab(Tab)
 					.SetBinding(SetBinding);
 
@@ -186,7 +186,7 @@ namespace UE::UAF
 				TSharedRef<SDockTab> Tab = SNew(SDockTab);
 				LabelsTab = Tab;
 
-				TSharedRef<SWidget> TabWidget = SNew(Labels::SLabelsTab)
+				TSharedRef<SWidget> TabWidget = SNew(SLabelsTab)
 					.ParentTab(Tab)
 					.LabelBinding(LabelBinding)
 					.AbstractSkeletonEditor(SharedThis(this));
diff --git a/Engine/Plugins/Experimental/UAF/UAF/Source/UAFEditor/Private/UAF/AbstractSkeleton/AbstractSkeletonEditor.h b/Engine/Plugins/Experimental/UAF/UAF/Source/UAFEditor/Private/UAF/AbstractSkeleton/AbstractSkeletonEditor.h
--- a/Engine/Plugins/Experimental/UAF/UAF/Source/UAFEditor/Private/UAF/AbstractSkeleton/AbstractSkeletonEditor.h
+++ b/Engine/Plugins/Experimental/UAF/UAF/Source/UAFEditor/Private/UAF/AbstractSkeleton/AbstractSkeletonEditor.h
@@ -10,7 +10,7 @@ class IPersonaToolkit;
 class UAbstractSkeletonLabelBinding;
 class UAbstractSkeletonSetBinding;
 
-namespace UE::UAF
+namespace UE::UAF::Editor
 {
 	class FAbstractSkeletonEditor : public FWorkflowCentricApplication, public IAbstractSkeletonEditor
 	{
diff --git a/Engine/Plugins/Experimental/UAF/UAF/Source/UAFEditor/Private/UAF/AbstractSkeleton/IAbstractSkeletonEditor.h b/Engine/Plugins/Experimental/UAF/UAF/Source/UAFEditor/Private/UAF/AbstractSkeleton/IAbstractSkeletonEditor.h
--- a/Engine/Plugins/Experimental/UAF/UAF/Source/UAFEditor/Private/UAF/AbstractSkeleton/IAbstractSkeletonEditor.h
+++ b/Engine/Plugins/Experimental/UAF/UAF/Source/UAFEditor/Private/UAF/AbstractSkeleton/IAbstractSkeletonEditor.h
@@ -6,7 +6,7 @@
 
 class IPersonaToolkit;
 
-namespace UE::UAF
+namespace UE::UAF::Editor
 {
 
 	class IAbstractSkeletonEditor
diff --git a/Engine/Plugins/Experimental/UAF/UAF/Source/UAFEditor/Private/UAF/AbstractSkeleton/Labels/BoneDragDrop.h b/Engine/Plugins/Experimental/UAF/UAF/Source/UAFEditor/Private/UAF/AbstractSkeleton/Labels/BoneDragDrop.h
--- a/Engine/Plugins/Experimental/UAF/UAF/Source/UAFEditor/Private/UAF/AbstractSkeleton/Labels/BoneDragDrop.h
+++ b/Engine/Plugins/Experimental/UAF/UAF/Source/UAFEditor/Private/UAF/AbstractSkeleton/Labels/BoneDragDrop.h
@@ -8,9 +8,9 @@
 
 class UAbstractSkeletonLabelBinding;
 
-#define LOCTEXT_NAMESPACE "UE::UAF::Labels::FBoneDragDropOp"
+#define LOCTEXT_NAMESPACE "UE::UAF::Editor::FBoneDragDropOp"
 
-namespace UE::UAF::Labels
+namespace UE::UAF::Editor
 {
 	class FBoneDragDropOp : public FDecoratedDragDropOp
 	{
diff --git a/Engine/Plugins/Experimental/UAF/UAF/Source/UAFEditor/Private/UAF/AbstractSkeleton/Labels/ILabelBinding.h b/Engine/Plugins/Experimental/UAF/UAF/Source/UAFEditor/Private/UAF/AbstractSkeleton/Labels/ILabelBinding.h
--- a/Engine/Plugins/Experimental/UAF/UAF/Source/UAFEditor/Private/UAF/AbstractSkeleton/Labels/ILabelBinding.h
+++ b/Engine/Plugins/Experimental/UAF/UAF/Source/UAFEditor/Private/UAF/AbstractSkeleton/Labels/ILabelBinding.h
@@ -7,7 +7,7 @@
 class UAbstractSkeletonLabelCollection;
 class UAbstractSkeletonLabelBinding;
 
-namespace UE::UAF::Labels
+namespace UE::UAF::Editor
 {
 
 	class ILabelBindingWidget
diff --git a/Engine/Plugins/Experimental/UAF/UAF/Source/UAFEditor/Private/UAF/AbstractSkeleton/Labels/ILabelSkeletonTree.h b/Engine/Plugins/Experimental/UAF/UAF/Source/UAFEditor/Private/UAF/AbstractSkeleton/Labels/ILabelSkeletonTree.h
--- a/Engine/Plugins/Experimental/UAF/UAF/Source/UAFEditor/Private/UAF/AbstractSkeleton/Labels/ILabelSkeletonTree.h
+++ b/Engine/Plugins/Experimental/UAF/UAF/Source/UAFEditor/Private/UAF/AbstractSkeleton/Labels/ILabelSkeletonTree.h
@@ -4,7 +4,7 @@
 
 #include "UObject/NameTypes.h"
 
-namespace UE::UAF::Labels
+namespace UE::UAF::Editor
 {
 
 	class ILabelSkeletonTreeWidget
diff --git a/Engine/Plugins/Experimental/UAF/UAF/Source/UAFEditor/Private/UAF/AbstractSkeleton/Labels/ILabelsTab.h b/Engine/Plugins/Experimental/UAF/UAF/Source/UAFEditor/Private/UAF/AbstractSkeleton/Labels/ILabelsTab.h
--- a/Engine/Plugins/Experimental/UAF/UAF/Source/UAFEditor/Private/UAF/AbstractSkeleton/Labels/ILabelsTab.h
+++ b/Engine/Plugins/Experimental/UAF/UAF/Source/UAFEditor/Private/UAF/AbstractSkeleton/Labels/ILabelsTab.h
@@ -4,13 +4,9 @@
 
 #include "Templates/SharedPointer.h"
 
-namespace UE::UAF
+namespace UE::UAF::Editor
 {
 	class IAbstractSkeletonEditor;
-}
-
-namespace UE::UAF::Labels
-{
 	class ILabelBindingWidget;
 	class ILabelSkeletonTreeWidget;
 
diff --git a/Engine/Plugins/Experimental/UAF/UAF/Source/UAFEditor/Private/UAF/AbstractSkeleton/Labels/LabelBindingAssetDefinition.cpp b/Engine/Plugins/Experimental/UAF/UAF/Source/UAFEditor/Private/UAF/AbstractSkeleton/Labels/LabelBindingAssetDefinition.cpp
--- a/Engine/Plugins/Experimental/UAF/UAF/Source/UAFEditor/Private/UAF/AbstractSkeleton/Labels/LabelBindingAssetDefinition.cpp
+++ b/Engine/Plugins/Experimental/UAF/UAF/Source/UAFEditor/Private/UAF/AbstractSkeleton/Labels/LabelBindingAssetDefinition.cpp
@@ -28,7 +28,7 @@ EAssetCommandResult UAssetDefinition_AbstractSkeletonLabelBinding::OpenAssets(co
 	{
 		TArray<UObject*> Assets = OpenArgs.LoadObjects<UObject>();
 
-		MakeShared<UE::UAF::FAbstractSkeletonEditor>()->InitEditor(Assets, OpenArgs.ToolkitHost);
+		MakeShared<UE::UAF::Editor::FAbstractSkeletonEditor>()->InitEditor(Assets, OpenArgs.ToolkitHost);
 	}
 
 	return EAssetCommandResult::Handled;
diff --git a/Engine/Plugins/Experimental/UAF/UAF/Source/UAFEditor/Private/UAF/AbstractSkeleton/Labels/LabelCollectionAssetDefinition.cpp b/Engine/Plugins/Experimental/UAF/UAF/Source/UAFEditor/Private/UAF/AbstractSkeleton/Labels/LabelCollectionAssetDefinition.cpp
--- a/Engine/Plugins/Experimental/UAF/UAF/Source/UAFEditor/Private/UAF/AbstractSkeleton/Labels/LabelCollectionAssetDefinition.cpp
+++ b/Engine/Plugins/Experimental/UAF/UAF/Source/UAFEditor/Private/UAF/AbstractSkeleton/Labels/LabelCollectionAssetDefinition.cpp
@@ -28,7 +28,7 @@ EAssetCommandResult UAssetDefinition_AbstractSkeletonLabelCollection::OpenAssets
 	{
 		TArray<UObject*> Assets = OpenArgs.LoadObjects<UObject>();
 		
-		MakeShared<UE::UAF::Labels::FLabelCollectionEditorToolkit>()->InitEditor(Assets);
+		MakeShared<UE::UAF::Editor::FLabelCollectionEditorToolkit>()->InitEditor(Assets);
 	}
 
 	return EAssetCommandResult::Handled;
diff --git a/Engine/Plugins/Experimental/UAF/UAF/Source/UAFEditor/Private/UAF/AbstractSkeleton/Labels/LabelCollectionEditor.cpp b/Engine/Plugins/Experimental/UAF/UAF/Source/UAFEditor/Private/UAF/AbstractSkeleton/Labels/LabelCollectionEditor.cpp
--- a/Engine/Plugins/Experimental/UAF/UAF/Source/UAFEditor/Private/UAF/AbstractSkeleton/Labels/LabelCollectionEditor.cpp
+++ b/Engine/Plugins/Experimental/UAF/UAF/Source/UAFEditor/Private/UAF/AbstractSkeleton/Labels/LabelCollectionEditor.cpp
@@ -8,7 +8,7 @@
 
 #define LOCTEXT_NAMESPACE "UE::UAF::FLabelCollectionEditorToolkit"
 
-namespace UE::UAF::Labels
+namespace UE::UAF::Editor
 {
 	static FName LabelsTabId("AbstractSkeletonLabelCollection_LabelsTab");
 
diff --git a/Engine/Plugins/Experimental/UAF/UAF/Source/UAFEditor/Private/UAF/AbstractSkeleton/Labels/LabelCollectionEditor.h b/Engine/Plugins/Experimental/UAF/UAF/Source/UAFEditor/Private/UAF/AbstractSkeleton/Labels/LabelCollectionEditor.h
--- a/Engine/Plugins/Experimental/UAF/UAF/Source/UAFEditor/Private/UAF/AbstractSkeleton/Labels/LabelCollectionEditor.h
+++ b/Engine/Plugins/Experimental/UAF/UAF/Source/UAFEditor/Private/UAF/AbstractSkeleton/Labels/LabelCollectionEditor.h
@@ -7,7 +7,7 @@
 
 class UAbstractSkeletonLabelCollection;
 
-namespace UE::UAF::Labels
+namespace UE::UAF::Editor
 {
 	class FLabelCollectionEditorToolkit : public FAssetEditorToolkit
 	{
diff --git a/Engine/Plugins/Experimental/UAF/UAF/Source/UAFEditor/Private/UAF/AbstractSkeleton/Labels/LabelDragDrop.h b/Engine/Plugins/Experimental/UAF/UAF/Source/UAFEditor/Private/UAF/AbstractSkeleton/Labels/LabelDragDrop.h
--- a/Engine/Plugins/Experimental/UAF/UAF/Source/UAFEditor/Private/UAF/AbstractSkeleton/Labels/LabelDragDrop.h
+++ b/Engine/Plugins/Experimental/UAF/UAF/Source/UAFEditor/Private/UAF/AbstractSkeleton/Labels/LabelDragDrop.h
@@ -7,7 +7,7 @@
 
 class UAbstractSkeletonLabelCollection;
 
-namespace UE::UAF::Labels
+namespace UE::UAF::Editor
 {
 	class FLabelDragDropOp : public FDecoratedDragDropOp
 	{
diff --git a/Engine/Plugins/Experimental/UAF/UAF/Source/UAFEditor/Private/UAF/AbstractSkeleton/Labels/SLabelBinding.cpp b/Engine/Plugins/Experimental/UAF/UAF/Source/UAFEditor/Private/UAF/AbstractSkeleton/Labels/SLabelBinding.cpp
--- a/Engine/Plugins/Experimental/UAF/UAF/Source/UAFEditor/Private/UAF/AbstractSkeleton/Labels/SLabelBinding.cpp
+++ b/Engine/Plugins/Experimental/UAF/UAF/Source/UAFEditor/Private/UAF/AbstractSkeleton/Labels/SLabelBinding.cpp
@@ -22,9 +22,9 @@
 #include "Widgets/Notifications/SProgressBar.h"
 #include "Widgets/Views/STreeView.h"
 
-#define LOCTEXT_NAMESPACE "UE::UAF::Labels::SLabelBinding"
+#define LOCTEXT_NAMESPACE "UE::UAF::Editor::SLabelBinding"
 
-namespace UE::UAF::Labels
+namespace UE::UAF::Editor
 {
 	void SLabelBinding::Construct(const FArguments& InArgs, TWeakObjectPtr<UAbstractSkeletonLabelBinding> InLabelBinding, TSharedPtr<ILabelsTab> InLabelsTabInterface)
 		{
diff --git a/Engine/Plugins/Experimental/UAF/UAF/Source/UAFEditor/Private/UAF/AbstractSkeleton/Labels/SLabelBinding.h b/Engine/Plugins/Experimental/UAF/UAF/Source/UAFEditor/Private/UAF/AbstractSkeleton/Labels/SLabelBinding.h
--- a/Engine/Plugins/Experimental/UAF/UAF/Source/UAFEditor/Private/UAF/AbstractSkeleton/Labels/SLabelBinding.h
+++ b/Engine/Plugins/Experimental/UAF/UAF/Source/UAFEditor/Private/UAF/AbstractSkeleton/Labels/SLabelBinding.h
@@ -16,7 +16,7 @@ template<typename> class STreeView;
 class UAbstractSkeletonLabelBinding;
 class UAbstractSkeletonLabelCollection;
 
-namespace UE::UAF::Labels
+namespace UE::UAF::Editor
 {
 	class ILabelsTab;
 
diff --git a/Engine/Plugins/Experimental/UAF/UAF/Source/UAFEditor/Private/UAF/AbstractSkeleton/Labels/SLabelCollection.cpp b/Engine/Plugins/Experimental/UAF/UAF/Source/UAFEditor/Private/UAF/AbstractSkeleton/Labels/SLabelCollection.cpp
--- a/Engine/Plugins/Experimental/UAF/UAF/Source/UAFEditor/Private/UAF/AbstractSkeleton/Labels/SLabelCollection.cpp
+++ b/Engine/Plugins/Experimental/UAF/UAF/Source/UAFEditor/Private/UAF/AbstractSkeleton/Labels/SLabelCollection.cpp
@@ -13,9 +13,9 @@
 #include "Widgets/Text/SRichTextBlock.h"
 #include "Widgets/Views/SListView.h"
 
-#define LOCTEXT_NAMESPACE "UE::UAF::Labels::SLabelCollection"
+#define LOCTEXT_NAMESPACE "UE::UAF::Editor::SLabelCollection"
 
-namespace UE::UAF::Labels
+namespace UE::UAF::Editor
 {
 	class FLabelCollectionDragDropOp : public FDecoratedDragDropOp
 	{
diff --git a/Engine/Plugins/Experimental/UAF/UAF/Source/UAFEditor/Private/UAF/AbstractSkeleton/Labels/SLabelCollection.h b/Engine/Plugins/Experimental/UAF/UAF/Source/UAFEditor/Private/UAF/AbstractSkeleton/Labels/SLabelCollection.h
--- a/Engine/Plugins/Experimental/UAF/UAF/Source/UAFEditor/Private/UAF/AbstractSkeleton/Labels/SLabelCollection.h
+++ b/Engine/Plugins/Experimental/UAF/UAF/Source/UAFEditor/Private/UAF/AbstractSkeleton/Labels/SLabelCollection.h
@@ -13,7 +13,7 @@ template<typename> class SListView;
 class SSearchBox;
 class STableViewBase;
 
-namespace UE::UAF::Labels
+namespace UE::UAF::Editor
 {
 	class SLabelCollection : public SCompoundWidget, public FSelfRegisteringEditorUndoClient
 	{
diff --git a/Engine/Plugins/Experimental/UAF/UAF/Source/UAFEditor/Private/UAF/AbstractSkeleton/Labels/SLabelSkeletonTree.cpp b/Engine/Plugins/Experimental/UAF/UAF/Source/UAFEditor/Private/UAF/AbstractSkeleton/Labels/SLabelSkeletonTree.cpp
--- a/Engine/Plugins/Experimental/UAF/UAF/Source/UAFEditor/Private/UAF/AbstractSkeleton/Labels/SLabelSkeletonTree.cpp
+++ b/Engine/Plugins/Experimental/UAF/UAF/Source/UAFEditor/Private/UAF/AbstractSkeleton/Labels/SLabelSkeletonTree.cpp
@@ -21,9 +21,9 @@
 #include "Widgets/Input/SHyperlink.h"
 #include "Widgets/Views/STreeView.h"
 
-#define LOCTEXT_NAMESPACE "UE::UAF::Labels::SLabelSkeletonTree"
+#define LOCTEXT_NAMESPACE "UE::UAF::Editor::SLabelSkeletonTree"
 
-namespace UE::UAF::Labels
+namespace UE::UAF::Editor
 {
 
 	void SLabelSkeletonTree::Construct(const FArguments& InArgs, TWeakObjectPtr<UAbstractSkeletonLabelBinding> InLabelBinding, TSharedPtr<ILabelsTab> InLabelsTabInterface)
diff --git a/Engine/Plugins/Experimental/UAF/UAF/Source/UAFEditor/Private/UAF/AbstractSkeleton/Labels/SLabelSkeletonTree.h b/Engine/Plugins/Experimental/UAF/UAF/Source/UAFEditor/Private/UAF/AbstractSkeleton/Labels/SLabelSkeletonTree.h
--- a/Engine/Plugins/Experimental/UAF/UAF/Source/UAFEditor/Private/UAF/AbstractSkeleton/Labels/SLabelSkeletonTree.h
+++ b/Engine/Plugins/Experimental/UAF/UAF/Source/UAFEditor/Private/UAF/AbstractSkeleton/Labels/SLabelSkeletonTree.h
@@ -15,7 +15,7 @@ template<typename> class STreeView;
 class UAbstractSkeletonLabelBinding;
 class UAbstractSkeletonLabelCollection;
 
-namespace UE::UAF::Labels
+namespace UE::UAF::Editor
 {
 	class ILabelsTab;
 
diff --git a/Engine/Plugins/Experimental/UAF/UAF/Source/UAFEditor/Private/UAF/AbstractSkeleton/Labels/SLabelsTab.cpp b/Engine/Plugins/Experimental/UAF/UAF/Source/UAFEditor/Private/UAF/AbstractSkeleton/Labels/SLabelsTab.cpp
--- a/Engine/Plugins/Experimental/UAF/UAF/Source/UAFEditor/Private/UAF/AbstractSkeleton/Labels/SLabelsTab.cpp
+++ b/Engine/Plugins/Experimental/UAF/UAF/Source/UAFEditor/Private/UAF/AbstractSkeleton/Labels/SLabelsTab.cpp
@@ -12,9 +12,9 @@
 #include "PropertyCustomizationHelpers.h"
 #include "Widgets/Docking/SDockTab.h"
 
-#define LOCTEXT_NAMESPACE "UE::UAF::Labels::SLabelsTab"
+#define LOCTEXT_NAMESPACE "UE::UAF::Editor::SLabelsTab"
 
-namespace UE::UAF::Labels
+namespace UE::UAF::Editor
 {
 	FName SLabelsTab::FTabs::SkeletonTreeId("AbstractSkeletonBindingEditor_Labels_SkeletonTreeTab");
 	FName SLabelsTab::FTabs::LabelBindingsId("AbstractSkeletonBindingEditor_Labels_BindingsTab");
diff --git a/Engine/Plugins/Experimental/UAF/UAF/Source/UAFEditor/Private/UAF/AbstractSkeleton/Labels/SLabelsTab.h b/Engine/Plugins/Experimental/UAF/UAF/Source/UAFEditor/Private/UAF/AbstractSkeleton/Labels/SLabelsTab.h
--- a/Engine/Plugins/Experimental/UAF/UAF/Source/UAFEditor/Private/UAF/AbstractSkeleton/Labels/SLabelsTab.h
+++ b/Engine/Plugins/Experimental/UAF/UAF/Source/UAFEditor/Private/UAF/AbstractSkeleton/Labels/SLabelsTab.h
@@ -10,13 +10,9 @@
 class UAbstractSkeletonLabelBinding;
 class SDockTab;
 
-namespace UE::UAF
+namespace UE::UAF::Editor
 {
 	class IAbstractSkeletonEditor;
-}
-
-namespace UE::UAF::Labels
-{
 	class SLabelBinding;
 	class SLabelSkeletonTree;
 
diff --git a/Engine/Plugins/Experimental/UAF/UAF/Source/UAFEditor/Private/UAF/AbstractSkeleton/Sets/AttributeDragDrop.cpp b/Engine/Plugins/Experimental/UAF/UAF/Source/UAFEditor/Private/UAF/AbstractSkeleton/Sets/AttributeDragDrop.cpp
--- a/Engine/Plugins/Experimental/UAF/UAF/Source/UAFEditor/Private/UAF/AbstractSkeleton/Sets/AttributeDragDrop.cpp
+++ b/Engine/Plugins/Experimental/UAF/UAF/Source/UAFEditor/Private/UAF/AbstractSkeleton/Sets/AttributeDragDrop.cpp
@@ -2,7 +2,7 @@
 
 #include "UAF/AbstractSkeleton/Sets/AttributeDragDrop.h"
 
-namespace UE::UAF::Sets
+namespace UE::UAF::Editor
 {
 	TSharedRef<FAttributeDragDropOp> FAttributeDragDropOp::New(TArray<FAnimationAttributeIdentifier>&& InAttributes)
 	{
@@ -23,7 +23,7 @@ namespace UE::UAF::Sets
 				SNew(STextBlock)
 				.Text(Attributes.Num() == 1
 					? FText::FromName(Attributes[0].GetName())
-					: FText::Format(NSLOCTEXT("UE::UAF::Sets::FAttributeDragDropOp", "MultipleAttributesTooltip", "{0} attributes"), FText::AsNumber(Attributes.Num())))
+					: FText::Format(NSLOCTEXT("UE::UAF::Editor::FAttributeDragDropOp", "MultipleAttributesTooltip", "{0} attributes"), FText::AsNumber(Attributes.Num())))
 			];
 	}
 }
\ No newline at end of file
diff --git a/Engine/Plugins/Experimental/UAF/UAF/Source/UAFEditor/Private/UAF/AbstractSkeleton/Sets/AttributeDragDrop.h b/Engine/Plugins/Experimental/UAF/UAF/Source/UAFEditor/Private/UAF/AbstractSkeleton/Sets/AttributeDragDrop.h
--- a/Engine/Plugins/Experimental/UAF/UAF/Source/UAFEditor/Private/UAF/AbstractSkeleton/Sets/AttributeDragDrop.h
+++ b/Engine/Plugins/Experimental/UAF/UAF/Source/UAFEditor/Private/UAF/AbstractSkeleton/Sets/AttributeDragDrop.h
@@ -6,7 +6,7 @@
 #include "DragAndDrop/DecoratedDragDropOp.h"
 #include "Animation/AnimData/AttributeIdentifier.h"
 
-namespace UE::UAF::Sets
+namespace UE::UAF::Editor
 {
 	class FAttributeDragDropOp : public FDecoratedDragDropOp
 	{
diff --git a/Engine/Plugins/Experimental/UAF/UAF/Source/UAFEditor/Private/UAF/AbstractSkeleton/Sets/SAttributesList.cpp b/Engine/Plugins/Experimental/UAF/UAF/Source/UAFEditor/Private/UAF/AbstractSkeleton/Sets/SAttributesList.cpp
--- a/Engine/Plugins/Experimental/UAF/UAF/Source/UAFEditor/Private/UAF/AbstractSkeleton/Sets/SAttributesList.cpp
+++ b/Engine/Plugins/Experimental/UAF/UAF/Source/UAFEditor/Private/UAF/AbstractSkeleton/Sets/SAttributesList.cpp
@@ -14,10 +14,9 @@
 #include "ToolMenus.h"
 #include "Widgets/Input/SSearchBox.h"
 
-#define LOCTEXT_NAMESPACE "UE::UAF::Sets::SAttributesList"
+#define LOCTEXT_NAMESPACE "UE::UAF::Editor::SAttributesList"
 
-// TODO: Move to UE::UAF::Editor
-using namespace UE::UAF::Sets;
+using namespace UE::UAF::Editor;
 
 static const FName ColumnId_Set("Set");
 static const FName ColumnId_Attribute("Attribute");
diff --git a/Engine/Plugins/Experimental/UAF/UAF/Source/UAFEditor/Private/UAF/AbstractSkeleton/Sets/SAttributesList.h b/Engine/Plugins/Experimental/UAF/UAF/Source/UAFEditor/Private/UAF/AbstractSkeleton/Sets/SAttributesList.h
--- a/Engine/Plugins/Experimental/UAF/UAF/Source/UAFEditor/Private/UAF/AbstractSkeleton/Sets/SAttributesList.h
+++ b/Engine/Plugins/Experimental/UAF/UAF/Source/UAFEditor/Private/UAF/AbstractSkeleton/Sets/SAttributesList.h
@@ -12,8 +12,7 @@
 
 class UAbstractSkeletonSetBinding;
 
-// TODO: Move to UE::UAF::Editor
-namespace UE::UAF::Sets
+namespace UE::UAF::Editor
 {
 
 	class SAttributesList : public SCompoundWidget, public FSelfRegisteringEditorUndoClient
@@ -78,5 +77,5 @@ class UAttributesListMenuContext : public UObject
 	GENERATED_BODY()
 	
 public:
-	TWeakPtr<UE::UAF::Sets::SAttributesList> SetBindingWidget;
+	TWeakPtr<UE::UAF::Editor::SAttributesList> SetBindingWidget;
 };
\ No newline at end of file
diff --git a/Engine/Plugins/Experimental/UAF/UAF/Source/UAFEditor/Private/UAF/AbstractSkeleton/Sets/SSetBinding.cpp b/Engine/Plugins/Experimental/UAF/UAF/Source/UAFEditor/Private/UAF/AbstractSkeleton/Sets/SSetBinding.cpp
--- a/Engine/Plugins/Experimental/UAF/UAF/Source/UAFEditor/Private/UAF/AbstractSkeleton/Sets/SSetBinding.cpp
+++ b/Engine/Plugins/Experimental/UAF/UAF/Source/UAFEditor/Private/UAF/AbstractSkeleton/Sets/SSetBinding.cpp
@@ -15,10 +15,9 @@
 #include "Widgets/Input/SSearchBox.h"
 #include "Widgets/Layout/SWidgetSwitcher.h"
 
-#define LOCTEXT_NAMESPACE "UE::UAF::Sets::SSetBinding"
+#define LOCTEXT_NAMESPACE "UE::UAF::Editor::SSetBinding"
 
-// TODO: Move to UE::UAF::Editor
-namespace UE::UAF::Sets
+namespace UE::UAF::Editor
 {
 	void SSetBinding::Construct(const FArguments& InArgs, TWeakObjectPtr<UAbstractSkeletonSetBinding> InSetBinding)
 	{
diff --git a/Engine/Plugins/Experimental/UAF/UAF/Source/UAFEditor/Private/UAF/AbstractSkeleton/Sets/SSetBinding.h b/Engine/Plugins/Experimental/UAF/UAF/Source/UAFEditor/Private/UAF/AbstractSkeleton/Sets/SSetBinding.h
--- a/Engine/Plugins/Experimental/UAF/UAF/Source/UAFEditor/Private/UAF/AbstractSkeleton/Sets/SSetBinding.h
+++ b/Engine/Plugins/Experimental/UAF/UAF/Source/UAFEditor/Private/UAF/AbstractSkeleton/Sets/SSetBinding.h
@@ -15,7 +15,7 @@ class FUICommandList;
 class SSearchBox;
 class UAbstractSkeletonSetBinding;
 
-namespace UE::UAF::Sets
+namespace UE::UAF::Editor
 {
 	class SSetBinding : public SCompoundWidget, public FSelfRegisteringEditorUndoClient
 	{
diff --git a/Engine/Plugins/Experimental/UAF/UAF/Source/UAFEditor/Private/UAF/AbstractSkeleton/Sets/SSetSkeletonTree.cpp b/Engine/Plugins/Experimental/UAF/UAF/Source/UAFEditor/Private/UAF/AbstractSkeleton/Sets/SSetSkeletonTree.cpp
--- a/Engine/Plugins/Experimental/UAF/UAF/Source/UAFEditor/Private/UAF/AbstractSkeleton/Sets/SSetSkeletonTree.cpp
+++ b/Engine/Plugins/Experimental/UAF/UAF/Source/UAFEditor/Private/UAF/AbstractSkeleton/Sets/SSetSkeletonTree.cpp
@@ -14,10 +14,9 @@
 #include "Widgets/Input/SSearchBox.h"
 #include "Widgets/Views/STreeView.h"
 
-#define LOCTEXT_NAMESPACE "UE::UAF::Sets::SSetsSkeletonTree"
+#define LOCTEXT_NAMESPACE "UE::UAF::Editor::SSetsSkeletonTree"
 
-// TODO: Move to UE::UAF::Editor
-namespace UE::UAF::Sets
+namespace UE::UAF::Editor
 {
 	void SSetsSkeletonTree::Construct(const FArguments& InArgs, TWeakObjectPtr<UAbstractSkeletonSetBinding> InSetBinding)
 	{
diff --git a/Engine/Plugins/Experimental/UAF/UAF/Source/UAFEditor/Private/UAF/AbstractSkeleton/Sets/SSetSkeletonTree.h b/Engine/Plugins/Experimental/UAF/UAF/Source/UAFEditor/Private/UAF/AbstractSkeleton/Sets/SSetSkeletonTree.h
--- a/Engine/Plugins/Experimental/UAF/UAF/Source/UAFEditor/Private/UAF/AbstractSkeleton/Sets/SSetSkeletonTree.h
+++ b/Engine/Plugins/Experimental/UAF/UAF/Source/UAFEditor/Private/UAF/AbstractSkeleton/Sets/SSetSkeletonTree.h
@@ -12,7 +12,7 @@ template <typename> class STreeView;
 class UAbstractSkeletonSetBinding;
 class FUICommandList;
 
-namespace UE::UAF::Sets
+namespace UE::UAF::Editor
 {
 	class SSetsSkeletonTree : public SCompoundWidget, public FSelfRegisteringEditorUndoClient
 	{
diff --git a/Engine/Plugins/Experimental/UAF/UAF/Source/UAFEditor/Private/UAF/AbstractSkeleton/Sets/SSetsTab.cpp b/Engine/Plugins/Experimental/UAF/UAF/Source/UAFEditor/Private/UAF/AbstractSkeleton/Sets/SSetsTab.cpp
--- a/Engine/Plugins/Experimental/UAF/UAF/Source/UAFEditor/Private/UAF/AbstractSkeleton/Sets/SSetsTab.cpp
+++ b/Engine/Plugins/Experimental/UAF/UAF/Source/UAFEditor/Private/UAF/AbstractSkeleton/Sets/SSetsTab.cpp
@@ -15,9 +15,9 @@
 #include "UAF/AbstractSkeleton/Sets/SAttributesList.h"
 #include "Widgets/Docking/SDockTab.h"
 
-#define LOCTEXT_NAMESPACE "UE::UAF::Labels::SSetsTab"
+#define LOCTEXT_NAMESPACE "UE::UAF::Editor::SSetsTab"
 
-namespace UE::UAF::Sets
+namespace UE::UAF::Editor
 {
 	namespace Tabs
 	{
diff --git a/Engine/Plugins/Experimental/UAF/UAF/Source/UAFEditor/Private/UAF/AbstractSkeleton/Sets/SSetsTab.h b/Engine/Plugins/Experimental/UAF/UAF/Source/UAFEditor/Private/UAF/AbstractSkeleton/Sets/SSetsTab.h
--- a/Engine/Plugins/Experimental/UAF/UAF/Source/UAFEditor/Private/UAF/AbstractSkeleton/Sets/SSetsTab.h
+++ b/Engine/Plugins/Experimental/UAF/UAF/Source/UAFEditor/Private/UAF/AbstractSkeleton/Sets/SSetsTab.h
@@ -9,7 +9,7 @@
 class SDockTab;
 class UAbstractSkeletonSetBinding;
 
-namespace UE::UAF::Sets
+namespace UE::UAF::Editor
 {
 	class SSetsSkeletonTree;
 	class SAttributesList;
diff --git a/Engine/Plugins/Experimental/UAF/UAF/Source/UAFEditor/Private/UAF/AbstractSkeleton/Sets/SetBindingAssetDefinition.cpp b/Engine/Plugins/Experimental/UAF/UAF/Source/UAFEditor/Private/UAF/AbstractSkeleton/Sets/SetBindingAssetDefinition.cpp
--- a/Engine/Plugins/Experimental/UAF/UAF/Source/UAFEditor/Private/UAF/AbstractSkeleton/Sets/SetBindingAssetDefinition.cpp
+++ b/Engine/Plugins/Experimental/UAF/UAF/Source/UAFEditor/Private/UAF/AbstractSkeleton/Sets/SetBindingAssetDefinition.cpp
@@ -28,7 +28,7 @@ EAssetCommandResult UAssetDefinition_AbstractSkeletonSetBinding::OpenAssets(cons
 	{
 		TArray<UObject*> Assets = OpenArgs.LoadObjects<UObject>();
 		
-		MakeShared<UE::UAF::FAbstractSkeletonEditor>()->InitEditor(Assets, OpenArgs.ToolkitHost);
+		MakeShared<UE::UAF::Editor::FAbstractSkeletonEditor>()->InitEditor(Assets, OpenArgs.ToolkitHost);
 	}
 
 	return EAssetCommandResult::Handled;
diff --git a/Engine/Plugins/Experimental/UAF/UAF/Source/UAFEditor/Private/UAF/AbstractSkeleton/Sets/SkeletonTreeDragDrop.h b/Engine/Plugins/Experimental/UAF/UAF/Source/UAFEditor/Private/UAF/AbstractSkeleton/Sets/SkeletonTreeDragDrop.h
--- a/Engine/Plugins/Experimental/UAF/UAF/Source/UAFEditor/Private/UAF/AbstractSkeleton/Sets/SkeletonTreeDragDrop.h
+++ b/Engine/Plugins/Experimental/UAF/UAF/Source/UAFEditor/Private/UAF/AbstractSkeleton/Sets/SkeletonTreeDragDrop.h
@@ -6,7 +6,7 @@
 #include "SSetSkeletonTree.h"
 #include "DragAndDrop/DecoratedDragDropOp.h"
 
-namespace UE::UAF::Sets
+namespace UE::UAF::Editor
 {
 	class FSkeletonTreeDragDropOp : public FDecoratedDragDropOp
 	{
@@ -30,7 +30,7 @@ namespace UE::UAF::Sets
 				.BorderImage(FAppStyle::GetBrush("Menu.Background"))
 				[
 					SNew(STextBlock)
-						.Text(Items.Num() == 1 ? FText::FromName(Items[0]->BoneName) :	FText::Format(NSLOCTEXT("UE::UAF::Sets::FSkeletonTreeDragDropOp", "MultipleBonesTooltip", "{0} bones"), FText::AsNumber(Items.Num())))
+						.Text(Items.Num() == 1 ? FText::FromName(Items[0]->BoneName) :	FText::Format(NSLOCTEXT("UE::UAF::Editor::FSkeletonTreeDragDropOp", "MultipleBonesTooltip", "{0} bones"), FText::AsNumber(Items.Num())))
 				];
 		}
 
