UAF: Remove context menu option to add a function via the asset outliner instead of using the function tab workflow

#Synced-CL 49616604

diff --git a/Engine/Plugins/Experimental/UAF/UAF/Source/UAFEditor/Internal/Common/AnimNextAssetItemDetails.cpp b/Engine/Plugins/Experimental/UAF/UAF/Source/UAFEditor/Internal/Common/AnimNextAssetItemDetails.cpp
--- a/Engine/Plugins/Experimental/UAF/UAF/Source/UAFEditor/Internal/Common/AnimNextAssetItemDetails.cpp
+++ b/Engine/Plugins/Experimental/UAF/UAF/Source/UAFEditor/Internal/Common/AnimNextAssetItemDetails.cpp
@@ -278,25 +278,6 @@ void FAnimNextAssetItemDetails::RegisterToolMenuExtensions()
 					FSlateIconFinder::FindIconForClass(SubEntryClass, "ClassIcon.Object"));
 			}
 		}
-
-		AssetSection.AddMenuEntry(
-			"AddFunction",
-			LOCTEXT("AddFunctionLabel", "Add Function"),
-			LOCTEXT("AddFunctionTooltip", "Add a function to this asset"),
-			FSlateIcon(FAppStyle::GetAppStyleSetName(), "GraphEditor.Function_16x"),
-			FUIAction(
-				FExecuteAction::CreateWeakLambda(WorkspaceItemContext, [WeakWorkspaceEditor = TWeakPtr<UE::Workspace::IWorkspaceEditor>(WorkspaceEditor), EditorData]()
-				{
-					URigVMLibraryNode* NewFunction = EditorData->AddFunction(TEXT("NewFunction"), true);
-
-					// Open the new function's graph
-					if(TSharedPtr<UE::Workspace::IWorkspaceEditor> PinnedWorkspaceEditor = WeakWorkspaceEditor.Pin())
-					{
-						UObject* EditorObject = EditorData->GetEditorObjectForRigVMGraph(NewFunction->GetContainedGraph());
-						PinnedWorkspaceEditor->OpenObjects({EditorObject});
-					}
-				})
-			));
 	}));
 }
 
