UAF: Fix editing and compilation issues after an asset is duplicated

#Synced-CL 49559750

diff --git a/Engine/Plugins/Experimental/UAF/UAF/Source/UAFUncookedOnly/Private/AnimNextRigVMAssetEditorData.cpp b/Engine/Plugins/Experimental/UAF/UAF/Source/UAFUncookedOnly/Private/AnimNextRigVMAssetEditorData.cpp
--- a/Engine/Plugins/Experimental/UAF/UAF/Source/UAFUncookedOnly/Private/AnimNextRigVMAssetEditorData.cpp
+++ b/Engine/Plugins/Experimental/UAF/UAF/Source/UAFUncookedOnly/Private/AnimNextRigVMAssetEditorData.cpp
@@ -344,8 +344,7 @@ void UAnimNextRigVMAssetEditorData::PostDuplicate(bool bDuplicateForPIE)
 {
 	Super::PostDuplicate(bDuplicateForPIE);
 
-	ReconstructAllNodes();
-	RecompileVM();
+	HandlePackageDone();
 }
 
 void UAnimNextRigVMAssetEditorData::HandlePackageDone(const FEndLoadPackageContext& Context)
