UAF: Remove making functions public per default to avoid breaking tests until the transaction handling is fixed

#Synced-CL 49651073

diff --git a/Engine/Plugins/Experimental/UAF/UAF/Source/UAFUncookedOnly/Private/AnimNextRigVMAssetEditorData.cpp b/Engine/Plugins/Experimental/UAF/UAF/Source/UAFUncookedOnly/Private/AnimNextRigVMAssetEditorData.cpp
--- a/Engine/Plugins/Experimental/UAF/UAF/Source/UAFUncookedOnly/Private/AnimNextRigVMAssetEditorData.cpp
+++ b/Engine/Plugins/Experimental/UAF/UAF/Source/UAFUncookedOnly/Private/AnimNextRigVMAssetEditorData.cpp
@@ -2727,10 +2727,6 @@ URigVMLibraryNode* UAnimNextRigVMAssetEditorData::AddFunction(FName InFunctionNa
 {
 	URigVMController* Controller = RigVMClient.GetOrCreateController(GetLocalFunctionLibrary());
 	URigVMLibraryNode* Node = Controller->AddFunctionToLibrary(InFunctionName, bInMutable, FVector2D::ZeroVector, bInSetupUndoRedo, false);
-	if (Node)
-	{
-		MarkFunctionPublic(Node->GetFName());
-	}
 	
 	if (bPrintPythonCommand)
 	{
