Move asset wizard files to UAF directory and remove call to AddToRoot when using blueprint factory

#Synced-CL 49678682

diff --git a/Engine/Plugins/Experimental/UAF/UAF/Source/UAFEditor/Private/AnimNextEditorModule.cpp b/Engine/Plugins/Experimental/UAF/UAF/Source/UAFEditor/Private/AnimNextEditorModule.cpp
--- a/Engine/Plugins/Experimental/UAF/UAF/Source/UAFEditor/Private/AnimNextEditorModule.cpp
+++ b/Engine/Plugins/Experimental/UAF/UAF/Source/UAFEditor/Private/AnimNextEditorModule.cpp
@@ -84,7 +84,6 @@
 #include "Variables/Outliner/VariablesOutlinerCommands.h"
 #include "WidgetDrawerConfig.h"
 #include "Widgets/Docking/SDockTab.h"
-#include "Wizard/AssetWizard.h"
 #include "Interfaces/IPluginManager.h"
 #include "AssetToolsModule.h"
 #include "UAF/Debug/VariableWatchTabSummoner.h"
@@ -92,6 +91,7 @@
 #include "Styling/SlateIconFinder.h"
 #include "UAF/AbstractSkeleton/Sets/Commands.h"
 #include "Variables/UAFAssetDataPropertyCustomization.h"
+#include "UAF/Wizard/AssetWizard.h"
 
 #define LOCTEXT_NAMESPACE "AnimNextEditorModule"
 
diff --git a/Engine/Plugins/Experimental/UAF/UAF/Source/UAFEditor/Private/UAF/Wizard/AssetWizard.cpp b/Engine/Plugins/Experimental/UAF/UAF/Source/UAFEditor/Private/UAF/Wizard/AssetWizard.cpp
new file mode 100644
--- /dev/null
+++ b/Engine/Plugins/Experimental/UAF/UAF/Source/UAFEditor/Private/UAF/Wizard/AssetWizard.cpp
@@ -0,0 +1,278 @@
+// Copyright Epic Games, Inc. All Rights Reserved.
+
+#include "UAF/Wizard/AssetWizard.h"
+
+#include "AssetToolsModule.h"
+#include "Component/AnimNextComponent.h"
+#include "Components/SkeletalMeshComponent.h"
+#include "ContentBrowserModule.h"
+#include "EditorModeManager.h"
+#include "Engine/SCS_Node.h"
+#include "Engine/SimpleConstructionScript.h"
+#include "Engine/SkeletalMesh.h"
+#include "Factories/BlueprintFactory.h"
+#include "Framework/Application/SlateApplication.h"
+#include "IContentBrowserSingleton.h"
+#include "IWorkspaceEditor.h"
+#include "IWorkspaceEditorModule.h"
+#include "Kismet2/KismetEditorUtilities.h"
+#include "Module/AnimNextModule.h"
+#include "ObjectTools.h"
+#include "PackageTools.h"
+#include "Serialization/ArchiveReplaceObjectRef.h"
+#include "TemplateConfig.h"
+#include "TemplateDataAsset.h"
+#include "UAF/Wizard/STemplatePicker.h"
+#include "UObject/SavePackage.h"
+#include "Workspace/AnimNextWorkspaceEditorMode.h"
+#include "Workspace/AnimNextWorkspaceFactory.h"
+#include "ObjectTools.h"
+#include "TemplateConfig.h"
+#include "IContentBrowserSingleton.h"
+#include "UObject/SavePackage.h"
+
+#define LOCTEXT_NAMESPACE "UE::UAF::Editor::FAssetWizard"
+
+namespace UE::UAF::Editor
+{
+	void FAssetWizard::CreateAssets(const TObjectPtr<const UUAFTemplateDataAsset> Template, const TObjectPtr<const UUAFTemplateConfig> TemplateConfig)
+	{
+		FScopedSlowTask SlowTask(4.0f, LOCTEXT("SlowTask", "Creating Assets from Template"));
+		SlowTask.MakeDialog(false);
+
+		TMap<UObject*, UObject*> TemplateToUserAssetMap;
+
+		// 1. Start by duplicating all the template assets into the user's chosen directory.
+		{
+			SlowTask.EnterProgressFrame(1.0f, LOCTEXT("CreateAssets", "Duplicating template assets"));
+
+			bool bFailedToInstantiateTemplateAssets = false;
+			
+			for (int32 TemplateAssetIndex = 0; TemplateAssetIndex < Template->Assets.Num(); ++TemplateAssetIndex)
+			{
+				TObjectPtr<UObject> TemplateAsset  = Template->Assets[TemplateAssetIndex];
+				
+				const FString UserAssetName = TemplateConfig->AssetNaming[TemplateAssetIndex];
+				
+				if (TObjectPtr<UObject> UserAsset = IAssetTools::Get().DuplicateAsset(UserAssetName, TemplateConfig->OutputPath.Path, TemplateAsset))
+				{
+					TemplateToUserAssetMap.Add(TemplateAsset, UserAsset);
+				}
+				else
+				{
+					UE_LOG(LogAnimation, Warning, TEXT("Asset Wizard: Failed to duplicate asset: %s"), *TemplateAsset->GetFullName());
+					bFailedToInstantiateTemplateAssets = true;
+					break;
+				}
+			}
+
+			if (bFailedToInstantiateTemplateAssets)
+			{
+				for (TPair<UObject*, UObject*> AssetPair : TemplateToUserAssetMap)
+				{
+					UObject* UserAsset = AssetPair.Value;
+					ObjectTools::DeleteSingleObject(UserAsset);
+				}
+				return;
+			}
+		}
+
+		// 2. Scan for all references within our new user assets that refer to the template assets and replace them with their user asset counterpart.
+		{		
+			SlowTask.EnterProgressFrame(1.0f, LOCTEXT("FixingReferences", "Fixing asset references"));
+
+			for (TPair<UObject*, UObject*> AssetPair : TemplateToUserAssetMap)
+			{
+				UObject* UserAsset = AssetPair.Value;
+
+				FArchiveReplaceObjectRef<UObject> ReplaceAr(UserAsset, TemplateToUserAssetMap);
+				UserAsset->Serialize(ReplaceAr);
+
+				UserAsset->MarkPackageDirty();
+
+				FString PackageFileName = FPackageName::LongPackageNameToFilename(
+					UserAsset->GetPackage()->GetName(),
+					FPackageName::GetAssetPackageExtension()
+				);
+
+				FSavePackageArgs Args;
+				Args.TopLevelFlags = RF_Public | RF_Standalone;
+				Args.Error = GError;
+
+				UPackage::SavePackage(UserAsset->GetPackage(), UserAsset, *PackageFileName, Args);
+			}
+		}
+
+		// 3. Create a blueprint asset depending on the template config
+		{
+			SlowTask.EnterProgressFrame(1.0f, LOCTEXT("ConstructingBlueprint", "Constructing blueprint"));
+
+			if (TemplateConfig->BlueprintMode == ETemplateBlueprintMode::CreateNewBlueprint || TemplateConfig->BlueprintMode == ETemplateBlueprintMode::ModifyExistingBlueprint)
+			{
+				UBlueprint* Blueprint = nullptr;
+
+				if (TemplateConfig->BlueprintMode == ETemplateBlueprintMode::CreateNewBlueprint)
+				{
+					TObjectPtr<UBlueprintFactory> BlueprintFactory = NewObject<UBlueprintFactory>();
+
+					BlueprintFactory->ParentClass = TemplateConfig->BlueprintClass.Get()
+						? TemplateConfig->BlueprintClass.Get()
+						: AActor::StaticClass();
+
+					FAssetToolsModule& AssetToolsModule = FModuleManager::LoadModuleChecked<FAssetToolsModule>("AssetTools");
+					IAssetTools& AssetTools = AssetToolsModule.Get();
+
+					TObjectPtr<UObject> NewAsset = AssetTools.CreateAsset(TemplateConfig->BlueprintAssetName, TemplateConfig->OutputPath.Path, UBlueprint::StaticClass(), BlueprintFactory);
+					Blueprint = Cast<UBlueprint>(NewAsset);
+
+					if (!Blueprint)
+					{
+						UE_LOG(LogAnimation, Error, TEXT("Asset Wizard: Failed to create blueprint"));
+					}
+				}
+				else
+				{
+					Blueprint = TemplateConfig->BlueprintToModify.LoadSynchronous();
+
+					if (!Blueprint)
+					{
+						UE_LOG(LogAnimation, Error, TEXT("Asset Wizard: Failed to load blueprint to modify"));
+					}
+				}
+
+				if (Blueprint)
+				{
+					AActor* Actor = CastChecked<AActor>(Blueprint->GeneratedClass->GetDefaultObject());
+
+					USkeletalMeshComponent* SkelMeshComp = Actor->GetComponentByClass<USkeletalMeshComponent>();
+					UAnimNextComponent* UAFComp = Actor->GetComponentByClass<UAnimNextComponent>();
+
+					// Create a USkeletalMeshComponent and/or a UAnimNextComponent if either are missing
+					{
+						USimpleConstructionScript* ConstructionScript = Blueprint->SimpleConstructionScript;
+						USCS_Node* DefaultRootNode = ConstructionScript->GetDefaultSceneRootNode();
+
+						if (!SkelMeshComp)
+						{
+							USCS_Node* SkelMeshCompNode = ConstructionScript->CreateNode(USkeletalMeshComponent::StaticClass(), "Mesh");
+							if (DefaultRootNode)
+							{
+								DefaultRootNode->AddChildNode(SkelMeshCompNode);
+							}
+							else
+							{
+								ConstructionScript->AddNode(SkelMeshCompNode);
+							}
+
+							SkelMeshComp = CastChecked<USkeletalMeshComponent>(SkelMeshCompNode->ComponentTemplate);
+						}
+
+						if (!UAFComp)
+						{
+							USCS_Node* UAFCompNode = ConstructionScript->CreateNode(UAnimNextComponent::StaticClass(), "AnimationFramework");
+							ConstructionScript->AddNode(UAFCompNode);
+
+							UAFComp = CastChecked<UAnimNextComponent>(UAFCompNode->ComponentTemplate);
+						}
+					}
+
+					// Configure components
+					{
+						check(SkelMeshComp && UAFComp);
+
+						SkelMeshComp->SetSkeletalMesh(TemplateConfig->SkeletalMesh.LoadSynchronous());
+						SkelMeshComp->SetEnableAnimation(false);
+
+						FTransform RelativeTransform;
+						RelativeTransform.SetTranslation(FVector(0, 0, -90.0));
+						RelativeTransform.SetRotation(FQuat::MakeFromEuler(FVector(0, 0, -90.0)));
+						SkelMeshComp->SetRelativeTransform(RelativeTransform);
+
+						if (Template->SystemToAssignToComponent)
+						{
+							UObject** UserAssetToAssignPtr = TemplateToUserAssetMap.Find(Template->SystemToAssignToComponent);
+							if (ensureAlwaysMsgf(UserAssetToAssignPtr, TEXT("Asset Wizard: Could not find user system asset counterpart for '%s' to assign to the UAF component. Is the template itself misconfigured?"), *Template->AssetToOpen->GetFullName()))
+							{
+								UAFComp->SetAssetFromObject(*UserAssetToAssignPtr);
+							}
+						}
+					}
+
+					FKismetEditorUtilities::CompileBlueprint(Blueprint);
+				}
+			}
+		}
+
+		// 4. Open up one of the newly created assets as determined by the template.
+		{
+			SlowTask.EnterProgressFrame(1.0f, LOCTEXT("Finalizing", "Finalizing"));
+
+			if (Template->AssetToOpen)
+			{
+				UObject** UserAssetToOpenPtr = TemplateToUserAssetMap.Find(Template->AssetToOpen);
+				if (ensureAlwaysMsgf(UserAssetToOpenPtr, TEXT("Asset Wizard: Could not open user asset counterpart for '%s' to open. Is the template itself misconfigured?"), *Template->AssetToOpen->GetFullName()))
+				{
+					Workspace::IWorkspaceEditorModule& WorkspaceEditorModule = FModuleManager::Get().LoadModuleChecked<Workspace::IWorkspaceEditorModule>("WorkspaceEditor");
+					Workspace::IWorkspaceEditor* WorkspaceEditor = WorkspaceEditorModule.OpenWorkspaceForObject(*UserAssetToOpenPtr, Workspace::EOpenWorkspaceMethod::Default, UAnimNextWorkspaceFactory::StaticClass());
+
+					UAnimNextWorkspaceEditorMode* EditorMode = Cast<UAnimNextWorkspaceEditorMode>(WorkspaceEditor->GetEditorModeManager().GetActiveScriptableMode(UAnimNextWorkspaceEditorMode::EM_AnimNextWorkspace));
+					EditorMode->HandleCompile();
+				}
+			}
+		}
+	}
+
+	void FAssetWizard::ShowTemplatePicker()
+	{
+		static TWeakPtr<SWindow> AssetWizardWindowWeak;
+
+		TSharedPtr<SWindow> AssetWizardWindow = AssetWizardWindowWeak.Pin();
+		if (AssetWizardWindow.IsValid())
+		{
+			AssetWizardWindow->BringToFront(true);
+			return;
+		}
+			
+		FContentBrowserModule& ContentBrowserModule = FModuleManager::LoadModuleChecked<FContentBrowserModule>("ContentBrowser");
+
+		TArray<FString> SelectedPaths;
+		const FContentBrowserItemPath CurrentPath = ContentBrowserModule.Get().GetCurrentPath();
+
+		GetMutableDefault<UUAFTemplateConfig>()->OutputPath.Path = CurrentPath.HasInternalPath() ? CurrentPath.GetInternalPathString() : TEXT("/Game");
+
+		TSharedRef<SWindow> TemplateWindow = SNew(SWindow)
+			.Title(LOCTEXT("AssetWizard", "UAF Asset Wizard"))
+			.ClientSize(FVector2D(1200, 800))
+			.SupportsMinimize(false)
+			.SupportsMaximize(false);
+
+		AssetWizardWindowWeak = TemplateWindow;
+
+		TemplateWindow->SetOnWindowClosed(FOnWindowClosed::CreateLambda([](const TSharedRef<SWindow>&)
+		{
+			AssetWizardWindowWeak.Reset();
+		}));
+
+		TemplateWindow->SetContent(
+			SNew(STemplatePicker)
+			.OnTemplateSelected_Lambda([TemplateWindow](const TObjectPtr<const UUAFTemplateDataAsset> InTemplate, const TObjectPtr<const UUAFTemplateConfig> InTemplateConfig) mutable
+			{
+				FAssetWizard::CreateAssets(InTemplate, InTemplateConfig);
+				FSlateApplication::Get().RequestDestroyWindow(TemplateWindow);
+			})
+			.OnCancel_Lambda([TemplateWindow]()
+			{
+				FSlateApplication::Get().RequestDestroyWindow(TemplateWindow);
+			})
+		);
+
+		FSlateApplication::Get().AddWindow(TemplateWindow);
+	}
+	
+	void FAssetWizard::Launch()
+	{
+		ShowTemplatePicker();		
+	}
+}
+
+#undef LOCTEXT_NAMESPACE
\ No newline at end of file
diff --git a/Engine/Plugins/Experimental/UAF/UAF/Source/UAFEditor/Private/UAF/Wizard/AssetWizard.h b/Engine/Plugins/Experimental/UAF/UAF/Source/UAFEditor/Private/UAF/Wizard/AssetWizard.h
new file mode 100644
--- /dev/null
+++ b/Engine/Plugins/Experimental/UAF/UAF/Source/UAFEditor/Private/UAF/Wizard/AssetWizard.h
@@ -0,0 +1,23 @@
+// Copyright Epic Games, Inc. All Rights Reserved.
+
+#pragma once
+
+#include "CoreMinimal.h"
+
+class UUAFTemplateDataAsset;
+class UUAFTemplateConfig;
+class UWorkspace;
+
+namespace UE::UAF::Editor
+{
+	class FAssetWizard
+	{
+	public:
+		static void Launch();
+		
+	private:
+		static void ShowTemplatePicker();
+
+		static void CreateAssets(const TObjectPtr<const UUAFTemplateDataAsset> Template, const TObjectPtr<const UUAFTemplateConfig> TemplateConfig);
+	};
+}
\ No newline at end of file
diff --git a/Engine/Plugins/Experimental/UAF/UAF/Source/UAFEditor/Private/UAF/Wizard/STemplatePicker.cpp b/Engine/Plugins/Experimental/UAF/UAF/Source/UAFEditor/Private/UAF/Wizard/STemplatePicker.cpp
new file mode 100644
--- /dev/null
+++ b/Engine/Plugins/Experimental/UAF/UAF/Source/UAFEditor/Private/UAF/Wizard/STemplatePicker.cpp
@@ -0,0 +1,851 @@
+// Copyright Epic Games, Inc. All Rights Reserved.
+
+#include "STemplatePicker.h"
+
+#include "UAFStyle.h"
+#include "AssetRegistry/AssetRegistryModule.h"
+#include "AssetThumbnail.h" 
+#include "DetailsViewArgs.h"
+#include "Engine/Blueprint.h"
+#include "Misc/MessageDialog.h"
+#include "PropertyEditorModule.h"
+#include "SPrimaryButton.h"
+#include "Styling/StyleColors.h"
+#include "Subsystems/EditorAssetSubsystem.h"
+#include "TemplateConfig.h"
+#include "TemplateDataAsset.h"
+#include "ThumbnailRendering/ThumbnailManager.h"
+#include "Widgets/Input/SEditableTextBox.h"
+#include "Widgets/Input/SSegmentedControl.h"
+#include "Widgets/Layout/SExpandableArea.h"
+#include "Widgets/Layout/SScaleBox.h"
+#include "Widgets/Layout/SWidgetSwitcher.h"
+#include "Widgets/Text/STextBlock.h"
+#include "Widgets/Views/STileView.h"
+
+#include "Editor/EditorEngine.h"
+#include "Editor.h"
+
+#define LOCTEXT_NAMESPACE "UE::UAF::Editor::STemplatePicker"
+
+namespace UE::UAF::Editor
+{
+	void STemplatePicker::AddReferencedObjects(FReferenceCollector& Collector)
+	{
+		Collector.AddReferencedObject(TemplateConfig);
+
+		for (FTemplateItemPtr TemplatePtr : Templates)
+		{
+			Collector.AddReferencedObject(TemplatePtr->Template);
+		}
+	}
+
+	void STemplatePicker::Construct(const FArguments& InArgs)
+	{
+		OnTemplateSelected = InArgs._OnTemplateSelected;
+		
+		TSharedRef<SWidget> TemplatesView = SNew(SBorder)
+			.BorderImage(FAppStyle::GetBrush("ToolPanel.GroupBorder"))
+			[
+				SAssignNew(TemplatesTileView, STileView<FTemplateItemPtr>)
+				.ListItemsSource(&Templates)
+				.SelectionMode(ESelectionMode::Single)
+				.ClearSelectionOnClick(true)
+				.ItemAlignment(EListItemAlignment::LeftAligned)
+				.OnGenerateTile_Lambda([](FTemplateItemPtr InItem, const TSharedRef<STableViewBase>& OwnerTable)
+				{
+					return SNew(SModuleTemplate, OwnerTable)
+						.Item(InItem);
+				})
+				.ItemWidth(111)
+				.ItemHeight(162)
+				.OnSelectionChanged_Lambda([this](FTemplateItemPtr InItem, ESelectInfo::Type SelectInfo)
+				{
+					SelectedTemplate = InItem;
+					HandleSelectedTemplateChanged();
+				})
+			];
+
+		TemplateConfig = GetMutableDefault<UUAFTemplateConfig>();
+		TemplateConfig->Reset();
+		TemplateConfig->OnOutputPathChanged.AddSP(this, &STemplatePicker::HandleOutputPathChanged);
+		TemplateConfig->OnBlueprintToModifyChanged.AddSP(this, &STemplatePicker::HandleBlueprintToModifyChanged);
+		
+		FDetailsViewArgs DetailsViewArgs;
+		DetailsViewArgs.bAllowSearch = false;
+		DetailsViewArgs.bLockable = false;
+		DetailsViewArgs.bShowOptions = false;
+		DetailsViewArgs.bAllowFavoriteSystem = false;
+		DetailsViewArgs.bHideSelectionTip = true;
+
+		FPropertyEditorModule& PropertyEditorModule = FModuleManager::Get().GetModuleChecked<FPropertyEditorModule>("PropertyEditor");
+		TSharedRef<IDetailsView> DetailsObjView = PropertyEditorModule.CreateDetailView(DetailsViewArgs);
+		DetailsObjView->SetObject(TemplateConfig);
+
+		TSharedRef<SWidget> DetailsView = SNew(SBorder)
+			.BorderImage(FAppStyle::GetBrush("ToolPanel.GroupBorder"))
+			[
+				SNew(SVerticalBox)
+				// Panel selector
+				+ SVerticalBox::Slot()
+				.HAlign(HAlign_Center)
+				.AutoHeight()
+				[
+					SNew(SHorizontalBox)
+					+ SHorizontalBox::Slot()
+					.AutoWidth()
+					[
+						SNew(SSegmentedControl<EPanel>)
+						.IsEnabled_Lambda([this]()
+						{
+							return SelectedTemplate.IsValid();
+						})
+						.UniformPadding(FMargin(8.0f, 0.0f))
+						.OnValueChanged_Lambda([this](const EPanel InPanel)
+						{
+							ShowPanel = InPanel;
+						})
+						.Value_Lambda([this]()
+						{
+							return ShowPanel;
+						})	
+						+ SSegmentedControl<EPanel>::Slot(EPanel::Overview)
+						[
+							SNew(SHorizontalBox)
+							+ SHorizontalBox::Slot()
+							.Padding(2.0f)
+							.AutoWidth()
+							[
+								SNew(SImage)
+								.Image(FAppStyle::GetBrush("GraphEditor.Animation_16x"))
+							]
+							+ SHorizontalBox::Slot()
+							.AutoWidth()
+							.Padding(2.0f)
+							[
+								SNew(STextBlock)
+								.Text(LOCTEXT("Overview", "Overview"))
+							]
+						]
+						+ SSegmentedControl<EPanel>::Slot(EPanel::Details)
+						[
+							SNew(SHorizontalBox)
+							+ SHorizontalBox::Slot()
+							.Padding(2.0f)
+							.AutoWidth()
+							[
+								SNew(SImage)
+								.Image(FAppStyle::GetBrush("LevelEditor.Tabs.Details"))
+							]
+							+ SHorizontalBox::Slot()
+							.AutoWidth()
+							.Padding(2.0f)
+							[
+								SNew(STextBlock)
+								.Text(LOCTEXT("Details", "Details"))
+							]
+						]
+					]
+				]
+				// Panel
+				+ SVerticalBox::Slot()
+				.FillHeight(1.0f)
+				[
+					SNew(SWidgetSwitcher)
+					.WidgetIndex_Lambda([this]()
+					{
+						return static_cast<int32>(ShowPanel);
+					})
+					+ SWidgetSwitcher::Slot() // EPanel::Overview
+					.Padding(8.0f)
+					[
+						SNew(SVerticalBox)
+						// Preview image
+						+ SVerticalBox::Slot()
+						.HAlign(HAlign_Center)
+						.AutoHeight()
+						[
+							SNew(SScaleBox)
+							.Stretch(EStretch::ScaleToFill)
+							[
+								SNew(SImage)
+								.DesiredSizeOverride(FVector2D(600,300))
+								.Image_Lambda([this]() -> const FSlateBrush*
+								{
+									if (SelectedTemplate && SelectedTemplate->Template->DetailsImage.IsSet())
+									{
+										return &SelectedTemplate->Template->DetailsImage;
+									}
+
+									return FAppStyle::GetNoBrush();
+								})
+							]
+						]
+						// Template Name
+						+ SVerticalBox::Slot()
+						.AutoHeight()
+						.Padding(FMargin(0.0f, 16.0f, 0.0f, 2.0f))
+						[
+							SNew(STextBlock)
+							.AutoWrapText(true)
+							.TextStyle(FAppStyle::Get(), "LargeText")
+							.Text_Lambda([this]()
+							{
+								return SelectedTemplate.IsValid()
+									? SelectedTemplate->Template->Title
+									: FText();
+							})
+						]
+						// Template Description
+						+ SVerticalBox::Slot()
+						.AutoHeight()
+						.Padding(FMargin(0.0f, 2.0f, 0.0f, 0.0f))
+						[
+							SNew(STextBlock)
+							.AutoWrapText(true)
+							.Text_Lambda([this]()
+							{
+								return SelectedTemplate.IsValid()
+									? SelectedTemplate->Template->Description
+									: FText();
+							})
+						]
+						// Documentation URL
+						+ SVerticalBox::Slot()
+						.Padding(0.0f, 24.0f)
+						.AutoHeight()
+						[
+							SNew(SHorizontalBox)
+							+ SHorizontalBox::Slot()
+							.AutoWidth()
+							[
+								SNew(SButton)
+								.Cursor(EMouseCursor::Hand)
+								.Visibility_Lambda([this]()
+								{
+									return SelectedTemplate.IsValid() && !SelectedTemplate->Template->DocumentationUrl.IsEmpty()
+										? EVisibility::Visible
+										: EVisibility::Collapsed;
+								})
+								.ButtonStyle(FCoreStyle::Get(), "NoBorder")
+								.OnClicked_Lambda([this]()
+								{
+									if (SelectedTemplate.IsValid() && !SelectedTemplate->Template->DocumentationUrl.IsEmpty())
+									{
+										FPlatformProcess::LaunchURL(*SelectedTemplate->Template->DocumentationUrl, nullptr, nullptr);
+									}
+									return FReply::Handled();
+								})
+								.ContentPadding(0)
+								[
+									SNew(SHorizontalBox)
+
+									+ SHorizontalBox::Slot()
+									.AutoWidth()
+									[
+										SNew(STextBlock)
+										.Text(LOCTEXT("Documentation", "Documentation"))
+										.ColorAndOpacity(FLinearColor(0.2f, 0.4f, 1.0f))
+									]
+
+									+ SHorizontalBox::Slot()
+									.AutoWidth()
+									.Padding(FMargin(4, 0, 0, 0))
+									[
+										SNew(SImage)
+										.Image(FUAFStyle::Get().GetBrush("AssetWizard.OpenExternal"))
+										.ColorAndOpacity(FLinearColor(0.2f, 0.4f, 1.0f))
+									]
+								]
+							]
+						]
+						// Tags
+						+ SVerticalBox::Slot()
+						.AutoHeight()
+						[
+							SAssignNew(TagsHBox, SHorizontalBox)
+						]
+					]
+					+ SWidgetSwitcher::Slot() // EPanel::Details
+					[
+						SNew(SVerticalBox)
+						.Visibility_Lambda([this]()
+						{
+							return SelectedTemplate.IsValid() ? EVisibility::Visible : EVisibility::Hidden;		
+						})
+						+ SVerticalBox::Slot()
+						.AutoHeight()
+						[
+							DetailsObjView
+						]
+						+ SVerticalBox::Slot()
+						.AutoHeight()
+						[
+							SNew(SExpandableArea)
+							.AreaTitle(LOCTEXT("IncludedAssets", "Included Assets"))
+							.BorderImage(FAppStyle::GetBrush("Brushes.Header"))
+							.BodyBorderImage(FAppStyle::GetBrush("Brushes.Recessed"))
+							.HeaderPadding(4.0f)
+							.Padding(2.0f)
+							.AllowAnimatedTransition(false)
+							.InitiallyCollapsed(false)
+							.BodyContent()
+							[
+								SAssignNew(IncludedAssetsListView, SListView<TSharedPtr<FAssetMetadata>>)
+								.SelectionMode(ESelectionMode::None)
+								.ListItemsSource(&IncludedAssets)
+								.OnGenerateRow_Lambda([this](TSharedPtr<FAssetMetadata> InAsset, const TSharedRef<STableViewBase>& InTableView)
+								{
+									TSharedPtr<FAssetThumbnailPool> Pool = UThumbnailManager::Get().GetSharedThumbnailPool();
+									TSharedPtr<FAssetThumbnail> AssetThumbnail = MakeShareable(new FAssetThumbnail(InAsset->AssetData.GetAsset(), 32, 32, Pool));
+
+									FAssetThumbnailConfig AssetThumbnailConfig;
+									AssetThumbnailConfig.ShowAssetBorder = true;
+																	
+									return SNew(STableRow<TSharedPtr<FAssetMetadata>>, InTableView)
+									.Content()
+									[
+										SNew(SHorizontalBox)
+										+ SHorizontalBox::Slot()
+										.FillWidth(0.5f)
+										.Padding(4.0f)
+										[
+											SNew(SHorizontalBox)
+											+ SHorizontalBox::Slot()
+											.AutoWidth()
+											[
+												AssetThumbnail->MakeThumbnailWidget(AssetThumbnailConfig)
+
+											]
+											+ SHorizontalBox::Slot()
+											.AutoWidth()
+											.VAlign(VAlign_Center)
+											.Padding(4.0f)
+											[
+												SNew(SVerticalBox)
+												+ SVerticalBox::Slot()
+												.AutoHeight()
+												[
+													SNew(STextBlock)
+													.Text(FText::FromName(InAsset->AssetData.AssetName))
+												]
+												+ SVerticalBox::Slot()
+												.AutoHeight()
+												[
+													SNew(STextBlock)
+													.TextStyle(&FAppStyle::Get(), "SmallText.Subdued")
+													.Text(InAsset->AssetData.GetClass()->GetDisplayNameText())
+												]
+											]
+										]
+										+ SHorizontalBox::Slot()
+										.FillWidth(0.5f)
+										.Padding(4.0f)
+										.VAlign(VAlign_Center)
+										[
+											SNew(SEditableTextBox)
+											.Text_Lambda([InAsset]
+											{
+												return FText::FromString(InAsset->AssetRename);
+											})
+											.OnTextCommitted_Lambda([InAsset](const FText& InText, ETextCommit::Type)
+											{
+												InAsset->AssetRename = InText.ToString();
+											})
+											.OnVerifyTextChanged_Lambda([this, InAsset](const FText& InNewText, FText& OutErrorMessage)
+											{
+												UEditorAssetSubsystem* EditorAssetSubsystem = GEditor->GetEditorSubsystem<UEditorAssetSubsystem>();
+
+												const FString AssetPath = FPaths::Combine(TemplateConfig->OutputPath.Path, InNewText.ToString());
+
+												if (EditorAssetSubsystem->DoesAssetExist(AssetPath))
+												{
+													OutErrorMessage = LOCTEXT("AssetAlreadyExists", "Name clashes with an existing asset");
+													return false;
+												}
+
+												for (TSharedPtr<FAssetMetadata> OtherAsset : IncludedAssets)
+												{
+													if (InAsset != OtherAsset && OtherAsset->AssetRename.Equals(InNewText.ToString()))
+													{
+														OutErrorMessage = LOCTEXT("AssetNameClash", "Name clashes with another template asset");
+														return false;
+													}
+												}
+
+												if (TemplateConfig->BlueprintMode == ETemplateBlueprintMode::CreateNewBlueprint && TemplateConfig->BlueprintAssetName.Equals(InNewText.ToString()))
+												{
+													OutErrorMessage = LOCTEXT("AssetNameClashWithBlueprint", "Name clashes with blueprint asset name");
+													return false;
+												}
+												
+												return true;
+											})
+										]
+									];
+								})
+							]
+						]
+					]
+				]
+			];
+
+		TSharedRef<SWidget> HeaderWidget = SNew(SBorder)
+			.BorderImage(FAppStyle::GetBrush("ToolPanel.GroupBorder"))
+			.Padding(16.0f)
+			[
+				SNew(STextBlock)
+				.Text(LOCTEXT("Header","Create a bundle preconfigured UAF assets by selecting a template to generate from."))
+			];
+
+		TSharedRef<SWidget> FooterWidget = SNew(SBorder)
+			.BorderImage(FAppStyle::GetBrush("ToolPanel.GroupBorder"))
+			[
+				SNew(SHorizontalBox)
+				+ SHorizontalBox::Slot()
+				.AutoWidth()
+				.Padding(16.0f)
+				[
+					SNullWidget::NullWidget
+				]
+				+ SHorizontalBox::Slot()
+				.FillWidth(1.0f)
+				.HAlign(HAlign_Right)
+				[
+					SNew(SHorizontalBox)
+					+ SHorizontalBox::Slot()
+					.HAlign(HAlign_Right)
+					.VAlign(VAlign_Center)
+					.AutoWidth()
+					[
+						SNew(SImage)
+						.Visibility_Lambda([this]()
+						{
+							return ErrorText.IsEmpty() ? EVisibility::Collapsed : EVisibility::Visible;
+						})
+						.ColorAndOpacity(FStyleColors::Error)
+						.Image(FAppStyle::GetBrush("Icons.Error"))
+					]
+					+ SHorizontalBox::Slot()
+					.HAlign(HAlign_Right)
+					.VAlign(VAlign_Center)
+					.AutoWidth()
+					.Padding(FMargin(4.0f, 0.0f, 8.0f, 0.0f))
+					[
+						SNew(STextBlock)
+						.ColorAndOpacity(FStyleColors::Error)
+						.Text_Lambda([this]()
+						{
+							return ErrorText;
+						})
+					]
+					+ SHorizontalBox::Slot()
+					.HAlign(HAlign_Right)
+					.AutoWidth()
+					.Padding(FMargin(16.0f, 16.0f, 8.0f, 16.0f))
+					[
+						SNew(SPrimaryButton)
+						.Text(LOCTEXT("CreateAssets", "Create Assets"))
+						.IsEnabled(this, &STemplatePicker::CanCreateAssets)
+						.OnClicked_Lambda([this]()
+						{
+							HandleUseTemplate();
+							return FReply::Handled();
+						})
+					]
+					+ SHorizontalBox::Slot()
+					.HAlign(HAlign_Right)
+					.AutoWidth()
+					.Padding(FMargin(8.0f, 16.0f, 16.0f, 16.0f))
+					[
+						SNew(SButton)
+						.Text(LOCTEXT("Cancel","Cancel"))
+						.OnClicked_Lambda([InArgs]()
+						{
+							InArgs._OnCancel.ExecuteIfBound();
+							return FReply::Handled();
+						})
+					]
+				]
+			];
+		
+		ChildSlot
+		[
+			SNew(SVerticalBox)
+			+ SVerticalBox::Slot()
+			.Padding(8.0f, 8.0f, 8.0f, 0.0f)
+			.AutoHeight()
+			[
+				HeaderWidget
+			]
+			+ SVerticalBox::Slot()
+			.Padding(4.0f)
+			.FillHeight(1.0f)
+			[
+				SNew(SHorizontalBox)
+				+ SHorizontalBox::Slot()
+				.Padding(4.0f)
+				.FillWidth(1.0f)
+				[
+					SNew(SSplitter)
+					.Orientation(Orient_Horizontal)
+					+ SSplitter::Slot()
+					.Value(0.5)
+					[
+						TemplatesView
+					]
+					+ SSplitter::Slot()
+					.Value(0.5)
+					[
+						DetailsView
+					]
+				]
+			]
+			+ SVerticalBox::Slot()
+			.Padding(8.0f, 0.0f, 8.0f, 8.0f)
+			.AutoHeight()
+			[
+				FooterWidget
+			]
+		];
+
+		RefreshModuleTemplates();
+
+		FTemplateItemPtr* DefaultSelectedTemplate = Templates.FindByPredicate([](const FTemplateItemPtr TemplateItem)
+		{
+			return TemplateItem->Template->GetPathName() == TEXT("/UAF/Templates/BasicCharacter/DA_DefaultCharacter.DA_DefaultCharacter");
+		});
+
+		if (DefaultSelectedTemplate)
+		{
+			SelectedTemplate = *DefaultSelectedTemplate;
+			TemplatesTileView->SetSelection(SelectedTemplate);
+		}
+		
+		HandleSelectedTemplateChanged();
+	}
+		
+	void STemplatePicker::HandleUseTemplate()
+	{
+		bool bSuccess = TemplateConfig->TryUpdateDefaultConfigFile();
+
+		for (const TSharedPtr<FAssetMetadata>& AssetMetadata : IncludedAssets)
+		{
+			TemplateConfig->AssetNaming.Add(AssetMetadata->AssetRename);
+		}
+		
+		OnTemplateSelected.ExecuteIfBound(SelectedTemplate->Template, TemplateConfig);
+	}
+
+	void STemplatePicker::HandleOutputPathChanged()
+	{
+		UpdateUserAssetNames();
+	}
+
+	void STemplatePicker::HandleBlueprintToModifyChanged()
+	{
+		if (!TemplateConfig->BlueprintToModify.IsNull())
+		{
+			UBlueprint* BP = TemplateConfig->BlueprintToModify.LoadSynchronous();
+			if (!BP->GeneratedClass->IsChildOf<AActor>())
+			{
+				// Restricting the blueprint picker to only allow Actor-derived BPs requires loading the asset to verify the class inheritance which
+				// is too slow so the check is performed at selection time instead.
+				FMessageDialog::Open(EAppMsgCategory::Error, EAppMsgType::Ok, LOCTEXT("InvalidBlueprintSelected", "The selected blueprint to modify must derive from Actor to be able to contain the components necessary for UAF to function."));
+				TemplateConfig->BlueprintToModify.Reset();
+			}
+		}
+	}
+
+	void STemplatePicker::UpdateUserAssetNames()
+	{
+		UEditorAssetSubsystem* EditorAssetSubsystem = GEditor->GetEditorSubsystem<UEditorAssetSubsystem>();
+
+		const auto GetNewUserAssetName = [&, this](const FString BaseAssetName)
+			{
+				if (!EditorAssetSubsystem->DoesAssetExist(FPaths::Combine(TemplateConfig->OutputPath.Path, BaseAssetName)))
+				{
+					return BaseAssetName;
+				}
+
+				FString OutAssetName = BaseAssetName;
+				int32 Suffix = 1;
+
+				while (EditorAssetSubsystem->DoesAssetExist(FPaths::Combine(TemplateConfig->OutputPath.Path, OutAssetName)))
+				{
+					OutAssetName = BaseAssetName + FString::FromInt(Suffix++);
+				}
+
+				return OutAssetName;
+			};
+
+		
+		IncludedAssets.Empty();
+
+		if (SelectedTemplate.IsValid())
+		{
+			check(SelectedTemplate->Template);
+			
+			TemplateConfig->BlueprintAssetName = GetNewUserAssetName(SelectedTemplate->Template->DefaultBlueprintAssetName);
+
+			for (const TObjectPtr<UObject>& Asset : SelectedTemplate->Template->Assets)
+			{
+				TSharedPtr<FAssetMetadata> AssetMetadata = MakeShared<FAssetMetadata>();
+				AssetMetadata->AssetData = FAssetData(Asset);
+
+				const FString BaseAssetName = AssetMetadata->AssetData.AssetName.ToString();
+				AssetMetadata->AssetRename = GetNewUserAssetName(BaseAssetName);
+				
+				IncludedAssets.Add(AssetMetadata);
+			}
+		}
+
+		IncludedAssetsListView->RequestListRefresh();
+	}
+	
+	void STemplatePicker::HandleSelectedTemplateChanged()
+	{
+		TagsHBox->ClearChildren();
+
+		if (SelectedTemplate.IsValid())
+		{
+			TemplateConfig->BlueprintAssetName = SelectedTemplate->Template->DefaultBlueprintAssetName;
+			
+			for (const FText& InTag : SelectedTemplate->Template->Tags)
+			{
+				TagsHBox->AddSlot()
+				.Padding(0.0f, 4.0f, 4.0f, 4.0f)
+				.AutoWidth()
+				[
+					SNew(SBorder)
+					.BorderImage(FUAFStyle::Get().GetBrush("AssetWizard.TemplateTag.OuterBorder"))
+					.BorderBackgroundColor(FLinearColor::White)
+					.Padding(1.0f)
+					[
+						SNew(SBorder)
+						.BorderImage(FUAFStyle::Get().GetBrush("AssetWizard.TemplateTag.InnerBorder"))
+						.Padding(12.0f, 2.0f)
+						[
+							SNew(STextBlock)
+							.Text(InTag)
+						]
+					]
+				];
+			}
+		}
+		
+		UpdateUserAssetNames();
+	}
+	
+	void STemplatePicker::RefreshModuleTemplates()
+	{
+		FAssetRegistryModule& AssetRegistryModule = FModuleManager::LoadModuleChecked<FAssetRegistryModule>("AssetRegistry");
+		IAssetRegistry& AssetRegistry = AssetRegistryModule.Get();
+
+		TArray<FAssetData> TemplateAssets;
+		AssetRegistry.GetAssetsByClass(UUAFTemplateDataAsset::StaticClass()->GetClassPathName(), TemplateAssets);
+
+		for (const FAssetData& AssetData : TemplateAssets)
+		{
+			TObjectPtr<UUAFTemplateDataAsset> ModuleTemplateDataAsset = CastChecked<UUAFTemplateDataAsset>(AssetData.GetAsset());
+
+			TSharedPtr<FTemplate> ModuleTemplate = Templates.Add_GetRef(MakeShared<FTemplate>(ModuleTemplateDataAsset));
+			ModuleTemplate->Template = ModuleTemplateDataAsset;
+		}
+
+		TemplatesTileView->RequestListRefresh();
+	}
+
+	bool STemplatePicker::CanCreateAssets() const
+	{
+		ErrorText = FText::GetEmpty();
+
+		if (!SelectedTemplate.IsValid())
+		{
+			return false;
+		}
+		
+		switch (TemplateConfig->BlueprintMode)
+		{
+		case ETemplateBlueprintMode::CreateNewBlueprint:
+			if (TemplateConfig->BlueprintAssetName.IsEmpty())
+			{
+				ErrorText = LOCTEXT("InvalidBlueprintAssetName", "Invalid blueprint asset name");
+				return false;
+			}
+			if (TemplateConfig->BlueprintClass.Get() == nullptr)
+			{
+				ErrorText = LOCTEXT("InvalidBlueprintClass", "Invalid blueprint class");
+				return false;
+			}
+			if (TemplateConfig->SkeletalMesh.IsNull())
+			{
+				ErrorText = LOCTEXT("InvalidSkeletalMesh", "Invalid skeletal mesh");
+				return false;
+			}
+			break;
+		case ETemplateBlueprintMode::ModifyExistingBlueprint:
+			if (TemplateConfig->BlueprintToModify.IsNull())
+			{
+				ErrorText = LOCTEXT("InvalidBlueprintToModify", "Invalid blueprint to modify");
+				return false;
+			}
+			break;
+		}
+
+		if (TemplateConfig->BlueprintMode == ETemplateBlueprintMode::CreateNewBlueprint)
+		{
+			for (TSharedPtr<FAssetMetadata> Asset : IncludedAssets)
+			{
+				if (Asset->AssetRename.Equals(TemplateConfig->BlueprintAssetName))
+				{
+					ErrorText = LOCTEXT("InvalidAssetNames", "Blueprint asset name clashes with template asset name");
+					return false;
+				}
+			}
+
+			UEditorAssetSubsystem* EditorAssetSubsystem = GEditor->GetEditorSubsystem<UEditorAssetSubsystem>();
+			if (EditorAssetSubsystem->DoesAssetExist(FPaths::Combine(TemplateConfig->OutputPath.Path, TemplateConfig->BlueprintAssetName)))
+			{
+				ErrorText = LOCTEXT("ClashingBlueprintAssetName", "Blueprint asset name clashes with an existing asset");
+				return false;
+			}
+		}
+
+		return true;
+	}
+	
+	void SModuleTemplate::Construct(const FArguments& InArgs, const TSharedRef<STableViewBase>& OwnerTable)
+	{
+		Item = InArgs._Item;
+
+		STableRow::FArguments TableRowArguments;
+		TableRowArguments._SignalSelectionMode = ETableRowSignalSelectionMode::Instantaneous;
+		
+		STableRow::Construct(
+			TableRowArguments
+			.Style(FAppStyle::Get(), "ProjectBrowser.TableRow")
+			.Padding(2.0f)
+			.Content()
+			[
+				SNew(SBorder)
+				.BorderImage(FAppStyle::Get().GetBrush("ProjectBrowser.ProjectTile.DropShadow"))
+				.Clipping(EWidgetClipping::ClipToBounds)
+				[
+					SNew(SOverlay)
+					+ SOverlay::Slot()
+					[
+						SNew(SVerticalBox)
+						+ SVerticalBox::Slot() // Thumbnail
+						.AutoHeight()
+						.HAlign(HAlign_Center)
+						.VAlign(VAlign_Center)
+						[
+							SNew(SBox)
+							.WidthOverride(102.0f)
+							.HeightOverride(102.0f)
+							[
+								SNew(SBorder)
+								.Padding(0.0f)
+								.BorderImage(FAppStyle::Get().GetBrush("ProjectBrowser.ProjectTile.ThumbnailAreaBackground"))
+								.HAlign(HAlign_Center)
+								.VAlign(VAlign_Center)
+								[
+									SNew(SImage)
+									.Image(FUAFStyle::Get().GetBrush("AssetWizard.Template"))
+								]
+							]
+						]
+						// Name
+						+ SVerticalBox::Slot()
+						[
+							SNew(SBorder)
+							.VAlign(VAlign_Top)
+							.Padding(FMargin(3.0f, 3.0f))
+							.BorderImage_Lambda
+							(
+								[this]()
+								{
+									const bool bIsSelected = IsSelected();
+									const bool bIsRowHovered = IsHovered();
+
+									if (bIsSelected && bIsRowHovered)
+									{
+										static const FName SelectedHover("ProjectBrowser.ProjectTile.NameAreaSelectedHoverBackground");
+										return FAppStyle::Get().GetBrush(SelectedHover);
+									}
+									if (bIsSelected)
+									{
+										static const FName Selected("ProjectBrowser.ProjectTile.NameAreaSelectedBackground");
+										return FAppStyle::Get().GetBrush(Selected);
+									}
+									if (bIsRowHovered)
+									{
+										static const FName Hovered("ProjectBrowser.ProjectTile.NameAreaHoverBackground");
+										return FAppStyle::Get().GetBrush(Hovered);
+									}
+
+									return FAppStyle::Get().GetBrush("ProjectBrowser.ProjectTile.NameAreaBackground");
+								}
+							)
+							[
+								SNew(STextBlock)
+								.Font(FAppStyle::Get().GetFontStyle("ProjectBrowser.ProjectTile.Font"))
+								.AutoWrapText(true)
+								.Text(Item->Template->Title)
+								.ColorAndOpacity_Lambda
+								(
+									[this]()
+									{
+										const bool bIsSelected = IsSelected();
+										const bool bIsRowHovered = IsHovered();
+
+										if (bIsSelected || bIsRowHovered)
+										{
+											return FStyleColors::White;
+										}
+
+										return FSlateColor::UseForeground();
+									}
+								)
+							]
+						]
+					]
+					+ SOverlay::Slot()
+					[
+						SNew(SImage)
+						.Visibility(EVisibility::HitTestInvisible)
+						.Image_Lambda
+						(
+							[this]()
+							{
+								const bool bIsSelected = IsSelected();
+								const bool bIsRowHovered = IsHovered();
+
+								if (bIsSelected && bIsRowHovered)
+								{
+									static const FName SelectedHover("ProjectBrowser.ProjectTile.SelectedHoverBorder");
+									return FAppStyle::Get().GetBrush(SelectedHover);
+								}
+								if (bIsSelected)
+								{
+									static const FName Selected("ProjectBrowser.ProjectTile.SelectedBorder");
+									return FAppStyle::Get().GetBrush(Selected);
+								}
+								if (bIsRowHovered)
+								{
+									static const FName Hovered("ProjectBrowser.ProjectTile.HoverBorder");
+									return FAppStyle::Get().GetBrush(Hovered);
+								}
+
+								return FStyleDefaults::GetNoBrush();
+							}
+						)
+					]
+				]
+			],
+			OwnerTable);
+	}
+}
+
+#undef LOCTEXT_NAMESPACE
diff --git a/Engine/Plugins/Experimental/UAF/UAF/Source/UAFEditor/Private/UAF/Wizard/STemplatePicker.h b/Engine/Plugins/Experimental/UAF/UAF/Source/UAFEditor/Private/UAF/Wizard/STemplatePicker.h
new file mode 100644
--- /dev/null
+++ b/Engine/Plugins/Experimental/UAF/UAF/Source/UAFEditor/Private/UAF/Wizard/STemplatePicker.h
@@ -0,0 +1,104 @@
+// Copyright Epic Games, Inc. All Rights Reserved.
+
+#pragma once
+
+#include "AssetRegistry/AssetData.h"
+#include "CoreMinimal.h"
+#include "TemplateDataAsset.h"
+#include "Widgets/SCompoundWidget.h"
+#include "Widgets/Views/STileView.h"
+
+class UUAFTemplateDataAsset;
+class UUAFTemplateConfig;
+
+namespace UE::UAF::Editor
+{
+	DECLARE_DELEGATE_TwoParams(FOnTemplateSelected, const TObjectPtr<const UUAFTemplateDataAsset>, const TObjectPtr<const UUAFTemplateConfig>);
+
+	class STemplatePicker : public SCompoundWidget, public FGCObject
+	{
+	public:
+		friend class SModuleTemplate;
+		
+		SLATE_BEGIN_ARGS(STemplatePicker) {}
+			SLATE_EVENT(FOnTemplateSelected, OnTemplateSelected)
+			SLATE_EVENT(FSimpleDelegate, OnCancel)
+		SLATE_END_ARGS()
+
+		/** FGCObject interface */
+		virtual void AddReferencedObjects( FReferenceCollector& Collector ) override;
+		virtual FString GetReferencerName() const override
+		{
+			return TEXT("STemplatePicker");
+		}
+
+		void Construct(const FArguments& InArgs);
+
+	private:
+		void RefreshModuleTemplates();
+
+		void HandleSelectedTemplateChanged();
+
+		void HandleUseTemplate();
+		
+		void HandleOutputPathChanged();
+
+		void HandleBlueprintToModifyChanged();
+
+		void UpdateUserAssetNames();
+
+		bool CanCreateAssets() const;
+		
+		enum class EPanel
+		{
+			Overview,
+			Details
+		};
+
+		struct FTemplate
+		{
+			TObjectPtr<UUAFTemplateDataAsset> Template;
+		};
+	
+		using FTemplateItemPtr = TSharedPtr<FTemplate>;
+		
+		TSharedPtr<STileView<FTemplateItemPtr>> TemplatesTileView;
+
+		TArray<FTemplateItemPtr> Templates;
+
+		FTemplateItemPtr SelectedTemplate;
+
+		TSharedPtr<SHorizontalBox> TagsHBox;
+		
+		EPanel ShowPanel = EPanel::Details;
+
+		struct FAssetMetadata
+		{
+			FString AssetRename;
+			FAssetData AssetData;
+		};
+		
+		TArray<TSharedPtr<FAssetMetadata>> IncludedAssets;
+		
+		FOnTemplateSelected OnTemplateSelected;
+
+		TObjectPtr<UUAFTemplateConfig> TemplateConfig;
+		
+		TSharedPtr<SListView<TSharedPtr<FAssetMetadata>>> IncludedAssetsListView;
+
+		mutable FText ErrorText;
+	};
+
+	class SModuleTemplate : public STableRow<TSharedPtr<STemplatePicker::FTemplate>>
+	{
+	public:
+		SLATE_BEGIN_ARGS(SModuleTemplate) {}
+		SLATE_ARGUMENT(TSharedPtr<STemplatePicker::FTemplate>, Item)
+	SLATE_END_ARGS()
+
+	void Construct(const FArguments& InArgs, const TSharedRef<STableViewBase>& OwnerTable);
+
+	private:
+		TSharedPtr<STemplatePicker::FTemplate> Item;
+	};
+}
\ No newline at end of file
diff --git a/Engine/Plugins/Experimental/UAF/UAF/Source/UAFEditor/Private/UAF/Wizard/TemplateConfig.cpp b/Engine/Plugins/Experimental/UAF/UAF/Source/UAFEditor/Private/UAF/Wizard/TemplateConfig.cpp
new file mode 100644
--- /dev/null
+++ b/Engine/Plugins/Experimental/UAF/UAF/Source/UAFEditor/Private/UAF/Wizard/TemplateConfig.cpp
@@ -0,0 +1,26 @@
+// Copyright Epic Games, Inc. All Rights Reserved.
+
+#include "TemplateConfig.h"
+
+#include UE_INLINE_GENERATED_CPP_BY_NAME(TemplateConfig)
+
+void UUAFTemplateConfig::PostEditChangeProperty(FPropertyChangedEvent& PropertyChangedEvent)
+{
+	UObject::PostEditChangeProperty(PropertyChangedEvent);
+
+	if (PropertyChangedEvent.GetMemberPropertyName() == GET_MEMBER_NAME_CHECKED(UUAFTemplateConfig, OutputPath))
+	{
+		OnOutputPathChanged.Broadcast();
+	}
+	else if (PropertyChangedEvent.GetMemberPropertyName() == GET_MEMBER_NAME_CHECKED(UUAFTemplateConfig, BlueprintToModify))
+	{
+		OnBlueprintToModifyChanged.Broadcast();
+	}
+}
+
+void UUAFTemplateConfig::Reset()
+{
+	OnOutputPathChanged.Clear();
+	OnBlueprintToModifyChanged.Clear();
+	AssetNaming.Empty();
+}
\ No newline at end of file
diff --git a/Engine/Plugins/Experimental/UAF/UAF/Source/UAFEditor/Private/UAF/Wizard/TemplateConfig.h b/Engine/Plugins/Experimental/UAF/UAF/Source/UAFEditor/Private/UAF/Wizard/TemplateConfig.h
new file mode 100644
--- /dev/null
+++ b/Engine/Plugins/Experimental/UAF/UAF/Source/UAFEditor/Private/UAF/Wizard/TemplateConfig.h
@@ -0,0 +1,57 @@
+// Copyright Epic Games, Inc. All Rights Reserved.
+
+#pragma once
+
+#include "CoreMinimal.h"
+#include "GameFramework/Pawn.h"
+#include "Templates/SubclassOf.h"
+
+#include "TemplateConfig.generated.h"
+
+class USkeletalMesh;
+class UBlueprint;
+
+UENUM()
+enum class ETemplateBlueprintMode
+{
+	CreateNewBlueprint,
+	ModifyExistingBlueprint,
+	DoNothing,
+};
+
+UCLASS(Config = EditorPerProjectUserSettings)
+class UUAFTemplateConfig :	public UObject
+{
+	GENERATED_BODY()
+
+public:
+	// Whether the template should also set up a new Blueprint for this template
+	UPROPERTY(Config, EditAnywhere, Category = "Config")
+	ETemplateBlueprintMode BlueprintMode = ETemplateBlueprintMode::CreateNewBlueprint;
+
+	UPROPERTY(Config, EditAnywhere, Category = "Config", Meta = (EditCondition = "BlueprintMode == ETemplateBlueprintMode::CreateNewBlueprint", EditConditionHides))
+	FString BlueprintAssetName;
+	
+	UPROPERTY(Config, EditAnywhere, Category = "Config", Meta = (EditCondition = "BlueprintMode == ETemplateBlueprintMode::CreateNewBlueprint", EditConditionHides))
+	TSubclassOf<AActor> BlueprintClass = APawn::StaticClass();
+	
+	UPROPERTY(Config, EditAnywhere, Category = "Config", Meta = (EditCondition = "BlueprintMode == ETemplateBlueprintMode::CreateNewBlueprint", EditConditionHides))
+	TSoftObjectPtr<USkeletalMesh> SkeletalMesh;
+
+	UPROPERTY(EditAnywhere, Category = "Config", Meta = (EditCondition = "BlueprintMode == ETemplateBlueprintMode::ModifyExistingBlueprint", EditConditionHides, DisallowCreateNew))
+	TSoftObjectPtr<UBlueprint> BlueprintToModify;
+
+	UPROPERTY(EditAnywhere, Category = "Output", Meta = (ContentDir))
+	FDirectoryPath OutputPath;
+
+	// UObject
+	virtual void PostEditChangeProperty(FPropertyChangedEvent& PropertyChangedEvent) override;
+
+	FSimpleMulticastDelegate OnOutputPathChanged;
+	FSimpleMulticastDelegate OnBlueprintToModifyChanged;
+
+	void Reset();
+
+public:
+	TArray<FString> AssetNaming;
+};
\ No newline at end of file
diff --git a/Engine/Plugins/Experimental/UAF/UAF/Source/UAFEditor/Private/UAF/Wizard/TemplateDataAsset.h b/Engine/Plugins/Experimental/UAF/UAF/Source/UAFEditor/Private/UAF/Wizard/TemplateDataAsset.h
new file mode 100644
--- /dev/null
+++ b/Engine/Plugins/Experimental/UAF/UAF/Source/UAFEditor/Private/UAF/Wizard/TemplateDataAsset.h
@@ -0,0 +1,51 @@
+// Copyright Epic Games, Inc. All Rights Reserved.
+
+#pragma once
+
+#include "CoreMinimal.h"
+#include "Engine/DataAsset.h"
+#include "Styling/SlateBrush.h"
+
+#include "TemplateDataAsset.generated.h"
+
+class UUAFSystem;
+
+UCLASS()
+class UUAFTemplateDataAsset : public UPrimaryDataAsset
+{
+	GENERATED_BODY()
+
+public:
+	UPROPERTY(EditAnywhere, Category = Default)
+	FText Title;
+
+	UPROPERTY(EditAnywhere, Category = Default)
+	FText Description;
+	
+	UPROPERTY(EditAnywhere, Category = Default)
+	FString DocumentationUrl;
+
+	UPROPERTY(EditAnywhere, Category = Default)
+	TArray<FText> Tags;
+
+	UPROPERTY(EditAnywhere, Category = Default)
+	FSlateBrush ThumbnailImage;
+
+	UPROPERTY(EditAnywhere, Category = Default)
+	FSlateBrush DetailsImage;
+
+	// These are the assets that form this template. Each of these will be duplicated into the user's chosen directory (with references patched).
+	UPROPERTY(EditAnywhere, Category = Default)
+	TArray<TObjectPtr<UObject>> Assets;
+
+	// This is the asset that will be opened automatically once the wizard has finished creating the assets for this template.
+	UPROPERTY(EditAnywhere, Category = Default)
+	TObjectPtr<UObject> AssetToOpen;
+
+	// This is the UAF system that we will assign to the UAF component if the user is configuring a blueprint.
+	UPROPERTY(EditAnywhere, Category = Default)
+	TObjectPtr<UUAFSystem> SystemToAssignToComponent;
+
+	UPROPERTY(EditAnywhere, Category = Default)
+	FString DefaultBlueprintAssetName;
+};
\ No newline at end of file
diff --git a/Engine/Plugins/Experimental/UAF/UAF/Source/UAFEditor/Private/Wizard/AssetWizard.cpp b/Engine/Plugins/Experimental/UAF/UAF/Source/UAFEditor/Private/Wizard/AssetWizard.cpp
deleted file mode 100644
diff --git a/Engine/Plugins/Experimental/UAF/UAF/Source/UAFEditor/Private/Wizard/AssetWizard.h b/Engine/Plugins/Experimental/UAF/UAF/Source/UAFEditor/Private/Wizard/AssetWizard.h
deleted file mode 100644
diff --git a/Engine/Plugins/Experimental/UAF/UAF/Source/UAFEditor/Private/Wizard/STemplatePicker.cpp b/Engine/Plugins/Experimental/UAF/UAF/Source/UAFEditor/Private/Wizard/STemplatePicker.cpp
deleted file mode 100644
diff --git a/Engine/Plugins/Experimental/UAF/UAF/Source/UAFEditor/Private/Wizard/STemplatePicker.h b/Engine/Plugins/Experimental/UAF/UAF/Source/UAFEditor/Private/Wizard/STemplatePicker.h
deleted file mode 100644
diff --git a/Engine/Plugins/Experimental/UAF/UAF/Source/UAFEditor/Private/Wizard/TemplateConfig.cpp b/Engine/Plugins/Experimental/UAF/UAF/Source/UAFEditor/Private/Wizard/TemplateConfig.cpp
deleted file mode 100644
diff --git a/Engine/Plugins/Experimental/UAF/UAF/Source/UAFEditor/Private/Wizard/TemplateConfig.h b/Engine/Plugins/Experimental/UAF/UAF/Source/UAFEditor/Private/Wizard/TemplateConfig.h
deleted file mode 100644
diff --git a/Engine/Plugins/Experimental/UAF/UAF/Source/UAFEditor/Private/Wizard/TemplateDataAsset.h b/Engine/Plugins/Experimental/UAF/UAF/Source/UAFEditor/Private/Wizard/TemplateDataAsset.h
deleted file mode 100644
