UAF: Add support to auto link variables when dropped onto a pin

#Synced-CL 49678682

diff --git a/Engine/Plugins/Experimental/UAF/UAF/Source/UAFEditor/Internal/Common/GraphEditorSchemaActions.h b/Engine/Plugins/Experimental/UAF/UAF/Source/UAFEditor/Internal/Common/GraphEditorSchemaActions.h
--- a/Engine/Plugins/Experimental/UAF/UAF/Source/UAFEditor/Internal/Common/GraphEditorSchemaActions.h
+++ b/Engine/Plugins/Experimental/UAF/UAF/Source/UAFEditor/Internal/Common/GraphEditorSchemaActions.h
@@ -103,6 +103,8 @@ struct FAnimNextSchemaAction_Variable : public FAnimNextSchemaAction
 	UE_API virtual const FSlateBrush* GetIconBrush() const override;
 
 	UE_API virtual const FLinearColor& GetIconColor() const override;
+	
+	UE_API void SetVariableAccessorChoice(EVariableAccessorChoice InVariableAccessorChoice, bool bUpdateSearchData = false);
 
 private:
 	FName Name;
diff --git a/Engine/Plugins/Experimental/UAF/UAF/Source/UAFEditor/Private/Common/GraphEditorSchemaActions.cpp b/Engine/Plugins/Experimental/UAF/UAF/Source/UAFEditor/Private/Common/GraphEditorSchemaActions.cpp
--- a/Engine/Plugins/Experimental/UAF/UAF/Source/UAFEditor/Private/Common/GraphEditorSchemaActions.cpp
+++ b/Engine/Plugins/Experimental/UAF/UAF/Source/UAFEditor/Private/Common/GraphEditorSchemaActions.cpp
@@ -154,8 +154,9 @@ UEdGraphNode* FAnimNextSchemaAction_DispatchFactory::PerformAction(UEdGraph* Par
 
 FAnimNextSchemaAction_Variable::FAnimNextSchemaAction_Variable(FName InName, const UObject* InSourceObject, const FAnimNextParamType& InType, const EVariableAccessorChoice InVariableAccessorChoice)
 	: Name(InName)
-	, VariableAccessorChoice(InVariableAccessorChoice)
 {
+	SetVariableAccessorChoice(InVariableAccessorChoice, true);
+	
 	check(InSourceObject->IsA<UAnimNextSharedVariables>() || InSourceObject->IsA<UScriptStruct>());
 	SourceObjectPath = InSourceObject->GetPathName();
 
@@ -165,31 +166,7 @@ FAnimNextSchemaAction_Variable::FAnimNextSchemaAction_Variable(FName InName, con
 	}
 
 	TypeName = InType.ToRigVMTemplateArgument().GetBaseCPPType();
-
-	static const FText VariablesCategory = LOCTEXT("Variables", "Variables");
-	static const FTextFormat GetVariableFormat = LOCTEXT("GetVariableFormat", "Get {0}");
-	static const FTextFormat SetVariableFormat = LOCTEXT("SetVariableFormat", "Set {0}");
-
-	FText MenuDesc;
-	FText ToolTip;
-
-	if(InVariableAccessorChoice == EVariableAccessorChoice::Get)
-	{
-		MenuDesc = FText::Format(GetVariableFormat, FText::FromName(Name));
-		ToolTip = FText::FromString(FString::Printf(TEXT("Get the value of variable %s"), *Name.ToString()));
-	}
-	else if (InVariableAccessorChoice == EVariableAccessorChoice::Set)
-	{
-		MenuDesc = FText::Format(SetVariableFormat, FText::FromName(Name));
-		ToolTip = FText::FromString(FString::Printf(TEXT("Set the value of variable %s"), *Name.ToString()));
-	}
-	else
-	{
-		MenuDesc = FText::FromName(Name);
-	}
-
-	UpdateSearchData(MenuDesc, ToolTip, VariablesCategory, FText::GetEmpty());
-
+	
 	FEdGraphPinType PinType = UE::UAF::UncookedOnly::FUtils::GetPinTypeFromParamType(InType);
 	VariableColor = GetDefault<UEdGraphSchema_K2>()->GetPinTypeColor(PinType);
 }
@@ -204,6 +181,38 @@ const FLinearColor& FAnimNextSchemaAction_Variable::GetIconColor() const
 	return VariableColor;
 }
 
+void FAnimNextSchemaAction_Variable::SetVariableAccessorChoice(EVariableAccessorChoice InVariableAccessorChoice, bool bUpdateSearchData)
+{
+	VariableAccessorChoice = InVariableAccessorChoice;
+	
+	if (bUpdateSearchData)
+	{
+		static const FText VariablesCategory = LOCTEXT("Variables", "Variables");
+		static const FTextFormat GetVariableFormat = LOCTEXT("GetVariableFormat", "Get {0}");
+		static const FTextFormat SetVariableFormat = LOCTEXT("SetVariableFormat", "Set {0}");
+	
+		FText MenuDesc;
+		FText ToolTip;
+	
+		if(VariableAccessorChoice == EVariableAccessorChoice::Get)
+		{
+			MenuDesc = FText::Format(GetVariableFormat, FText::FromName(Name));
+			ToolTip = FText::FromString(FString::Printf(TEXT("Get the value of variable %s"), *Name.ToString()));
+		}
+		else if (VariableAccessorChoice == EVariableAccessorChoice::Set)
+		{
+			MenuDesc = FText::Format(SetVariableFormat, FText::FromName(Name));
+			ToolTip = FText::FromString(FString::Printf(TEXT("Set the value of variable %s"), *Name.ToString()));
+		}
+		else
+		{
+			MenuDesc = FText::FromName(Name);
+		}
+	
+		UpdateSearchData(MenuDesc, ToolTip, VariablesCategory, FText::GetEmpty());
+	}
+}
+
 UEdGraphNode* FAnimNextSchemaAction_Variable::PerformAction(UEdGraph* ParentGraph, TArray<UEdGraphPin*>& FromPins, const FVector2f& Location, bool bSelectNewNode)
 {
 	const auto AddNodeLambda = [](UEdGraph* ParentGraph, FName InName, const FString& InSourceObjectPath, const FString& InTypeName, const FString& InTypeObjectPath, const FVector2f& Location, const bool bIsGetter) -> URigVMEdGraphNode*
@@ -236,6 +245,7 @@ UEdGraphNode* FAnimNextSchemaAction_Variable::PerformAction(UEdGraph* ParentGrap
 
 			if (CPPTypeObject == nullptr)
 			{
+				Controller->CancelUndoBracket();
 				return nullptr;
 			}
 		}
@@ -248,6 +258,7 @@ UEdGraphNode* FAnimNextSchemaAction_Variable::PerformAction(UEdGraph* ParentGrap
 
 		if (SourceObject == nullptr)
 		{
+			Controller->CancelUndoBracket();
 			return nullptr;
 		}
 
diff --git a/Engine/Plugins/Experimental/UAF/UAF/Source/UAFEditor/Private/Common/Outliner/OutlinerDragDropOps.cpp b/Engine/Plugins/Experimental/UAF/UAF/Source/UAFEditor/Private/Common/Outliner/OutlinerDragDropOps.cpp
--- a/Engine/Plugins/Experimental/UAF/UAF/Source/UAFEditor/Private/Common/Outliner/OutlinerDragDropOps.cpp
+++ b/Engine/Plugins/Experimental/UAF/UAF/Source/UAFEditor/Private/Common/Outliner/OutlinerDragDropOps.cpp
@@ -1,8 +1,14 @@
 ﻿// Copyright Epic Games, Inc. All Rights Reserved.
 
 #include "OutlinerDragDropOps.h"
+
+#include "AnimNextEdGraphSchema.h"
+#include "ScopedTransaction.h"
 #include "Common/GraphEditorSchemaActions.h"
+#include "EdGraph/RigVMEdGraph.h"
+#include "Entries/AnimNextVariableEntry.h"
 #include "Variables/Outliner/VariablesOutlinerCategoryItem.h"
+#include "Variables/Outliner/VariablesOutlinerEntryItem.h"
 
 #define LOCTEXT_NAMESPACE "OutlinerDragDropOps"
 
@@ -47,6 +53,89 @@ void FFunctionDragDropOp::GetDefaultStatusSymbol(const FSlateBrush*& PrimaryBrus
 	SecondaryBrushOut = nullptr;
 	SecondaryColorOut = FSlateColor::UseForeground();
 }
+	
+FReply FVariableDragDropOp::DroppedOnPin(const FVector2f& ScreenPosition, const FVector2f& GraphPosition)
+{
+    TSharedPtr<FAnimNextSchemaAction_Variable> SourceVariableAction = GetAction();
+    if (SourceVariableAction.IsValid())
+    {
+    	UEdGraphPin* Pin = GetHoveredPin();
+    	const URigVMEdGraph* RigVMEdGraph = Cast<URigVMEdGraph>(GetHoveredGraph());
+    	if (RigVMEdGraph && Pin)
+    	{
+    		// Create External Variable from the outliner item to be able to reuse the RigVM function to check if it can be linked to the hovered pin
+    		FRigVMExternalVariable CreatedExternalVariable;
+    		if (TSharedPtr<FVariablesOutlinerEntryItem> VariablesItem = WeakItem.Pin())
+    		{
+    			if(const IAnimNextRigVMVariableInterface* Variable = Cast<IAnimNextRigVMVariableInterface>(VariablesItem->WeakEntry.Get()))
+    			{
+    				if(const UPropertyBag* PropertyBag = Variable->GetPropertyBag().GetPropertyBagStruct())
+    				{
+    					const TConstArrayView<FPropertyBagPropertyDesc> VariableDescs = PropertyBag->GetPropertyDescs();
+    					if(VariableDescs.Num() != 0)
+    					{
+    						uint8* Container = const_cast<uint8*>(Variable->GetPropertyBag().GetValue().GetMemory());
+    						for(const FPropertyBagPropertyDesc& Desc : VariableDescs)
+    						{
+    							FRigVMExternalVariable ExternalVariable = FRigVMExternalVariable::Make(Desc.CachedProperty, Container);
+    							if(!ExternalVariable.IsValid())
+    							{
+    								continue;
+    							}
+    							
+    							CreatedExternalVariable = MoveTemp(ExternalVariable);
+    							break;
+    						}
+    					}
+    				}
+    			}
+    			else if (const FProperty* Property = VariablesItem->PropertyPath.Get())
+    			{
+    				CreatedExternalVariable = FRigVMExternalVariable::Make(Property, nullptr);
+    			}
+    		}
+    		
+    		if (CreatedExternalVariable.IsValid(true))
+    		{
+    			if (const URigVMEdGraph* Graph = Cast<URigVMEdGraph>(GetHoveredGraph()))
+    			{
+    				const URigVMGraph* ModelGraph = Graph->GetModel();
+    				const URigVMPin* ModelPin = ModelGraph ? ModelGraph->FindPin(Pin->GetName()) : nullptr;
+    				if (ModelPin && ModelPin->CanBeBoundToVariable(CreatedExternalVariable))
+    				{
+    					if (URigVMController* Controller = Graph->GetController())
+    					{
+    						FScopedTransaction Transaction(LOCTEXT("LinkVariableToPin", "Link Variable to Node Pin"));
+    						Controller->OpenUndoBracket(TEXT("Link Variable to Pin"));
+    					
+    						FAnimNextSchemaAction_Variable Action = *SourceVariableAction.Get();
+    						Action.SetVariableAccessorChoice(FAnimNextSchemaAction_Variable::EVariableAccessorChoice::Get);
+    			
+    						TArray<UEdGraphPin*> DummyPins;
+    						DummyPins.Add(Pin);
+    						if (const URigVMEdGraphNode* RigVMEdNode = Cast<URigVMEdGraphNode>(Action.PerformAction(GetHoveredGraph(), DummyPins, GraphPosition)))
+    						{
+    							const URigVMNode* ModelNode = RigVMEdNode->GetModelNode();
+    							const URigVMPin* VariableNodePin = ModelNode ? ModelNode->FindPin(TEXT("Value")) : nullptr;
+    							if (VariableNodePin)
+    							{
+    								Controller->AddLink(VariableNodePin->GetPinPath(), ModelPin->GetPinPath(), true);
+    				
+    								Controller->CloseUndoBracket();
+    								return FReply::Handled();
+    							}
+    						}
+    			
+    						Controller->CancelUndoBracket();
+    					}
+    				}
+    			}
+    		}
+    	}
+    }
+    
+    return FGraphSchemaActionDragDropAction::DroppedOnPin(ScreenPosition, GraphPosition);
+}
 
 TSharedPtr<FAnimNextSchemaAction_Function> FFunctionDragDropOp::GetAction() const
 {
diff --git a/Engine/Plugins/Experimental/UAF/UAF/Source/UAFEditor/Private/Common/Outliner/OutlinerDragDropOps.h b/Engine/Plugins/Experimental/UAF/UAF/Source/UAFEditor/Private/Common/Outliner/OutlinerDragDropOps.h
--- a/Engine/Plugins/Experimental/UAF/UAF/Source/UAFEditor/Private/Common/Outliner/OutlinerDragDropOps.h
+++ b/Engine/Plugins/Experimental/UAF/UAF/Source/UAFEditor/Private/Common/Outliner/OutlinerDragDropOps.h
@@ -25,7 +25,8 @@ public:
 
 	// FGraphSchemaActionDragDropAction interface
 	virtual void GetDefaultStatusSymbol(const FSlateBrush*& PrimaryBrushOut, FSlateColor& IconColorOut, FSlateBrush const*& SecondaryBrushOut, FSlateColor& SecondaryColorOut) const override;
-
+	virtual FReply DroppedOnPin(const FVector2f& ScreenPosition, const FVector2f& GraphPosition) override;
+	
 	TSharedPtr<FAnimNextSchemaAction_Variable> GetAction() const;
 
 	TWeakPtr<FVariablesOutlinerEntryItem> WeakItem;
