Fix for bug in old deprecation path for UAF system assets: unintentional struct slicing caused FUAFSystemFactoryAsset_System to be converted to FUAFSystemFactoryAsset

#Synced-CL 49656583

diff --git a/Engine/Plugins/Experimental/UAF/UAF/Source/UAF/Private/Component/AnimNextComponent.cpp b/Engine/Plugins/Experimental/UAF/UAF/Source/UAF/Private/Component/AnimNextComponent.cpp
--- a/Engine/Plugins/Experimental/UAF/UAF/Source/UAF/Private/Component/AnimNextComponent.cpp
+++ b/Engine/Plugins/Experimental/UAF/UAF/Source/UAF/Private/Component/AnimNextComponent.cpp
@@ -74,21 +74,21 @@ void UAnimNextComponent::Serialize(FArchive& Ar)
 #if WITH_EDITORONLY_DATA
 	if (Ar.IsLoading())
 	{
+		PRAGMA_DISABLE_DEPRECATION_WARNINGS
 		if (Ar.CustomVer(FFortniteMainBranchObjectVersion::GUID) < FFortniteMainBranchObjectVersion::UAFProceduralSystems)
 		{
-			AssetData.InitializeAs(FUAFSystemFactoryAsset_System(Module_DEPRECATED));
+			AssetData.InitializeAs<FUAFSystemFactoryAsset_System>(Module_DEPRECATED);
+			Module_DEPRECATED = nullptr;
 		}
 		else if (Ar.CustomVer(FFortniteMainBranchObjectVersion::GUID) < FFortniteMainBranchObjectVersion::UAFAssetData)
 		{
-			PRAGMA_DISABLE_DEPRECATION_WARNINGS
 			if (AssetData.IsValid() == false && Asset_DEPRECATED != nullptr)
 			{
 				AssetData = UE::UAF::FAssetDataFactory::CreateUAFAssetDataFromObject<FUAFSystemFactoryAsset>(Asset_DEPRECATED);
 			}
-
 			Asset_DEPRECATED = nullptr;
-			PRAGMA_ENABLE_DEPRECATION_WARNINGS
 		}
+		PRAGMA_ENABLE_DEPRECATION_WARNINGS
 	}
 #endif
 }
