Fix compile error in StressTestUnity due to CVarEnableUAFBrowser naming clash

#rnx
#Synced-CL 49678682

diff --git a/Engine/Plugins/Experimental/UAF/UAF/Source/UAFEditor/Private/Workspace/AnimNextWorkspaceEditorMode.cpp b/Engine/Plugins/Experimental/UAF/UAF/Source/UAFEditor/Private/Workspace/AnimNextWorkspaceEditorMode.cpp
--- a/Engine/Plugins/Experimental/UAF/UAF/Source/UAFEditor/Private/Workspace/AnimNextWorkspaceEditorMode.cpp
+++ b/Engine/Plugins/Experimental/UAF/UAF/Source/UAFEditor/Private/Workspace/AnimNextWorkspaceEditorMode.cpp
@@ -73,8 +73,8 @@ void UAnimNextWorkspaceEditorMode::Enter()
 	}
 
 	// Add UAF Browser Drawer
-	IConsoleVariable* CVarEnableUAFBrowser = IConsoleManager::Get().FindConsoleVariable(TEXT("a.UAF.EnableUAFBrowser"));
-	if (CVarEnableUAFBrowser && CVarEnableUAFBrowser->GetBool())
+	IConsoleVariable* EnableUAFBrowserCVar = IConsoleManager::Get().FindConsoleVariable(TEXT("a.UAF.EnableUAFBrowser"));
+	if (EnableUAFBrowserCVar && EnableUAFBrowserCVar->GetBool())
 	{
 		using namespace UE::UAF::Editor;
 
