Fix for incorrect deprecation logic when loading blueprints with data inheritance on deprecated variables

#Synced-CL 49656583

diff --git a/Engine/Plugins/Experimental/UAF/UAF/Source/UAF/Private/Component/AnimNextComponent.cpp b/Engine/Plugins/Experimental/UAF/UAF/Source/UAF/Private/Component/AnimNextComponent.cpp
--- a/Engine/Plugins/Experimental/UAF/UAF/Source/UAF/Private/Component/AnimNextComponent.cpp
+++ b/Engine/Plugins/Experimental/UAF/UAF/Source/UAF/Private/Component/AnimNextComponent.cpp
@@ -77,16 +77,19 @@ void UAnimNextComponent::Serialize(FArchive& Ar)
 		PRAGMA_DISABLE_DEPRECATION_WARNINGS
 		if (Ar.CustomVer(FFortniteMainBranchObjectVersion::GUID) < FFortniteMainBranchObjectVersion::UAFProceduralSystems)
 		{
-			AssetData.InitializeAs<FUAFSystemFactoryAsset_System>(Module_DEPRECATED);
-			Module_DEPRECATED = nullptr;
+			if (Module_DEPRECATED != nullptr)
+			{
+				AssetData.InitializeAs<FUAFSystemFactoryAsset_System>(Module_DEPRECATED);
+				Module_DEPRECATED = nullptr;
+			}
 		}
 		else if (Ar.CustomVer(FFortniteMainBranchObjectVersion::GUID) < FFortniteMainBranchObjectVersion::UAFAssetData)
 		{
-			if (AssetData.IsValid() == false && Asset_DEPRECATED != nullptr)
+			if (Asset_DEPRECATED != nullptr)
 			{
 				AssetData = UE::UAF::FAssetDataFactory::CreateUAFAssetDataFromObject<FUAFSystemFactoryAsset>(Asset_DEPRECATED);
+				Asset_DEPRECATED = nullptr;
 			}
-			Asset_DEPRECATED = nullptr;
 		}
 		PRAGMA_ENABLE_DEPRECATION_WARNINGS
 	}
