UAF: Crash fix for receiving invalid data on serialization in debug builds when using shared variables and adding a non basic type

#Synced-CL 49917485

diff --git a/Engine/Plugins/Experimental/UAF/UAF/Source/UAFUncookedOnly/Private/AnimNextRigVMAssetEditorData.cpp b/Engine/Plugins/Experimental/UAF/UAF/Source/UAFUncookedOnly/Private/AnimNextRigVMAssetEditorData.cpp
--- a/Engine/Plugins/Experimental/UAF/UAF/Source/UAFUncookedOnly/Private/AnimNextRigVMAssetEditorData.cpp
+++ b/Engine/Plugins/Experimental/UAF/UAF/Source/UAFUncookedOnly/Private/AnimNextRigVMAssetEditorData.cpp
@@ -3275,8 +3275,8 @@ FInstancedPropertyBag UUAFRigVMAssetEditorData::GenerateCombinedPropertyBag(cons
 		const UUAFRigVMAssetEditorData* ReferencedEditorData = UE::UAF::UncookedOnly::FUtils::GetEditorData<UUAFRigVMAssetEditorData>(ReferencedVariableAsset);
 
 		FAnimNextRigVMAssetCompileContext CompileContext = { ReferencedEditorData };
-		AssetBag = GeneratedPropertyBags.Add_GetRef(UE::UAF::UncookedOnly::FUtils::MakePropertyBagForEditorData(ReferencedEditorData, ReferencedEditorData->GetVariableCompileContext(InSettings, CompileContext)));
-		AddInstancedPropertyBag(AssetBag, ReferencedVariableAsset, true);
+		FInstancedPropertyBag& ReferencedVariableAssetBag = GeneratedPropertyBags.Add_GetRef(UE::UAF::UncookedOnly::FUtils::MakePropertyBagForEditorData(ReferencedEditorData, ReferencedEditorData->GetVariableCompileContext(InSettings, CompileContext)));
+		AddInstancedPropertyBag(ReferencedVariableAssetBag, ReferencedVariableAsset, true);
 	}
 
 	// Required to keep default struct property values from being freed
