UAF: Fix function nodes not updating correctly when selected via graph node

#Synced-CL 49917485

diff --git a/Engine/Plugins/Experimental/UAF/UAF/Source/UAFEditor/Private/AnimNextEditorModule.cpp b/Engine/Plugins/Experimental/UAF/UAF/Source/UAFEditor/Private/AnimNextEditorModule.cpp
--- a/Engine/Plugins/Experimental/UAF/UAF/Source/UAFEditor/Private/AnimNextEditorModule.cpp
+++ b/Engine/Plugins/Experimental/UAF/UAF/Source/UAFEditor/Private/AnimNextEditorModule.cpp
@@ -1253,20 +1253,40 @@ void FAnimNextEditorModule::RegisterWorkspaceDocumentTypes(Workspace::IWorkspace
 		TGuardValue<bool> SelectGuard(RigVMEdGraph->bIsSelecting, true);
 
 		TArray<FName> NodeNamesToSelect;
+		TArray<UObject*> ObjectsToSelect;
 		for (UObject* Object : NewSelection)
 		{
-			if (URigVMEdGraphNode* RigVMEdGraphNode = Cast<URigVMEdGraphNode>(Object))
+			if (const URigVMEdGraphNode* RigVMEdGraphNode = Cast<URigVMEdGraphNode>(Object))
 			{
 				NodeNamesToSelect.Add(RigVMEdGraphNode->GetModelNodeName());
+				if (const URigVMFunctionReferenceNode* FunctionReferenceNode = Cast<URigVMFunctionReferenceNode>(RigVMEdGraphNode->GetModelNode()))
+				{
+					if (const URigVMLibraryNode* LibraryNode = FunctionReferenceNode->LoadReferencedNode())
+					{
+						if (const URigVMGraph* ContainedGraph = LibraryNode->GetContainedGraph())
+						{					
+							if (const IRigVMClientHost* Host = LibraryNode->GetImplementingOuter<IRigVMClientHost>())
+							{
+								if (URigVMEdGraph* ContainedEdGraph = Cast<URigVMEdGraph>(Host->GetEditorObjectForRigVMGraph(ContainedGraph)))
+								{							
+									ObjectsToSelect.Add(ContainedEdGraph);
+									continue;
+								}
+							}						
+						}
+					}
+				}
 			}
-			else if(UEdGraphNode* Node = Cast<UEdGraphNode>(Object))
+			else if (const UEdGraphNode* Node = Cast<UEdGraphNode>(Object))
 			{
 				NodeNamesToSelect.Add(Node->GetFName());
 			}
+			
+			ObjectsToSelect.Add(Object);
 		}
 		RigVMEdGraph->GetController()->SetNodeSelection(NodeNamesToSelect, true, true);
-
-		InContext.WorkspaceEditor->SetDetailsObjects(NewSelection.Array());
+			
+		InContext.WorkspaceEditor->SetDetailsObjects(ObjectsToSelect);
 	});
 	GraphArgs.OnNodeDoubleClicked = Workspace::FOnNodeDoubleClicked::CreateLambda([this](const Workspace::FWorkspaceEditorContext& InContext, const UEdGraphNode* InNode)
 	{
