Temp commit disable num args check in UAF function to prevent content browser thumbnail crashes

#tests editor, thumbnail, PIE

#Synced-CL 49955823

diff --git a/Engine/Plugins/Experimental/UAF/UAF/Source/UAF/Private/AnimNextRigVMAsset.cpp b/Engine/Plugins/Experimental/UAF/UAF/Source/UAF/Private/AnimNextRigVMAsset.cpp
--- a/Engine/Plugins/Experimental/UAF/UAF/Source/UAF/Private/AnimNextRigVMAsset.cpp
+++ b/Engine/Plugins/Experimental/UAF/UAF/Source/UAF/Private/AnimNextRigVMAsset.cpp
@@ -138,7 +138,11 @@ void UUAFRigVMAsset::CallFunctionHandle(UE::UAF::FFunctionHandle InHandle, FRigV
 
 	const FAnimNextRigVMFunctionData& Function = FunctionData[InHandle.FunctionIndex];
 	const int32 NumArgs = Function.ArgIndices.Num();
-	check(InArgs.GetNumPropertiesInBag() == NumArgs);
+
+	if (!ensure(InArgs.GetNumPropertiesInBag() == NumArgs))
+	{
+		return;
+	}
 
 	const UPropertyBag* PropertyBag = InArgs.GetPropertyBagStruct();
 	if (PropertyBag)
