- Deprecated and removed external-packages behaviour from UAF assets (breaking change for existing content)
- Pulled behaviour from ::PostLoad into ::Serialize where appropriate
- Removed explicit ConditionalPostLoads from non-GT codepaths to ensure ALT behaves nicely

#Synced-CL 50044035

diff --git a/Engine/Plugins/Experimental/UAF/UAF/Source/UAF/Internal/Module/AnimNextModule.h b/Engine/Plugins/Experimental/UAF/UAF/Source/UAF/Internal/Module/AnimNextModule.h
--- a/Engine/Plugins/Experimental/UAF/UAF/Source/UAF/Internal/Module/AnimNextModule.h
+++ b/Engine/Plugins/Experimental/UAF/UAF/Source/UAF/Internal/Module/AnimNextModule.h
@@ -55,7 +55,6 @@ public:
 
 	// UObject interface
 	UE_API virtual void Serialize(FArchive& Ar) override;
-	UE_API virtual void PostLoad() override;
 
 protected:
 	friend class UUAFSystemFactory;
diff --git a/Engine/Plugins/Experimental/UAF/UAF/Source/UAF/Private/AnimNextRigVMAsset.cpp b/Engine/Plugins/Experimental/UAF/UAF/Source/UAF/Private/AnimNextRigVMAsset.cpp
--- a/Engine/Plugins/Experimental/UAF/UAF/Source/UAF/Private/AnimNextRigVMAsset.cpp
+++ b/Engine/Plugins/Experimental/UAF/UAF/Source/UAF/Private/AnimNextRigVMAsset.cpp
@@ -38,19 +38,31 @@ void UUAFRigVMAsset::BeginDestroy()
 	}
 }
 
+void UUAFRigVMAsset::Serialize(FArchive& Ar)
+{
+	Super::Serialize(Ar);
+
+	if (Ar.IsLoading())
+	{
+#if WITH_EDITORONLY_DATA
+		// External editor-side APIs assume that the editordata is also loaded with the asset
+		EditorData->ConditionalPreload();
+#endif
+
+		ExtendedExecuteContext.InvalidateCachedMemory();
+
+		VM = RigVM;
+	}
+}
+
 void UUAFRigVMAsset::PostLoad()
 {
 	Super::PostLoad();
 
 #if WITH_EDITORONLY_DATA
 	// External editor-side APIs assume that the editordata is also postloaded with the asset
-	EditorData->ConditionalPreload();
 	EditorData->ConditionalPostLoad();
 #endif
-	
-	ExtendedExecuteContext.InvalidateCachedMemory();
-
-	VM = RigVM;
 
 	// In packaged builds, initialize the VM
 	// In editor, the VM will be recompiled and initialized at UUAFRigVMAssetEditorData::HandlePackageDone::RecompileVM
diff --git a/Engine/Plugins/Experimental/UAF/UAF/Source/UAF/Private/Module/AnimNextModule.cpp b/Engine/Plugins/Experimental/UAF/UAF/Source/UAF/Private/Module/AnimNextModule.cpp
--- a/Engine/Plugins/Experimental/UAF/UAF/Source/UAF/Private/Module/AnimNextModule.cpp
+++ b/Engine/Plugins/Experimental/UAF/UAF/Source/UAF/Private/Module/AnimNextModule.cpp
@@ -40,26 +40,20 @@ void UUAFSystem::Serialize(FArchive& Ar)
 				Ar.Seek(Export.SerialOffset + Export.SerialSize);
 			}
 		}
-	}
-#endif
-}
 
-void UUAFSystem::PostLoad()
-{
-	Super::PostLoad();
+		if (GetLinkerCustomVersion(FFortniteMainBranchObjectVersion::GUID) < FFortniteMainBranchObjectVersion::AnimNextGraphAccessSpecifiers)
+		{
+			DefaultState_DEPRECATED.State = PropertyBag_DEPRECATED;
+		}
 
-#if WITH_EDITORONLY_DATA
-	if(GetLinkerCustomVersion(FFortniteMainBranchObjectVersion::GUID) < FFortniteMainBranchObjectVersion::AnimNextGraphAccessSpecifiers)
-	{
-		DefaultState_DEPRECATED.State = PropertyBag_DEPRECATED;
-	}
-	if (GetLinkerCustomVersion(FFortniteMainBranchObjectVersion::GUID) < FFortniteMainBranchObjectVersion::UAFProceduralSystems)
-	{
-		for (TObjectPtr<UScriptStruct> RequiredComponent : RequiredComponents_DEPRECATED)
+		if (GetLinkerCustomVersion(FFortniteMainBranchObjectVersion::GUID) < FFortniteMainBranchObjectVersion::UAFProceduralSystems)
 		{
-			TInstancedStruct<FUAFAssetInstanceComponent> Component;
-			Component.InitializeAsScriptStruct(RequiredComponent);
-			Components.Add(MoveTemp(Component));
+			for (TObjectPtr<UScriptStruct> RequiredComponent : RequiredComponents_DEPRECATED)
+			{
+				TInstancedStruct<FUAFAssetInstanceComponent> Component;
+				Component.InitializeAsScriptStruct(RequiredComponent);
+				Components.Add(MoveTemp(Component));
+			}
 		}
 	}
 #endif
diff --git a/Engine/Plugins/Experimental/UAF/UAF/Source/UAF/Public/AnimNextRigVMAsset.h b/Engine/Plugins/Experimental/UAF/UAF/Source/UAF/Public/AnimNextRigVMAsset.h
--- a/Engine/Plugins/Experimental/UAF/UAF/Source/UAF/Public/AnimNextRigVMAsset.h
+++ b/Engine/Plugins/Experimental/UAF/UAF/Source/UAF/Public/AnimNextRigVMAsset.h
@@ -83,6 +83,7 @@ protected:
 
 	// UObject interface
 	UE_API virtual void BeginDestroy() override;
+	UE_API virtual void Serialize(FArchive& Ar) override;
 	UE_API virtual void PostLoad() override;
 	UE_API virtual void GetAssetRegistryTags(FAssetRegistryTagsContext Context) const override;
 	UE_API virtual void PreDuplicate(FObjectDuplicationParameters& DupParams) override;
diff --git a/Engine/Plugins/Experimental/UAF/UAF/Source/UAFEditor/Private/AnimNextEditorModule.cpp b/Engine/Plugins/Experimental/UAF/UAF/Source/UAFEditor/Private/AnimNextEditorModule.cpp
--- a/Engine/Plugins/Experimental/UAF/UAF/Source/UAFEditor/Private/AnimNextEditorModule.cpp
+++ b/Engine/Plugins/Experimental/UAF/UAF/Source/UAFEditor/Private/AnimNextEditorModule.cpp
@@ -46,7 +46,6 @@
 #include "ToolMenus.h"
 #include "Framework/Docking/LayoutExtender.h"
 #include "Common/AnimNextAssetItemDetails.h"
-#include "Common/AnimNextRigVMAssetEditorDataCustomization.h"
 #include "Common/AssetPreview/AnimSequenceAssetPreview.h"
 #include "Common/AssetPreview/IUAFAssetPreview.h"
 #include "Common/AssetPreview/UAFAssetPreviewFactorySubsystem.h"
@@ -207,9 +206,6 @@ void FAnimNextEditorModule::StartupModule()
 	PropertyModule.RegisterCustomClassLayout("AnimNextVariableEntry", 
 		FOnGetDetailCustomizationInstance::CreateLambda([] { return MakeShared<FVariableCustomization>(); }));
 
-	PropertyModule.RegisterCustomClassLayout("UAFRigVMAssetEditorData",
-		FOnGetDetailCustomizationInstance::CreateLambda([] { return MakeShared<FAnimNextRigVMAssetEditorDataCustomization>(); }));
-
 	GraphPanelPinFactory = MakeShared<FAnimNextGraphPanelPinFactoryEditor>();
 	FEdGraphUtilities::RegisterVisualPinFactory(GraphPanelPinFactory);
 
diff --git a/Engine/Plugins/Experimental/UAF/UAF/Source/UAFEditor/Private/Common/AnimNextRigVMAssetEditorDataCustomization.cpp b/Engine/Plugins/Experimental/UAF/UAF/Source/UAFEditor/Private/Common/AnimNextRigVMAssetEditorDataCustomization.cpp
deleted file mode 100644
diff --git a/Engine/Plugins/Experimental/UAF/UAF/Source/UAFEditor/Private/Common/AnimNextRigVMAssetEditorDataCustomization.h b/Engine/Plugins/Experimental/UAF/UAF/Source/UAFEditor/Private/Common/AnimNextRigVMAssetEditorDataCustomization.h
deleted file mode 100644
diff --git a/Engine/Plugins/Experimental/UAF/UAF/Source/UAFEditor/Private/Module/AnimNextModuleFactory.cpp b/Engine/Plugins/Experimental/UAF/UAF/Source/UAFEditor/Private/Module/AnimNextModuleFactory.cpp
--- a/Engine/Plugins/Experimental/UAF/UAF/Source/UAFEditor/Private/Module/AnimNextModuleFactory.cpp
+++ b/Engine/Plugins/Experimental/UAF/UAF/Source/UAFEditor/Private/Module/AnimNextModuleFactory.cpp
@@ -36,7 +36,6 @@ UObject* UUAFSystemFactory::FactoryCreateNew(UClass* Class, UObject* InParent, F
 	// Create internal editor data
 	UUAFSystem_EditorData* EditorData = NewObject<UUAFSystem_EditorData>(NewModule, TEXT("EditorData"), RF_Transactional);
 	NewModule->EditorData = EditorData;
-	EditorData->bUsesExternalPackages = false;
 	EditorData->Initialize(/*bRecompileVM*/false);
 
 	UAnimNextEventGraphEntry* GraphEntry = EditorData->AddEventGraph(UUAFSystem_EditorData::DefaultEventGraphName, FRigUnit_AnimNextInitializeEvent::StaticStruct());
diff --git a/Engine/Plugins/Experimental/UAF/UAF/Source/UAFEditor/Private/Variables/AnimNextSharedVariablesAssetDefinition.cpp b/Engine/Plugins/Experimental/UAF/UAF/Source/UAFEditor/Private/Variables/AnimNextSharedVariablesAssetDefinition.cpp
--- a/Engine/Plugins/Experimental/UAF/UAF/Source/UAFEditor/Private/Variables/AnimNextSharedVariablesAssetDefinition.cpp
+++ b/Engine/Plugins/Experimental/UAF/UAF/Source/UAFEditor/Private/Variables/AnimNextSharedVariablesAssetDefinition.cpp
@@ -28,128 +28,4 @@ EAssetCommandResult UAssetDefinition_AnimNextSharedVariables::OpenAssets(const F
 	return EAssetCommandResult::Handled;
 }
 
-namespace UE::UAF::Editor
-{
-
-static FDelayedAutoRegisterHelper AutoRegisterGraphMenuItems(EDelayedRegisterRunPhase::EndOfEngineInit, []
-{
-	UToolMenus::RegisterStartupCallback(FSimpleMulticastDelegate::FDelegate::CreateLambda([]()
-	{
-		FToolMenuOwnerScoped OwnerScoped(UE_MODULE_NAME);
-
-		UToolMenu* Menu = UE::ContentBrowser::ExtendToolMenu_AssetContextMenu(UUAFSharedVariables::StaticClass());
-
-		FToolMenuSection& Section = Menu->FindOrAddSection("GetAssetActions");
-		Section.AddDynamicEntry(NAME_None, FNewToolMenuSectionDelegate::CreateLambda([](FToolMenuSection& InSection)
-		{
-			auto SetExternalPackagingStatus = [](const FToolMenuContext& InContext, bool bInUseExternalPackaging)
-			{
-				UContentBrowserAssetContextMenuContext* AssetContext = InContext.FindContext<UContentBrowserAssetContextMenuContext>();
-				if(AssetContext == nullptr)
-				{
-					return;
-				}
-
-				TArray<UUAFRigVMAsset*> Assets;
-				for (const FAssetData& AssetData : AssetContext->SelectedAssets)
-				{
-					if(!AssetData.GetClass()->IsChildOf(UUAFRigVMAsset::StaticClass()))
-					{
-						continue;
-					}
-
-					bool bUsesExternalPackages = true;
-					AssetData.GetTagValue<bool>(UUAFRigVMAssetEditorData::GetUsesExternalPackagesPropertyName(), bUsesExternalPackages);
-
-					if(bUsesExternalPackages == bInUseExternalPackaging)
-					{
-						// Status is as requested, skip
-						continue;
-					}
-					
-					UUAFRigVMAsset* Asset = Cast<UUAFRigVMAsset>(AssetData.GetAsset());
-					if(Asset == nullptr)
-					{
-						continue;
-					}
-
-					Assets.Add(Asset);
-
-					UUAFRigVMAssetEditorData* EditorData = UncookedOnly::FUtils::GetEditorData<UUAFRigVMAssetEditorData>(Asset);
-					if(EditorData == nullptr)
-					{
-						continue;
-					}
-
-					ensure(EditorData->IsUsingExternalPackages() == bUsesExternalPackages);
-				}
-
-				UUAFRigVMAssetEditorData::SetUseExternalPackages(Assets, bInUseExternalPackaging);
-			};
-
-			auto ExternalPackagingStatusMatches = [](const FToolMenuContext& InContext, bool bInUsesExternalPackaging)
-			{
-				UContentBrowserAssetContextMenuContext* AssetContext = InContext.FindContext<UContentBrowserAssetContextMenuContext>();
-				if(AssetContext == nullptr)
-				{
-					return false;
-				}
-
-				for (const FAssetData& AssetData : AssetContext->SelectedAssets)
-				{
-					if(!AssetData.GetClass()->IsChildOf(UUAFRigVMAsset::StaticClass()))
-					{
-						continue;
-					}
-
-					bool bUsesExternalPackages = true;
-					AssetData.GetTagValue<bool>(UUAFRigVMAssetEditorData::GetUsesExternalPackagesPropertyName(), bUsesExternalPackages);
-
-					if(bUsesExternalPackages == bInUsesExternalPackaging)
-					{
-						return true;
-					}
-				}
-
-				return false;
-			};
-
-			FToolUIAction UIAction;
-			UIAction.ExecuteAction = FToolMenuExecuteAction::CreateLambda([&SetExternalPackagingStatus](const FToolMenuContext& InContext)
-			{
-				SetExternalPackagingStatus(InContext, true);
-			});
-			UIAction.IsActionVisibleDelegate = FToolMenuIsActionButtonVisible::CreateLambda([&ExternalPackagingStatusMatches](const FToolMenuContext& InContext)
-			{
-				return ExternalPackagingStatusMatches(InContext, false);
-			});
-
-			InSection.AddMenuEntry(
-				"EnableExternalPackages",
-				LOCTEXT("EnableExternalPackagesLabel", "Use External Packages"),
-				LOCTEXT("EnableExternalPackagesTooltip", "Set the asset(s) to use external packaging for its entries (graphs, variables etc.)\nThis will create the external packages for all entries, add them to version control if enabled and save all packages.\nWarning: This operation cannot be undone, so a connection to version control is recommended."),
-				FSlateIcon(FAppStyle::Get().GetStyleSetName(), "Icons.Edit"),
-				UIAction);
-
-			UIAction.ExecuteAction = FToolMenuExecuteAction::CreateLambda([&SetExternalPackagingStatus](const FToolMenuContext& InContext)
-			{
-				SetExternalPackagingStatus(InContext, false);
-			});
-			UIAction.IsActionVisibleDelegate = FToolMenuIsActionButtonVisible::CreateLambda([&ExternalPackagingStatusMatches](const FToolMenuContext& InContext)
-			{
-				return ExternalPackagingStatusMatches(InContext, true);
-			});
-
-			InSection.AddMenuEntry(
-				"DisableExternalPackages",
-				LOCTEXT("DisableExternalPackagesLabel", "Use Single Package"),
-				LOCTEXT("DisableExternalPackagesTooltip", "Set the asset(s) to use a single package.\nThis will remove any external packages for existing entries, remove them from version control if enabled and save all packages.\nWarning: This operation cannot be undone, so a connection to version control is recommended."),
-				FSlateIcon(FAppStyle::Get().GetStyleSetName(), "Icons.Edit"),
-				UIAction);
-		}));
-	}));
-});
-
-}
-
 #undef LOCTEXT_NAMESPACE
diff --git a/Engine/Plugins/Experimental/UAF/UAF/Source/UAFEditor/Private/Variables/AnimNextSharedVariablesFactory.cpp b/Engine/Plugins/Experimental/UAF/UAF/Source/UAFEditor/Private/Variables/AnimNextSharedVariablesFactory.cpp
--- a/Engine/Plugins/Experimental/UAF/UAF/Source/UAFEditor/Private/Variables/AnimNextSharedVariablesFactory.cpp
+++ b/Engine/Plugins/Experimental/UAF/UAF/Source/UAFEditor/Private/Variables/AnimNextSharedVariablesFactory.cpp
@@ -34,7 +34,6 @@ UObject* UUAFSharedVariablesFactory::FactoryCreateNew(UClass* Class, UObject* In
 	// Create internal editor data
 	UUAFSharedVariables_EditorData* EditorData = NewObject<UUAFSharedVariables_EditorData>(NewDataInterface, TEXT("EditorData"), RF_Transactional);
 	NewDataInterface->EditorData = EditorData;
-	EditorData->bUsesExternalPackages = false;
 	EditorData->Initialize(/*bRecompileVM*/false);
 
 	// Compile the initial skeleton
diff --git a/Engine/Plugins/Experimental/UAF/UAF/Source/UAFUncookedOnly/Internal/AnimNextRigVMAssetEditorData.h b/Engine/Plugins/Experimental/UAF/UAF/Source/UAFUncookedOnly/Internal/AnimNextRigVMAssetEditorData.h
--- a/Engine/Plugins/Experimental/UAF/UAF/Source/UAFUncookedOnly/Internal/AnimNextRigVMAssetEditorData.h
+++ b/Engine/Plugins/Experimental/UAF/UAF/Source/UAFUncookedOnly/Internal/AnimNextRigVMAssetEditorData.h
@@ -211,17 +211,6 @@ public:
 	/** Reorders variable entry to appear before provided BeforeVariableEntry */
 	UE_API void ReorderVariable(UAnimNextVariableEntry* VariableEntry, const UAnimNextVariableEntry* BeforeVariableEntry, bool bSetupUndoRedo = true, bool bPrintPythonCommand = true);
 
-	// Get the external packaging status of this asset
-	bool IsUsingExternalPackages() const { return bUsesExternalPackages; }
-
-#if WITH_EDITOR
-	// Switch the external packaging status of this asset
-	static UE_API void SetUseExternalPackages(TArrayView<UUAFRigVMAsset*> InAssets, bool bInUseExternalPackages);
-
-	// UI helper function
-	static FName GetUsesExternalPackagesPropertyName() { return GET_MEMBER_NAME_CHECKED(UUAFRigVMAssetEditorData, bUsesExternalPackages); }
-#endif
-
 	// Report an error to the user, typically used for scripting APIs
 	static UE_API void ReportError(const TCHAR* InMessage);
 
@@ -552,9 +541,6 @@ protected:
 	UE_API void BroadcastModified(EAnimNextEditorDataNotifType InType, UObject* InSubject);
 
 	UE_API void ReconstructAllNodes();
-
-	// Called from PostLoad to load external packages
-	UE_API void PostLoadExternalPackages();
 	
 	UE_API void RequestVMRecompilation();
 
@@ -672,13 +658,13 @@ protected:
 
 	void SetVMRecompilationRequired(bool bInIsRequired);
 
-	/** All entries in this asset - not saved, discovered at load time and also contains InternalEntries */
-	UPROPERTY(Transient, DuplicateTransient)
+	/** All entries in this asset */
+	UPROPERTY()
 	TArray<TObjectPtr<UUAFRigVMAssetEntry>> Entries;
 
-	/** Internal entries in this asset */
+	UE_DEPRECATED(5.7, "This property is deprecated. Please use Entries instead")
 	UPROPERTY()
-	TArray<TObjectPtr<UUAFRigVMAssetEntry>> InternalEntries;
+	TArray<TObjectPtr<UUAFRigVMAssetEntry>> InternalEntries_DEPRECATED;
 
 	UPROPERTY()
 	FRigVMClient RigVMClient;
@@ -729,8 +715,9 @@ protected:
 	UPROPERTY(transient, DuplicateTransient)
 	bool bIsCompiling = false;
 
-	UPROPERTY(VisibleAnywhere, Category = "External Packages", AssetRegistrySearchable)
-	bool bUsesExternalPackages = true;
+	UE_DEPRECATED(5.7, "Storing asset entries as external packages has been deprecated")
+	UPROPERTY()
+	bool bUsesExternalPackages_DEPRECATED = false;
 
 	UPROPERTY()
 	FName DefaultInjectionSite_DEPRECATED;
diff --git a/Engine/Plugins/Experimental/UAF/UAF/Source/UAFUncookedOnly/Internal/Entries/AnimNextSharedVariablesEntry.h b/Engine/Plugins/Experimental/UAF/UAF/Source/UAFUncookedOnly/Internal/Entries/AnimNextSharedVariablesEntry.h
--- a/Engine/Plugins/Experimental/UAF/UAF/Source/UAFUncookedOnly/Internal/Entries/AnimNextSharedVariablesEntry.h
+++ b/Engine/Plugins/Experimental/UAF/UAF/Source/UAFUncookedOnly/Internal/Entries/AnimNextSharedVariablesEntry.h
@@ -60,6 +60,10 @@ class UUAFSharedVariablesEntry : public UUAFRigVMAssetEntry
 	friend UE::UAF::Editor::FFunctionsOutlinerMode;
 	friend class UAnimNextAssetFindReplaceVariables;
 
+	// UObject interface
+	UAFUNCOOKEDONLY_API virtual void Serialize(FArchive& Ar) override;
+	UAFUNCOOKEDONLY_API virtual void PostLoad() override;
+
 	// UUAFRigVMAssetEntry interface
 	virtual void Initialize(UUAFRigVMAssetEditorData* InEditorData) override;
 	virtual FName GetEntryName() const override;
diff --git a/Engine/Plugins/Experimental/UAF/UAF/Source/UAFUncookedOnly/Private/AnimNextRigVMAssetEditorData.cpp b/Engine/Plugins/Experimental/UAF/UAF/Source/UAFUncookedOnly/Private/AnimNextRigVMAssetEditorData.cpp
--- a/Engine/Plugins/Experimental/UAF/UAF/Source/UAFUncookedOnly/Private/AnimNextRigVMAssetEditorData.cpp
+++ b/Engine/Plugins/Experimental/UAF/UAF/Source/UAFUncookedOnly/Private/AnimNextRigVMAssetEditorData.cpp
@@ -140,6 +140,24 @@ void UUAFRigVMAssetEditorData::Serialize(FArchive& Ar)
 
 		bUpgradeDataInterfacesOnLoad = true;
 	}
+	
+	if (Ar.IsLoading() && Ar.CustomVer(FFortniteMainBranchObjectVersion::GUID) < FFortniteMainBranchObjectVersion::DeprecateUAFExternalPackages)
+	{
+		PRAGMA_DISABLE_DEPRECATION_WARNINGS
+		if (bUsesExternalPackages_DEPRECATED)
+		{
+			UE_LOG(LogAnimation, Error, TEXT("External packages support for UAF has been deprecated in 5.7, please resave content as non-external packages in 5.6 before upgrading."));
+       	}
+		// Copy any internal entries to the main entries array and clear out deprecated array
+		Entries.Append(InternalEntries_DEPRECATED);
+		InternalEntries_DEPRECATED.Empty();
+		PRAGMA_ENABLE_DEPRECATION_WARNINGS
+	}
+	
+	for (UUAFRigVMAssetEntry* Entry : Entries)
+	{
+		Entry->ConditionalPreload();
+	}
 }
 
 void UUAFRigVMAssetEditorData::Initialize(bool bRecompileVM)
@@ -217,13 +235,10 @@ void UUAFRigVMAssetEditorData::PostLoad()
 	Super::PostLoad();
 
 	GraphModels.Reset();
-	
-	PostLoadExternalPackages();
 
-	// Preload our entries, as we need them for RefreshExternalModels
+	// Postload our entries, as we need them for RefreshExternalModels
 	for (UUAFRigVMAssetEntry* Entry : Entries)
 	{
-		Entry->ConditionalPreload();
 		Entry->ConditionalPostLoad();
 	}
 
@@ -240,25 +255,6 @@ void UUAFRigVMAssetEditorData::PostLoad()
 	FCoreUObjectDelegates::OnEndLoadPackage.AddUObject(this, &UUAFRigVMAssetEditorData::HandlePackageDone);
 }
 
-void UUAFRigVMAssetEditorData::PostLoadExternalPackages()
-{
-	if(bUsesExternalPackages)
-	{
-		FExternalPackageHelper::LoadObjectsFromExternalPackages<UUAFRigVMAssetEntry>(this, [this](UUAFRigVMAssetEntry* InLoadedEntry)
-		{
-			check(IsValid(InLoadedEntry));
-			InLoadedEntry->Initialize(this);
-			Entries.Add(InLoadedEntry);
-		});
-	}
-
-	// Internal entries should be empty if we are externally packaged
-	ensure(!bUsesExternalPackages || InternalEntries.IsEmpty());
-
-	// Copy any internal entries to the main entries array
-	Entries.Append(InternalEntries);
-}
-
 void UUAFRigVMAssetEditorData::RequestVMRecompilation()
 {
 	if (VMRecompilationBracket == 0)
@@ -818,12 +814,6 @@ TObjectPtr<URigVMGraph> UUAFRigVMAssetEditorData::CreateContainedGraphModel(URig
 	URigVMGraph* CollapseNodeModelRootGraph = CollapseNode->GetRootGraph();
 	check(CollapseNodeModelRootGraph);
 
-	// If we are a transient asset, or not using external packages dont use external packages
-	if (bUsesExternalPackages && !CollapseNodeModelRootGraph->HasAnyFlags(RF_Transient))
-	{
-		Model->SetExternalPackage(CollapseNodeModelRootGraph->GetExternalPackage());
-	}
-
 	return Model;
 }
 
@@ -1830,13 +1820,13 @@ void UUAFRigVMAssetEditorData::ReorderCategory(const FString& CategoryName, cons
 
 void UUAFRigVMAssetEditorData::ReorderVariable(UAnimNextVariableEntry* VariableEntry, const UAnimNextVariableEntry* BeforeVariableEntry, bool bSetupUndoRedo, bool bPrintPythonCommand)
 {
-	const int32 VariableEntryIndex = InternalEntries.IndexOfByKey(VariableEntry);
+	const int32 VariableEntryIndex = Entries.IndexOfByKey(VariableEntry);
 	if (VariableEntryIndex == INDEX_NONE)
 	{
 		return;
 	}
 	
-	const int32 BeforeVariableEntryIndex = InternalEntries.IndexOfByKey(BeforeVariableEntry);
+	const int32 BeforeVariableEntryIndex = Entries.IndexOfByKey(BeforeVariableEntry);
 	if (BeforeVariableEntryIndex == INDEX_NONE)
 	{
 		return;
@@ -2041,11 +2031,6 @@ UObject* UUAFRigVMAssetEditorData::CreateNewSubEntry(UUAFRigVMAssetEditorData* I
 	UUAFRigVMAsset* Asset = UE::UAF::UncookedOnly::FUtils::GetAsset(InEditorData);
 	check(Asset);
 
-	// Additionally check external packaging flag
-	if(!Asset->HasAnyFlags(RF_Transient) && InEditorData->bUsesExternalPackages)
-	{
-		FExternalPackageHelper::SetPackagingMode(NewEntry, InEditorData, true, false, PKG_None);
-	}
 	return NewEntry;
 }
 
@@ -3008,198 +2993,14 @@ void UUAFRigVMAssetEditorData::AddEntryInternal(UUAFRigVMAssetEntry* InEntry)
 
 void UUAFRigVMAssetEditorData::InsertEntryInternal(UUAFRigVMAssetEntry* InEntry, int32 InsertionIndex)
 {
-	// If we are using external packages, do not persist this entry
-	if(bUsesExternalPackages)
-	{
-		Entries.Insert(InEntry, InsertionIndex);
-	}
-	else
-	{
-		InternalEntries.Insert(InEntry, InsertionIndex);
-		Entries.Insert(InEntry, InsertionIndex);
-	}
+	Entries.Insert(InEntry, InsertionIndex);
 }
 
 void UUAFRigVMAssetEditorData::RemoveEntryInternal(UUAFRigVMAssetEntry* InEntry)
 {
-	if(bUsesExternalPackages)
-	{
-		Entries.Remove(InEntry);
-	}
-	else
-	{
-		InternalEntries.Remove(InEntry);
-		Entries.Remove(InEntry);
-	}
-}
-
-#if WITH_EDITOR
-
-void UUAFRigVMAssetEditorData::SetUseExternalPackages(TArrayView<UUAFRigVMAsset*> InAssets, bool bInUseExternalPackages)
-{
-	TArray<UUAFRigVMAssetEditorData*> EditorDatas;
-	for(UUAFRigVMAsset* Asset : InAssets)
-	{
-		if(Asset == nullptr)
-		{
-			continue;
-		}
-
-		UUAFRigVMAssetEditorData* EditorData = UE::UAF::UncookedOnly::FUtils::GetEditorData<UUAFRigVMAssetEditorData>(Asset);
-		if(EditorData == nullptr)
-		{
-			continue;
-		}
-
-		if(bInUseExternalPackages != EditorData->bUsesExternalPackages)
-		{
-			EditorDatas.Add(EditorData);
-		}
-	}
-
-	if(EditorDatas.Num() == 0)
-	{
-		return;
-	}
-
-	if(bInUseExternalPackages)
-	{
-		TArray<UPackage*> PackagesToCheckOut;
-		TArray<UPackage*> PackagesToSave;
-		TArray<UPackage*> PackagesToAdd;
-
-		for(UUAFRigVMAssetEditorData* EditorData : EditorDatas)
-		{
-			UPackage* Package = EditorData->GetPackage();
-			PackagesToCheckOut.Add(Package);
-			PackagesToSave.Add(Package);
-		}
-
-		// Prompt the user to check out this package, allowing user to decide against this operation
-		if(!FEditorFileUtils::PromptToCheckoutPackages(false, PackagesToCheckOut))
-		{
-			return;
-		}
-
-		FScopedSlowTask SlowTask(3.0f, LOCTEXT("ConvertingAssets", "Converting Assets"));
-		SlowTask.MakeDialog();
-		SlowTask.EnterProgressFrame(1.0f, LOCTEXT("SettingPackagingStatus", "Setting Packaging Status"));
-
-		for(UUAFRigVMAssetEditorData* EditorData : EditorDatas)
-		{
-			EditorData->MarkPackageDirty();
-
-			// Set all internal entries to use external packages
-			TArray<UPackage*> ExternalPackages;
-			for (UUAFRigVMAssetEntry* Entry : EditorData->InternalEntries)
-			{
-				FExternalPackageHelper::SetPackagingMode(Entry, EditorData, bInUseExternalPackages, true, PKG_None);
-				UPackage* ExternalPackage = Entry->GetExternalPackage();
-
-				// Switch any graphs to be packaged externally
-				if(IUAFRigVMGraphInterface* GraphEntry = Cast<IUAFRigVMGraphInterface>(Entry))
-				{
-					GraphEntry->GetRigVMGraph()->SetExternalPackage(ExternalPackage);
-				}
-
-				check(ExternalPackage);
-				PackagesToAdd.Add(ExternalPackage);
-				PackagesToSave.Add(ExternalPackage);
-			}
-
-			// Clear all internal packages, switch to discovery on PostLoad rather than serialized entries
-			EditorData->InternalEntries.Empty();
-			EditorData->bUsesExternalPackages = bInUseExternalPackages;
-		}
-
-		SlowTask.EnterProgressFrame(1.0f, LOCTEXT("AddOrRevertVersionControl", "Adding/Reverting In Version Control"));
-
-		// Add (or revert delete) packages to source control
-		FPackageSourceControlHelper SCCHelper;
-		bool bAdded = SCCHelper.AddToSourceControl(PackagesToAdd);
-
-		SlowTask.EnterProgressFrame(1.0f, LOCTEXT("SavingPackages", "Saving Packages"));
-
-		// Finally save all packages, they need to be consistent on disk after this operation
-		FEditorFileUtils::FPromptForCheckoutAndSaveParams SaveParams;
-		SaveParams.bAlreadyCheckedOut = true;
-		SaveParams.bCanBeDeclined = false;
-		SaveParams.bPromptToSave = false;
-		FEditorFileUtils::PromptForCheckoutAndSave(PackagesToSave, SaveParams);
-	}
-	else
-	{
-		// Gather packages we will modify/delete
-		TArray<UPackage*> PackagesToCheckOut;
-		TArray<UPackage*> PackagesToSave;
-		TArray<UObject*> ObjectsToDelete;
-		for(UUAFRigVMAssetEditorData* EditorData : EditorDatas)
-		{
-			UPackage* ThisPackage = EditorData->GetPackage();
-			PackagesToCheckOut.Add(ThisPackage);
-			for (UUAFRigVMAssetEntry* Entry : EditorData->Entries)
-			{
-				UPackage* ExternalPackage = Entry->GetExternalPackage();
-				check(ExternalPackage);
-				ObjectsToDelete.Add(ExternalPackage);
-				PackagesToCheckOut.Add(ExternalPackage);
-			}
-		}
-
-		// Prompt the user to check out files, allowing user to decide against this operation
-		if(!FEditorFileUtils::PromptToCheckoutPackages(false, PackagesToCheckOut))
-		{
-			return;
-		}
-
-		FScopedSlowTask SlowTask(3.0f, LOCTEXT("ConvertingAssets", "Converting Assets"));
-		SlowTask.MakeDialog();
-		SlowTask.EnterProgressFrame(1.0f, LOCTEXT("SettingPackagingStatus", "Setting Packaging Status"));
-
-		for(UUAFRigVMAssetEditorData* EditorData : EditorDatas)
-		{
-			EditorData->MarkPackageDirty();
-
-			ensure(EditorData->InternalEntries.IsEmpty());
-			EditorData->InternalEntries.Empty();
-
-			// Set all entries to not use external packages
-			for (UUAFRigVMAssetEntry* Entry : EditorData->Entries)
-			{
-				FExternalPackageHelper::SetPackagingMode(Entry, EditorData, bInUseExternalPackages, true, PKG_None);
-
-				// Switch any graphs to be packaged internally
-				if(IUAFRigVMGraphInterface* GraphEntry = Cast<IUAFRigVMGraphInterface>(Entry))
-				{
-					GraphEntry->GetRigVMGraph()->SetExternalPackage(nullptr);
-				}
-			}
-
-			// Ensure we save all of our entries if we are not using external packages
-			EditorData->InternalEntries.Append(EditorData->Entries);
-			EditorData->bUsesExternalPackages = bInUseExternalPackages;
-
-			PackagesToSave.Add(EditorData->GetPackage());
-		}
-
-		SlowTask.EnterProgressFrame(1.0f, LOCTEXT("DeletingOldPackages", "Deleting Old Packages"));
-
-		// Delete the old external packages
-		ObjectTools::DeleteObjectsUnchecked(ObjectsToDelete);
-
-		SlowTask.EnterProgressFrame(1.0f, LOCTEXT("SavingPackages", "Saving Packages"));
-
-		// Finally save our packages, they needs to be consistent on disk after this operation
-		FEditorFileUtils::FPromptForCheckoutAndSaveParams SaveParams;
-		SaveParams.bAlreadyCheckedOut = true;
-		SaveParams.bCanBeDeclined = false;
-		SaveParams.bPromptToSave = false;
-		FEditorFileUtils::PromptForCheckoutAndSave(PackagesToSave, SaveParams);
-	}
+	Entries.Remove(InEntry);
 }
 
-#endif
-
 FInstancedPropertyBag UUAFRigVMAssetEditorData::GenerateCombinedPropertyBag(const FRigVMCompileSettings& InSettings, const FAnimNextGetVariableCompileContext& InCompileContext) const
 {
 	// TODO:
diff --git a/Engine/Plugins/Experimental/UAF/UAF/Source/UAFUncookedOnly/Private/AnimNextSharedVariableNode.cpp b/Engine/Plugins/Experimental/UAF/UAF/Source/UAFUncookedOnly/Private/AnimNextSharedVariableNode.cpp
--- a/Engine/Plugins/Experimental/UAF/UAF/Source/UAFUncookedOnly/Private/AnimNextSharedVariableNode.cpp
+++ b/Engine/Plugins/Experimental/UAF/UAF/Source/UAFUncookedOnly/Private/AnimNextSharedVariableNode.cpp
@@ -41,8 +41,6 @@ void UUAFSharedVariableNode::PostLoad()
 				{
 					if (const UUAFSharedVariables* SharedVariablesAsset = Asset.Get())
 					{
-						const_cast<UUAFSharedVariables*>(SharedVariablesAsset)->ConditionalPostLoad();
-						
 						const UUAFSharedVariables_EditorData* EditorData = UE::UAF::UncookedOnly::FUtils::GetEditorData<UUAFSharedVariables_EditorData>(SharedVariablesAsset);
 						const FRigVMTemplateArgumentType ExpectedVariableType(FName(*GetCPPType()), GetCPPTypeObject());
 
@@ -184,8 +182,6 @@ void UUAFSharedVariableNode::UpdateCachedVariableGUID()
 			{
 				const UUAFSharedVariables_EditorData* EditorData = UE::UAF::UncookedOnly::FUtils::GetEditorData<const UUAFSharedVariables_EditorData>(Asset.Get());
 
-				// FindEntry depends on UUAFSharedVariables_EditorData::PostLoadExternalPackages having been called (copying internal_entries into entries array)
-				const_cast<UUAFSharedVariables_EditorData*>(EditorData)->ConditionalPostLoad();
 				if (const IUAFRigVMVariableInterface* VariableInterface = EditorData != nullptr ? Cast<IUAFRigVMVariableInterface>(EditorData->FindEntry(GetVariableName())) : nullptr)
 				{
 					CachedGuid = VariableInterface->GetGuid();
diff --git a/Engine/Plugins/Experimental/UAF/UAF/Source/UAFUncookedOnly/Private/AnimNextUncookedOnlyModule.cpp b/Engine/Plugins/Experimental/UAF/UAF/Source/UAFUncookedOnly/Private/AnimNextUncookedOnlyModule.cpp
--- a/Engine/Plugins/Experimental/UAF/UAF/Source/UAFUncookedOnly/Private/AnimNextUncookedOnlyModule.cpp
+++ b/Engine/Plugins/Experimental/UAF/UAF/Source/UAFUncookedOnly/Private/AnimNextUncookedOnlyModule.cpp
@@ -60,7 +60,6 @@ FGuid FModule::GetVariableGuidByName(const FName InName, const UObject* InObject
 	if (InObject)
 	{
 		const_cast<UObject*>(InObject)->ConditionalPreload();
-		const_cast<UObject*>(InObject)->ConditionalPostLoad();
 		
 		auto GetGuidFromEditorData = [InName, &Guid](const UUAFRigVMAssetEditorData* EditorData)
 		{
@@ -115,7 +114,6 @@ FName FModule::GetVariableNameByGuid(const FGuid InGuid, const UObject* InObject
 	if (InObject)
 	{
 		const_cast<UObject*>(InObject)->ConditionalPreload();
-		const_cast<UObject*>(InObject)->ConditionalPostLoad();
 	
 		if (const UUAFRigVMAsset* Asset = Cast<UUAFRigVMAsset>(InObject))
 		{
diff --git a/Engine/Plugins/Experimental/UAF/UAF/Source/UAFUncookedOnly/Private/Entries/AnimNextSharedVariablesEntry.cpp b/Engine/Plugins/Experimental/UAF/UAF/Source/UAFUncookedOnly/Private/Entries/AnimNextSharedVariablesEntry.cpp
--- a/Engine/Plugins/Experimental/UAF/UAF/Source/UAFUncookedOnly/Private/Entries/AnimNextSharedVariablesEntry.cpp
+++ b/Engine/Plugins/Experimental/UAF/UAF/Source/UAFUncookedOnly/Private/Entries/AnimNextSharedVariablesEntry.cpp
@@ -18,15 +18,55 @@
 
 #define LOCTEXT_NAMESPACE "AnimNextSharedVariablesEntry"
 
+void UUAFSharedVariablesEntry::Serialize(FArchive& Ar)
+{
+	Super::Serialize(Ar);
+
+#if WITH_EDITORONLY_DATA
+	if (Ar.IsLoading())
+	{
+		// Make sure Asset and Struct are loaded
+		if (Asset)
+		{
+			UUAFSharedVariables* NonConstAsset = const_cast<UUAFSharedVariables*>(Asset.Get());
+			NonConstAsset->ConditionalPreload();
+		}
+
+		if (Struct)
+		{
+			UScriptStruct* NonConstStruct = const_cast<UScriptStruct*>(Struct.Get());
+			NonConstStruct->ConditionalPreload();
+		}
+	}
+#endif
+}
+
+void UUAFSharedVariablesEntry::PostLoad()
+{
+	Super::PostLoad();
+
+#if WITH_EDITORONLY_DATA
+	// Make sure Asset and Struct are loaded
+	if (Asset)
+	{
+		UUAFSharedVariables* NonConstAsset = const_cast<UUAFSharedVariables*>(Asset.Get());
+		NonConstAsset->ConditionalPostLoad();
+	}
+	
+	if (Struct)
+	{
+		UScriptStruct* NonConstStruct = const_cast<UScriptStruct*>(Struct.Get());
+		NonConstStruct->ConditionalPostLoad();
+	}
+#endif
+}
+
 void UUAFSharedVariablesEntry::Initialize(UUAFRigVMAssetEditorData* InEditorData)
 {
 	Super::Initialize(InEditorData);
 
 	if(Asset)
 	{
-		const_cast<UUAFSharedVariables*>(Asset.Get())->ConditionalPreload();
-		const_cast<UUAFSharedVariables*>(Asset.Get())->ConditionalPostLoad();
-		
 		UUAFSharedVariables_EditorData* EditorData = UE::UAF::UncookedOnly::FUtils::GetEditorData<UUAFSharedVariables_EditorData>(Asset.Get());
 		EditorData->ModifiedDelegate.AddUObject(this, &UUAFSharedVariablesEntry::HandleAssetModified);
 	}
diff --git a/Engine/Plugins/Experimental/UAF/UAF/Source/UAFUncookedOnly/Private/Entries/AnimNextVariableEntry.cpp b/Engine/Plugins/Experimental/UAF/UAF/Source/UAFUncookedOnly/Private/Entries/AnimNextVariableEntry.cpp
--- a/Engine/Plugins/Experimental/UAF/UAF/Source/UAFUncookedOnly/Private/Entries/AnimNextVariableEntry.cpp
+++ b/Engine/Plugins/Experimental/UAF/UAF/Source/UAFUncookedOnly/Private/Entries/AnimNextVariableEntry.cpp
@@ -247,36 +247,40 @@ void UAnimNextVariableEntry::Serialize(FArchive& Ar)
 {
 	Super::Serialize(Ar);
 	Ar.UsingCustomVersion(FFortniteMainBranchObjectVersion::GUID);
-}
 
-void UAnimNextVariableEntry::PostLoad()
-{
-	Super::PostLoad();
-	if(GetLinkerCustomVersion(FFortniteMainBranchObjectVersion::GUID) < FFortniteMainBranchObjectVersion::AnimNextModuleRefactor)
+	if (Ar.IsLoading())
 	{
-		// Add a property for this type
-		DefaultValue.Reset();
-		DefaultValue.AddProperties({ FPropertyBagPropertyDesc(IUAFRigVMVariableInterface::ValueName, Type.GetContainerType(), Type.GetValueType(), Type.GetValueTypeObject()) });
-
-		// Copy any default value from the module's defaults, if found
-		UUAFRigVMAsset* Asset = GetTypedOuter<UUAFRigVMAsset>();
-		FString FullName = Asset->GetPathName() + TEXT(":") + GetVariableName().ToString();
-		const FPropertyBagPropertyDesc* OldDesc = Asset->VariableDefaults.GetPropertyBagStruct() ? Asset->VariableDefaults.GetPropertyBagStruct()->FindPropertyDescByName(*FullName) : nullptr;
-		const FPropertyBagPropertyDesc* NewDesc = DefaultValue.GetPropertyBagStruct() ? DefaultValue.GetPropertyBagStruct()->FindPropertyDescByName(IUAFRigVMVariableInterface::ValueName) : nullptr;
-		if(OldDesc && NewDesc)
+		if (GetLinkerCustomVersion(FFortniteMainBranchObjectVersion::GUID) < FFortniteMainBranchObjectVersion::AnimNextModuleRefactor)
 		{
-			const void* OldValue = OldDesc->CachedProperty->ContainerPtrToValuePtr<void>(Asset->VariableDefaults.GetValue().GetMemory());
-			void* NewValue = NewDesc->CachedProperty->ContainerPtrToValuePtr<void>(DefaultValue.GetMutableValue().GetMemory());
-			NewDesc->CachedProperty->CopyCompleteValue(NewValue, OldValue);
+			// Add a property for this type
+			DefaultValue.Reset();
+			DefaultValue.AddProperties({ FPropertyBagPropertyDesc(IUAFRigVMVariableInterface::ValueName, Type.GetContainerType(), Type.GetValueType(), Type.GetValueTypeObject()) });
+
+			// Copy any default value from the module's defaults, if found
+			UUAFRigVMAsset* Asset = GetTypedOuter<UUAFRigVMAsset>();
+			FString FullName = Asset->GetPathName() + TEXT(":") + GetVariableName().ToString();
+			const FPropertyBagPropertyDesc* OldDesc = Asset->VariableDefaults.GetPropertyBagStruct() ? Asset->VariableDefaults.GetPropertyBagStruct()->FindPropertyDescByName(*FullName) : nullptr;
+			const FPropertyBagPropertyDesc* NewDesc = DefaultValue.GetPropertyBagStruct() ? DefaultValue.GetPropertyBagStruct()->FindPropertyDescByName(IUAFRigVMVariableInterface::ValueName) : nullptr;
+			if (OldDesc && NewDesc)
+			{
+				const void* OldValue = OldDesc->CachedProperty->ContainerPtrToValuePtr<void>(Asset->VariableDefaults.GetValue().GetMemory());
+				void* NewValue = NewDesc->CachedProperty->ContainerPtrToValuePtr<void>(DefaultValue.GetMutableValue().GetMemory());
+				NewDesc->CachedProperty->CopyCompleteValue(NewValue, OldValue);
+			}
 		}
-	}
 
-	if (GetLinkerCustomVersion(FFortniteMainBranchObjectVersion::GUID) < FFortniteMainBranchObjectVersion::UAFVariablesGuid)
-	{
-		Guid = FGuid::NewGuid();
+		if (GetLinkerCustomVersion(FFortniteMainBranchObjectVersion::GUID) < FFortniteMainBranchObjectVersion::UAFVariablesGuid)
+		{
+			Guid = FGuid::NewGuid();
+		}
 	}
 }
 
+void UAnimNextVariableEntry::PostLoad()
+{
+	Super::PostLoad();
+}
+
 #if WITH_EDITOR
 
 void UAnimNextVariableEntry::PostEditChangeProperty(FPropertyChangedEvent& PropertyChangedEvent)
diff --git a/Engine/Plugins/Experimental/UAF/UAF/Source/UAFUncookedOnly/Private/UncookedOnlyUtils.cpp b/Engine/Plugins/Experimental/UAF/UAF/Source/UAFUncookedOnly/Private/UncookedOnlyUtils.cpp
--- a/Engine/Plugins/Experimental/UAF/UAF/Source/UAFUncookedOnly/Private/UncookedOnlyUtils.cpp
+++ b/Engine/Plugins/Experimental/UAF/UAF/Source/UAFUncookedOnly/Private/UncookedOnlyUtils.cpp
@@ -952,8 +952,8 @@ FAnimNextParamType FUtils::FindVariableType(const FAnimNextVariableReference& In
 	}
 	else if (const UUAFRigVMAsset* Asset = Cast<UUAFRigVMAsset>(InVariableReference.GetObject()))
 	{
-		const UUAFRigVMAssetEditorData* AssetEditorData = CastChecked<UUAFRigVMAssetEditorData>(Asset->EditorData);
-
+		UUAFRigVMAssetEditorData* AssetEditorData = CastChecked<UUAFRigVMAssetEditorData>(Asset->EditorData);
+		AssetEditorData->ConditionalPreload();
 		if (const IUAFRigVMVariableInterface* Variable = Cast<IUAFRigVMVariableInterface>(AssetEditorData->FindEntry(InVariableReference.GetName())))
 		{
 			return Variable->GetType();
diff --git a/Engine/Plugins/Experimental/UAF/UAF/Source/UAFUncookedOnly/Private/Variables/K2Node_AnimNextComponentSetVariable.cpp b/Engine/Plugins/Experimental/UAF/UAF/Source/UAFUncookedOnly/Private/Variables/K2Node_AnimNextComponentSetVariable.cpp
--- a/Engine/Plugins/Experimental/UAF/UAF/Source/UAFUncookedOnly/Private/Variables/K2Node_AnimNextComponentSetVariable.cpp
+++ b/Engine/Plugins/Experimental/UAF/UAF/Source/UAFUncookedOnly/Private/Variables/K2Node_AnimNextComponentSetVariable.cpp
@@ -140,7 +140,6 @@ void UK2Node_AnimNextComponentSetVariable::PreloadRequiredAssets()
 	{
 		UObject* Object = const_cast<UObject*>(VariableReference.GetObject().Get());
 		PreloadObject(Object);
-		Object->ConditionalPostLoad();
 	}
 }
 
diff --git a/Engine/Plugins/Experimental/UAF/UAFAnimGraph/Source/UAFAnimGraphEditor/Private/Graph/AnimNextAnimationGraphFactory.cpp b/Engine/Plugins/Experimental/UAF/UAFAnimGraph/Source/UAFAnimGraphEditor/Private/Graph/AnimNextAnimationGraphFactory.cpp
--- a/Engine/Plugins/Experimental/UAF/UAFAnimGraph/Source/UAFAnimGraphEditor/Private/Graph/AnimNextAnimationGraphFactory.cpp
+++ b/Engine/Plugins/Experimental/UAF/UAFAnimGraph/Source/UAFAnimGraphEditor/Private/Graph/AnimNextAnimationGraphFactory.cpp
@@ -36,7 +36,6 @@ UObject* UUAFAnimGraphFactory::FactoryCreateNew(UClass* Class, UObject* InParent
 	// Create internal editor data
 	UUAFAnimGraph_EditorData* EditorData = NewObject<UUAFAnimGraph_EditorData>(NewAnimationGraph, TEXT("EditorData"), RF_Transactional);
 	NewAnimationGraph->EditorData = EditorData;
-	EditorData->bUsesExternalPackages = false;
 	EditorData->Initialize(/*bRecompileVM*/false);
 
 	// Add a single internal graph entry
diff --git a/Engine/Plugins/Experimental/UAF/UAFLayering/Source/UAFLayeringEditor/Private/UAFLayerStackFactory.cpp b/Engine/Plugins/Experimental/UAF/UAFLayering/Source/UAFLayeringEditor/Private/UAFLayerStackFactory.cpp
--- a/Engine/Plugins/Experimental/UAF/UAFLayering/Source/UAFLayeringEditor/Private/UAFLayerStackFactory.cpp
+++ b/Engine/Plugins/Experimental/UAF/UAFLayering/Source/UAFLayeringEditor/Private/UAFLayerStackFactory.cpp
@@ -35,7 +35,6 @@ UObject* UUAFLayerStackFactory::FactoryCreateNew(UClass* Class, UObject* InParen
 	// Set the layer controllers editor data reference
 	LayerStack->EditorData = EditorData;
 
-	EditorData->bUsesExternalPackages = false;
 	EditorData->Initialize(/*bRecompileVM*/false);
 
 	UE::UAF::UncookedOnly::Compilation::RequestAssetCompilation(LayerStack);
diff --git a/Engine/Plugins/Experimental/UAF/UAFStateTree/Source/UAFStateTreeEditor/Private/AnimNextStateTreeFactory.cpp b/Engine/Plugins/Experimental/UAF/UAFStateTree/Source/UAFStateTreeEditor/Private/AnimNextStateTreeFactory.cpp
--- a/Engine/Plugins/Experimental/UAF/UAFStateTree/Source/UAFStateTreeEditor/Private/AnimNextStateTreeFactory.cpp
+++ b/Engine/Plugins/Experimental/UAF/UAFStateTree/Source/UAFStateTreeEditor/Private/AnimNextStateTreeFactory.cpp
@@ -41,7 +41,6 @@ UObject* UAnimNextStateTreeFactory::FactoryCreateNew(UClass* Class, UObject* InP
 	UStateTree* StateTree = CastChecked<UStateTree>(StateTreeFactory->FactoryCreateNew(UStateTree::StaticClass(), NewStateTree, FName(Name.ToString() + "_StateTree"), RF_Transactional, nullptr, nullptr));
 	NewStateTree->StateTree = StateTree;
 
-	EditorData->bUsesExternalPackages = false;
 	EditorData->Initialize(/*bRecompileVM*/false);
 
 	// Compile the initial skeleton
diff --git a/Engine/Source/Runtime/Core/Public/UObject/FortniteMainBranchObjectVersions.inl b/Engine/Source/Runtime/Core/Public/UObject/FortniteMainBranchObjectVersions.inl
--- a/Engine/Source/Runtime/Core/Public/UObject/FortniteMainBranchObjectVersions.inl
+++ b/Engine/Source/Runtime/Core/Public/UObject/FortniteMainBranchObjectVersions.inl
@@ -754,5 +754,8 @@ FORTNITE_MAIN_VERSION(OpenColorIOv251, 246)
 // Reparameterize Spline in SplineComponent based on SplineCurves parameterization
 FORTNITE_MAIN_VERSION(SplineComponentReparameterizeOnLoad, 247)
 
+// Deprecating support for external UPackage storing UAF asset entries
+FORTNITE_MAIN_VERSION(DeprecateUAFExternalPackages, 248)
+
 // -----<new versions can be added above this line>-------------------------------------------------
 #undef FORTNITE_MAIN_VERSION
