UAT engine code changes to allow dynamic setting of a component's UAF asset.
- Move UAFSystem to Public.
- Add public way to call UUAFComponent::SetAsset that can only be used when component is not registered.

#Synced-CL 50197869

diff --git a/Engine/Plugins/Experimental/UAF/UAF/Source/UAF/Internal/Module/AnimNextModule.h b/Engine/Plugins/Experimental/UAF/UAF/Source/UAF/Internal/Module/AnimNextModule.h
deleted file mode 100644
diff --git a/Engine/Plugins/Experimental/UAF/UAF/Source/UAF/Internal/Module/AnimNextModuleInstance.h b/Engine/Plugins/Experimental/UAF/UAF/Source/UAF/Internal/Module/AnimNextModuleInstance.h
--- a/Engine/Plugins/Experimental/UAF/UAF/Source/UAF/Internal/Module/AnimNextModuleInstance.h
+++ b/Engine/Plugins/Experimental/UAF/UAF/Source/UAF/Internal/Module/AnimNextModuleInstance.h
@@ -3,7 +3,7 @@
 #pragma once
 
 #include "CoreMinimal.h"
-#include "AnimNextModule.h"
+#include "Module/AnimNextModule.h"
 #include "AnimNextPoolHandle.h"
 #include "Variables/UAFInstanceVariableDataProxy.h"
 #include "UAFAssetInstance.h"
diff --git a/Engine/Plugins/Experimental/UAF/UAF/Source/UAF/Private/Component/AnimNextComponent.cpp b/Engine/Plugins/Experimental/UAF/UAF/Source/UAF/Private/Component/AnimNextComponent.cpp
--- a/Engine/Plugins/Experimental/UAF/UAF/Source/UAF/Private/Component/AnimNextComponent.cpp
+++ b/Engine/Plugins/Experimental/UAF/UAF/Source/UAF/Private/Component/AnimNextComponent.cpp
@@ -470,6 +470,14 @@ FUAFWeakSystemReference UUAFComponent::GetSystemReference() const
 	return FUAFWeakSystemReference(SystemReference);
 }
 
+void UUAFComponent::SetAsset(TInstancedStruct<FUAFSystemFactoryAsset>&& InAssetData)
+{
+	if (ensure(!IsRegistered()))
+	{
+		SetAssetInternal(InAssetData);
+	}
+}
+
 void UUAFComponent::SetAssetFromObject(const UObject* InObject)
 {
 	TInstancedStruct<FUAFSystemFactoryAsset> AssetDataFromFactory = UE::UAF::FAssetDataFactory::CreateUAFAssetDataFromObject<FUAFSystemFactoryAsset>(InObject);
@@ -479,12 +487,12 @@ void UUAFComponent::SetAssetFromObject(const UObject* InObject)
 	}
 }
 
-void UUAFComponent::SetAsset(UE::UAF::FSystemAssetHandle&& InAssetData)
+void UUAFComponent::SetAssetInternal(UE::UAF::FSystemAssetHandle&& InAssetData)
 {
 	AssetData = MoveTemp(InAssetData);
 }
 
-void UUAFComponent::SetAsset(TConstStructView<FUAFSystemFactoryAsset> InAssetData)
+void UUAFComponent::SetAssetInternal(TConstStructView<FUAFSystemFactoryAsset> InAssetData)
 {
 	AssetData = InAssetData;
 }
diff --git a/Engine/Plugins/Experimental/UAF/UAF/Source/UAF/Public/Component/AnimNextComponent.h b/Engine/Plugins/Experimental/UAF/UAF/Source/UAF/Public/Component/AnimNextComponent.h
--- a/Engine/Plugins/Experimental/UAF/UAF/Source/UAF/Public/Component/AnimNextComponent.h
+++ b/Engine/Plugins/Experimental/UAF/UAF/Source/UAF/Public/Component/AnimNextComponent.h
@@ -193,10 +193,13 @@ public:
 
 	TConstStructView<FUAFSystemFactoryAsset> GetAssetData() const { return AssetData; }
 
+	// can only be called when component is not registered
+	UAF_API void SetAsset(TInstancedStruct<FUAFSystemFactoryAsset>&& AssetData);
+
 private:
 	UAF_API void SetAssetFromObject(const UObject* InAsset);
-	UAF_API void SetAsset(TInstancedStruct<FUAFSystemFactoryAsset>&& AssetData);
-	UAF_API void SetAsset(TConstStructView<FUAFSystemFactoryAsset> AssetData);
+	UAF_API void SetAssetInternal(TConstStructView<FUAFSystemFactoryAsset> AssetData);
+	UAF_API void SetAssetInternal(TInstancedStruct<FUAFSystemFactoryAsset>&& AssetData);
 
 	UAF_API void RegisterSystem();
 	UAF_API void UnregisterSystem();
diff --git a/Engine/Plugins/Experimental/UAF/UAF/Source/UAF/Public/Module/AnimNextModule.h b/Engine/Plugins/Experimental/UAF/UAF/Source/UAF/Public/Module/AnimNextModule.h
new file mode 100644
--- /dev/null
+++ b/Engine/Plugins/Experimental/UAF/UAF/Source/UAF/Public/Module/AnimNextModule.h
@@ -0,0 +1,90 @@
+// Copyright Epic Games, Inc. All Rights Reserved.
+
+#pragma once
+
+#include "CoreMinimal.h"
+#include "Variables/AnimNextSharedVariables.h"
+#include "RigVMCore/RigVM.h"
+#include "RigVMHost.h"
+#include "Graph/AnimNextGraphState.h"
+#include "Module/RigVMTrait_ModuleEventDependency.h"
+
+#include "AnimNextModule.generated.h"
+
+#define UE_API UAF_API
+
+class UEdGraph;
+class UUAFSystem;
+class UAnimGraphNode_AnimNextGraph;
+struct FAnimNode_AnimNextGraph;
+struct FRigUnit_AnimNextGraphEvaluator;
+struct FAnimNextGraphInstance;
+struct FAnimNextScheduleGraphTask;
+struct FAnimNextEditorParam;
+struct FAnimNextParam;
+struct FAnimNextModuleInstance;
+
+namespace UE::UAF
+{
+	struct FContext;
+	struct FExecutionContext;
+	class FAnimNextModuleImpl;
+	struct FTestUtils;
+	struct FParametersProxy;
+}
+
+namespace UE::UAF::UncookedOnly
+{
+	struct FUtils;
+}
+
+namespace UE::UAF::Editor
+{
+	class FModuleEditor;
+	class FVariableCustomization;
+}
+
+// Root asset represented by a component when instantiated
+UCLASS(MinimalAPI, BlueprintType, DisplayName="UAF System")
+class UUAFSystem : public UUAFSharedVariables
+{
+	GENERATED_BODY()
+
+public:
+	UE_API UUAFSystem(const FObjectInitializer& ObjectInitializer);
+
+	// UObject interface
+	UE_API virtual void Serialize(FArchive& Ar) override;
+
+protected:
+	friend class UUAFSystemFactory;
+	friend class UUAFSystem_EditorData;
+	friend struct UE::UAF::UncookedOnly::FUtils;
+	friend class UE::UAF::Editor::FModuleEditor;
+	friend struct UE::UAF::FTestUtils;
+	friend FAnimNextGraphInstance;
+	friend class UAnimGraphNode_AnimNextGraph;
+	friend UE::UAF::FExecutionContext;
+	friend struct FAnimNextScheduleGraphTask;
+	friend UE::UAF::FAnimNextModuleImpl;
+	friend class UE::UAF::Editor::FVariableCustomization;
+	friend struct UE::UAF::FParametersProxy;
+	friend FAnimNextModuleInstance;
+
+	// All dependencies that should be set up when the module initializes
+	UPROPERTY()
+	TArray<TInstancedStruct<FRigVMTrait_ModuleEventDependency>> Dependencies;
+
+#if WITH_EDITORONLY_DATA
+	UPROPERTY()
+	TArray<TObjectPtr<UScriptStruct>> RequiredComponents_DEPRECATED;
+
+	UPROPERTY()
+	FAnimNextGraphState DefaultState_DEPRECATED;
+	
+	UPROPERTY()
+	FInstancedPropertyBag PropertyBag_DEPRECATED;
+#endif
+};
+
+#undef UE_API
diff --git a/Engine/Plugins/Experimental/UAF/UAFAnimGraph/Source/UAFAnimGraphEditor/Private/AnimGraphViewportController.cpp b/Engine/Plugins/Experimental/UAF/UAFAnimGraph/Source/UAFAnimGraphEditor/Private/AnimGraphViewportController.cpp
--- a/Engine/Plugins/Experimental/UAF/UAFAnimGraph/Source/UAFAnimGraphEditor/Private/AnimGraphViewportController.cpp
+++ b/Engine/Plugins/Experimental/UAF/UAFAnimGraph/Source/UAFAnimGraphEditor/Private/AnimGraphViewportController.cpp
@@ -35,7 +35,7 @@ namespace UE::UAF::Editor
 
 				UUAFComponent* UAFComponent = NewObject<UUAFComponent>(PreviewActor);
 				PreviewActor->AddInstanceComponent(UAFComponent);
-				UAFComponent->SetAsset(FUAFSystemFactoryAsset_System(Module));
+				UAFComponent->SetAssetInternal(FUAFSystemFactoryAsset_System(Module));
 				UAFComponent->RegisterComponent();
 				UAFComponent->InitializeComponent();
 
diff --git a/Engine/Plugins/Experimental/UAF/UAFTests/Source/UAFTests/Private/UAFBlueprintTests.cpp b/Engine/Plugins/Experimental/UAF/UAFTests/Source/UAFTests/Private/UAFBlueprintTests.cpp
--- a/Engine/Plugins/Experimental/UAF/UAFTests/Source/UAFTests/Private/UAFBlueprintTests.cpp
+++ b/Engine/Plugins/Experimental/UAF/UAFTests/Source/UAFTests/Private/UAFBlueprintTests.cpp
@@ -94,7 +94,7 @@ TEST_CLASS(FAssetBlueprintTests, "Animation.UAF.Functional")
 			.Then(TEXT("Set the UAF Module and add the AnimNext Component to the Character Blueprint and Open in Editor"), [&]()
 			{
 				TArray<UActorComponent*> Components;
-				AnimNextComponent->SetAsset(FUAFSystemFactoryAsset_System(Cast<UUAFSystem>(UAFAsset)));
+				AnimNextComponent->SetAssetInternal(FUAFSystemFactoryAsset_System(Cast<UUAFSystem>(UAFAsset)));
 				Components.Add(AnimNextComponent);
 				FKismetEditorUtilities::AddComponentsToBlueprint(Blueprint, Components);
 				ASSERT_THAT(IsTrue(AssetEditorSubsystem->OpenEditorForAsset(Blueprint)));
