Add missing MoveTemp in UUAFComponent::SetAssetFromObject

#Synced-CL 50197869

diff --git a/Engine/Plugins/Experimental/UAF/UAF/Source/UAF/Private/Component/AnimNextComponent.cpp b/Engine/Plugins/Experimental/UAF/UAF/Source/UAF/Private/Component/AnimNextComponent.cpp
--- a/Engine/Plugins/Experimental/UAF/UAF/Source/UAF/Private/Component/AnimNextComponent.cpp
+++ b/Engine/Plugins/Experimental/UAF/UAF/Source/UAF/Private/Component/AnimNextComponent.cpp
@@ -474,7 +474,7 @@ void UUAFComponent::SetAsset(TInstancedStruct<FUAFSystemFactoryAsset>&& InAssetD
 {
 	if (ensure(!IsRegistered()))
 	{
-		SetAssetInternal(InAssetData);
+		SetAssetInternal(MoveTemp(InAssetData));
 	}
 }
 
