RigVM: Set VM pointer on extended context correctly for UAF Component

#rnx

#Synced-CL 50244857

diff --git a/Engine/Plugins/Experimental/UAF/UAF/Source/UAF/Private/Script/UAFRigVMComponent.cpp b/Engine/Plugins/Experimental/UAF/UAF/Source/UAF/Private/Script/UAFRigVMComponent.cpp
--- a/Engine/Plugins/Experimental/UAF/UAF/Source/UAF/Private/Script/UAFRigVMComponent.cpp
+++ b/Engine/Plugins/Experimental/UAF/UAF/Source/UAF/Private/Script/UAFRigVMComponent.cpp
@@ -48,6 +48,7 @@ void FUAFRigVMComponent::OnBindToInstance()
 
 	// Now initialize the 'instance', cache memory handles etc. in the context
 	VM->InitializeInstance(ExtendedExecuteContext);
+	ExtendedExecuteContext.SetVM(VM);
 
 #if DO_CHECK
 	bIsInitialized = true;
