Fix asset icons from UAF rename

#rnx

#Synced-CL 50197033

diff --git a/Engine/Plugins/Experimental/UAF/UAF/Source/UAFEditor/Private/UAFStyle.cpp b/Engine/Plugins/Experimental/UAF/UAF/Source/UAFEditor/Private/UAFStyle.cpp
--- a/Engine/Plugins/Experimental/UAF/UAF/Source/UAFEditor/Private/UAFStyle.cpp
+++ b/Engine/Plugins/Experimental/UAF/UAF/Source/UAFEditor/Private/UAFStyle.cpp
@@ -40,8 +40,8 @@ void FUAFStyle::SetClassIcons()
 	const FVector2D Icon16x16(16.0f, 16.0f);
 	
 	Set("ClassIcon.UAFSharedVariables", new IMAGE_PLUGIN_BRUSH_SVG("Icons/Archive_16.svg", Icon16x16));
-	Set("ClassIcon.AnimNextModule", new IMAGE_BRUSH_SVG("Starship/Common/Modules", Icon16x16 ));
-	Set("ClassIcon.AnimNextAnimationGraph", new IMAGE_BRUSH_SVG("Starship/Blueprints/icon_BlueprintEditor_EventGraph", Icon16x16));
+	Set("ClassIcon.UAFSystem", new IMAGE_BRUSH_SVG("Starship/Common/Modules", Icon16x16 ));
+	Set("ClassIcon.UAFAnimGraph", new IMAGE_BRUSH_SVG("Starship/Blueprints/icon_BlueprintEditor_EventGraph", Icon16x16));
 	Set("ClassIcon.AnimNextStateTree", new IMAGE_BRUSH_SVG(TEXT("UMG/TreeView"), Icon16x16));
 	Set("ClassIcon.AbstractSkeletonLabelCollection", new IMAGE_PLUGIN_BRUSH_SVG("Icons/LabelCollection_16.svg", Icon16x16));
 	Set("ClassIcon.AbstractSkeletonLabelBinding", new IMAGE_PLUGIN_BRUSH_SVG("Icons/LabelBinding_16.svg", Icon16x16));
