Crash when updating variables get out of sync in Shared Variables asset
#fix ensure layout if force-updated when underlying variable type (and thus value property-bag) changes

#Synced-CL 50296236

diff --git a/Engine/Plugins/Experimental/UAF/UAF/Source/UAFEditor/Private/Variables/VariableCustomization.cpp b/Engine/Plugins/Experimental/UAF/UAF/Source/UAFEditor/Private/Variables/VariableCustomization.cpp
--- a/Engine/Plugins/Experimental/UAF/UAF/Source/UAFEditor/Private/Variables/VariableCustomization.cpp
+++ b/Engine/Plugins/Experimental/UAF/UAF/Source/UAFEditor/Private/Variables/VariableCustomization.cpp
@@ -123,6 +123,30 @@ void FVariableCustomization::CustomizeDetails(IDetailLayoutBuilder& DetailBuilde
 			]
 		];
 	}
+	
+	// Register modified notify callback to force update the details customization (handle mismatch of property value widget vs new type) 
+	for (const TWeakObjectPtr<UObject>& CustomizedObject : Objects)
+	{
+		if (CustomizedObject.Get())
+		{
+			if (UUAFRigVMAssetEditorData* OuterEditorData = CustomizedObject->GetTypedOuter<UUAFRigVMAssetEditorData>())
+			{
+				if (WeakLayoutBuilder.IsValid())
+				{
+					OuterEditorData->ModifiedDelegate.AddSPLambda(this, [LocalWeakLayoutBuilder = WeakLayoutBuilder, WeakObject = CustomizedObject](UUAFRigVMAssetEditorData* InEditorData, EAnimNextEditorDataNotifType InType, const UObject* InSubject)
+					{
+						if (InType == EAnimNextEditorDataNotifType::VariableTypeChanged && WeakObject.Get() == InSubject) 
+						{
+							if (auto SharedLayoutBuilder = LocalWeakLayoutBuilder.Pin())
+							{
+								SharedLayoutBuilder->ForceRefreshDetails();
+							}
+						}
+					});
+				}
+			}
+		}
+	}
 
 	// Don't customize default value if we have multi-selection
 	if(Objects.Num() > 1)
@@ -169,6 +193,7 @@ void FVariableCustomization::CustomizeDetails(IDetailLayoutBuilder& DetailBuilde
 
 void FVariableCustomization::CustomizeDetails(const TSharedPtr<IDetailLayoutBuilder>& DetailBuilder) 
 {
+	WeakLayoutBuilder = DetailBuilder;
 	CustomizeDetails(*DetailBuilder);
 }
 
diff --git a/Engine/Plugins/Experimental/UAF/UAF/Source/UAFEditor/Private/Variables/VariableCustomization.h b/Engine/Plugins/Experimental/UAF/UAF/Source/UAFEditor/Private/Variables/VariableCustomization.h
--- a/Engine/Plugins/Experimental/UAF/UAF/Source/UAFEditor/Private/Variables/VariableCustomization.h
+++ b/Engine/Plugins/Experimental/UAF/UAF/Source/UAFEditor/Private/Variables/VariableCustomization.h
@@ -46,6 +46,7 @@ private:
 	TArray<TWeakObjectPtr<UAnimNextVariableEntry>> Variables;
 	TWeakObjectPtr<UUAFRigVMAsset> OuterRigVMAsset;
 
+	TWeakPtr<IDetailLayoutBuilder> WeakLayoutBuilder;
 	FString DefaultCategoryName;
 };
 
