Crash when changing UAF Graph variable to array
#fix do not process asset pin reference in case they're arrays

#Synced-CL 50296236

diff --git a/Engine/Plugins/Experimental/UAF/UAF/Source/UAFUncookedOnly/Private/UncookedOnlyUtils.cpp b/Engine/Plugins/Experimental/UAF/UAF/Source/UAFUncookedOnly/Private/UncookedOnlyUtils.cpp
--- a/Engine/Plugins/Experimental/UAF/UAF/Source/UAFUncookedOnly/Private/UncookedOnlyUtils.cpp
+++ b/Engine/Plugins/Experimental/UAF/UAF/Source/UAFUncookedOnly/Private/UncookedOnlyUtils.cpp
@@ -1350,7 +1350,7 @@ void FUtils::ProcessPinAssetReferences(const URigVMPin* InPin, FWorkspaceOutline
 					auto HandleParentPinVariableLink = [&HandleSoftObjectPath](const IUAFRigVMVariableInterface* InVariableInterface, const URigVMPin* InPin)
 					{
 						const FAnimNextParamType& ParameterType = InVariableInterface->GetType();
-						if (ParameterType.GetValueType() == FAnimNextParamType::EValueType::Struct)
+						if (ParameterType.GetValueType() == FAnimNextParamType::EValueType::Struct && ParameterType.GetContainerType() == FAnimNextParamType::EContainerType::None)
 						{
 							if (const UScriptStruct* TypeStruct = Cast<UScriptStruct>(ParameterType.GetValueTypeObject()))
 							{
