Support EndTickGroup in UAF

#Synced-CL 50439659

diff --git a/Engine/Plugins/Experimental/UAF/UAF/Source/UAF/Internal/Module/RigUnit_AnimNextModuleEvents.h b/Engine/Plugins/Experimental/UAF/UAF/Source/UAF/Internal/Module/RigUnit_AnimNextModuleEvents.h
--- a/Engine/Plugins/Experimental/UAF/UAF/Source/UAF/Internal/Module/RigUnit_AnimNextModuleEvents.h
+++ b/Engine/Plugins/Experimental/UAF/UAF/Source/UAF/Internal/Module/RigUnit_AnimNextModuleEvents.h
@@ -32,6 +32,9 @@ struct FRigUnit_AnimNextModuleEventBase : public FRigUnit_AnimNextBase
 	// Get the tick group for this event
 	virtual ETickingGroup GetTickGroup() const { return TG_PrePhysics; }
 
+	// Get the end tick group for this event
+	virtual ETickingGroup GetEndTickGroup() const { return TG_PrePhysics; }
+
 	// Get the integer sort key for this event (used to sort events in tick groups - smaller values get sorted earlier)
 	virtual int32 GetSortOrder() const { return 0; }
 
@@ -121,6 +124,7 @@ struct FRigUnit_AnimNextUserEvent : public FRigUnit_AnimNextModuleEventBase
 	virtual UE::UAF::EScriptEventPhase GetEventPhase() const override final { return UE::UAF::EScriptEventPhase::Execute; }
 	UE_API virtual UE::UAF::FTickFunctionBindingFunction GetBindingFunction() const override final;
 	virtual ETickingGroup GetTickGroup() const override final { return TickGroup; }
+	virtual ETickingGroup GetEndTickGroup() const override final { return EndTickGroup; }
 	virtual int32 GetSortOrder() const override final { return SortOrder; }
 
 	// The name of the event
@@ -134,6 +138,10 @@ struct FRigUnit_AnimNextUserEvent : public FRigUnit_AnimNextModuleEventBase
 	// The tick group the event executes in 
 	UPROPERTY()
 	TEnumAsByte<ETickingGroup> TickGroup = ETickingGroup::TG_PrePhysics;
+
+	// The tick group the event completes by 
+	UPROPERTY(EditAnywhere, Category = "Tick", meta = (Input, Constant, DetailsOnly))
+	TEnumAsByte<ETickingGroup> EndTickGroup = ETickingGroup::TG_PrePhysics;
 };
 
 /** Schedule event called before world physics is updated */
diff --git a/Engine/Plugins/Experimental/UAF/UAF/Source/UAF/Private/Module/AnimNextModuleInstance.cpp b/Engine/Plugins/Experimental/UAF/UAF/Source/UAF/Private/Module/AnimNextModuleInstance.cpp
--- a/Engine/Plugins/Experimental/UAF/UAF/Source/UAF/Private/Module/AnimNextModuleInstance.cpp
+++ b/Engine/Plugins/Experimental/UAF/UAF/Source/UAF/Private/Module/AnimNextModuleInstance.cpp
@@ -141,6 +141,7 @@ void FAnimNextModuleInstance::Initialize(const FUAFSystemFactoryParams& InFactor
 				TickFunction.ModuleInstance = this;
 				TickFunction.Event = ModuleEvent.Event;
 				TickFunction.TickGroup = ModuleEvent.TickGroup;
+				TickFunction.EndTickGroup = ModuleEvent.EndTickGroup;
 				TickFunction.bUserEvent = ModuleEvent.bUserEvent;
 
 				// Perform custom setup
diff --git a/Engine/Plugins/Experimental/UAF/UAF/Source/UAF/Private/Script/UAFRigVMComponent.cpp b/Engine/Plugins/Experimental/UAF/UAF/Source/UAF/Private/Script/UAFRigVMComponent.cpp
--- a/Engine/Plugins/Experimental/UAF/UAF/Source/UAF/Private/Script/UAFRigVMComponent.cpp
+++ b/Engine/Plugins/Experimental/UAF/UAF/Source/UAF/Private/Script/UAFRigVMComponent.cpp
@@ -137,6 +137,7 @@ TConstArrayView<UE::UAF::FScriptEventInfo> FUAFRigVMComponent::GetScriptEvents()
 			NewEvent.Binding = Event.GetBindingFunction();
 			NewEvent.Phase = Event.GetEventPhase();
 			NewEvent.TickGroup = Event.GetTickGroup();
+			NewEvent.EndTickGroup = Event.GetEndTickGroup();
 			NewEvent.SortOrder = Event.GetSortOrder();
 			NewEvent.bUserEvent = Event.IsUserEvent();
 			NewEvent.bIsTask = Event.IsTask();
@@ -165,6 +166,12 @@ TConstArrayView<UE::UAF::FScriptEventInfo> FUAFRigVMComponent::GetScriptEvents()
 						check(Operand.GetMemoryType() == ERigVMMemoryType::Literal);
 						NewEvent.SortOrder = *VM->LiteralMemoryStorage.GetData<int32>(Operand.GetRegisterIndex());
 					}
+					else if (OperandName == GET_MEMBER_NAME_CHECKED(FRigUnit_AnimNextUserEvent, EndTickGroup))
+					{
+						check(Operand.GetMemoryType() == ERigVMMemoryType::Literal);
+						NewEvent.EndTickGroup = *VM->LiteralMemoryStorage.GetData<
+							TEnumAsByte<enum ETickingGroup>>(Operand.GetRegisterIndex());
+					}
 				}
 			}
 
diff --git a/Engine/Plugins/Experimental/UAF/UAF/Source/UAF/Public/Script/UAFScriptComponent.h b/Engine/Plugins/Experimental/UAF/UAF/Source/UAF/Public/Script/UAFScriptComponent.h
--- a/Engine/Plugins/Experimental/UAF/UAF/Source/UAF/Public/Script/UAFScriptComponent.h
+++ b/Engine/Plugins/Experimental/UAF/UAF/Source/UAF/Public/Script/UAFScriptComponent.h
@@ -1,4 +1,4 @@
-﻿// Copyright Epic Games, Inc. All Rights Reserved.
+// Copyright Epic Games, Inc. All Rights Reserved.
 
 #pragma once
 
@@ -75,6 +75,9 @@ struct FScriptEventInfo
 	// The tick group the event runs in
 	ETickingGroup TickGroup = ETickingGroup::TG_PrePhysics;
 
+	// The tick group the event completes in
+	ETickingGroup EndTickGroup = ETickingGroup::TG_PrePhysics;
+
 	// Whether this is a user-authored event, or a compiler-generated event
 	uint8 bUserEvent : 1 = true;
 
