- Added ensures to catch creation of FAnimNext(Soft)VariableReference to internal variables
- Marked GetVariableGuidByName and GetVariableNameByGuid WITH_EDITOR (already in UncookedModule)
- Prevent attempting to resolve variable reference to function internal (programmatically generated) variables for instance variable mapping explicitly

#Synced-CL 50532483

diff --git a/Engine/Plugins/Experimental/UAF/UAF/Source/UAF/Private/Variables/AnimNextVariableReference.cpp b/Engine/Plugins/Experimental/UAF/UAF/Source/UAF/Private/Variables/AnimNextVariableReference.cpp
--- a/Engine/Plugins/Experimental/UAF/UAF/Source/UAF/Private/Variables/AnimNextVariableReference.cpp
+++ b/Engine/Plugins/Experimental/UAF/UAF/Source/UAF/Private/Variables/AnimNextVariableReference.cpp
@@ -18,6 +18,7 @@ FAnimNextVariableReference FAnimNextVariableReference::FromProperty(const FPrope
 	check(Name != NAME_None);
 	FGuid Guid;
 #if WITH_EDITOR
+	ensureAlwaysMsgf(!UE::UAF::IsInternalVariableName(Name), TEXT("Attempting to create variable reference using a variable name indicative of a programmatic variable %s"), *Name.ToString());
 	Guid = UE::UAF::UncookedOnly::IAnimNextUncookedOnlyModule::Get().GetVariableGuidByName(Name, InStruct);
 	if (!Guid.IsValid())
 	{
@@ -33,6 +34,8 @@ FAnimNextVariableReference FAnimNextVariableReference::FromName(FName InName, co
 	check(InObject && InName != NAME_None);
 	FGuid Guid;
 #if WITH_EDITOR
+	ensureAlwaysMsgf(!UE::UAF::IsInternalVariableName(InName), TEXT("Attempting to create variable reference using a variable name indicative of a programmatic variable %s"), *InName.ToString());
+
 	Guid = UE::UAF::UncookedOnly::IAnimNextUncookedOnlyModule::Get().GetVariableGuidByName(InName, InObject);
 	if (!Guid.IsValid())
 	{
@@ -54,6 +57,7 @@ FAnimNextVariableReference::FAnimNextVariableReference(FName InName, FGuid InGui
 	
 #if WITH_EDITORONLY_DATA 
 	check(Name != NAME_None	|| InGuid.IsValid());
+	ensureAlwaysMsgf(!UE::UAF::IsInternalVariableName(Name), TEXT("Attempting to create variable reference using a variable name indicative of a programmatic variable %s"), *Name.ToString());
 #else
 	check(Name != NAME_None);
 #endif // WITH_EDITORONLY_DATA
@@ -66,6 +70,7 @@ FAnimNextVariableReference::FAnimNextVariableReference(const FAnimNextSoftVariab
 	check(Object == nullptr || Object->IsA<UUAFRigVMAsset>() || Object->IsA<UScriptStruct>());
 
 #if WITH_EDITOR
+	ensureAlwaysMsgf(!UE::UAF::IsInternalVariableName(Name), TEXT("Attempting to create variable reference using a variable name indicative of a programmatic variable %s"), *Name.ToString());
 	if (Object != nullptr)
 	{
 		CachedGuid = UE::UAF::UncookedOnly::IAnimNextUncookedOnlyModule::Get().GetVariableGuidByName(Name, Object);
@@ -82,6 +87,7 @@ FAnimNextVariableReference::FAnimNextVariableReference(FName InName, const UUAFR
 	, Object(InAsset)
 {
 	check(InAsset != nullptr);
+	ensureAlwaysMsgf(!UE::UAF::IsInternalVariableName(Name), TEXT("Attempting to create variable reference using a variable name indicative of a programmatic variable %s"), *Name.ToString());
 }
 
 FAnimNextVariableReference::FAnimNextVariableReference(FName InName, const UScriptStruct* InStruct)
@@ -89,6 +95,7 @@ FAnimNextVariableReference::FAnimNextVariableReference(FName InName, const UScri
 	, Object(InStruct)
 {
 	check(InStruct != nullptr);
+	ensureAlwaysMsgf(!UE::UAF::IsInternalVariableName(Name), TEXT("Attempting to create variable reference using a variable name indicative of a programmatic variable %s"), *Name.ToString());
 }
 
 #if WITH_EDITOR
diff --git a/Engine/Plugins/Experimental/UAF/UAF/Source/UAF/Private/Variables/UAFInstanceVariableData.cpp b/Engine/Plugins/Experimental/UAF/UAF/Source/UAF/Private/Variables/UAFInstanceVariableData.cpp
--- a/Engine/Plugins/Experimental/UAF/UAF/Source/UAF/Private/Variables/UAFInstanceVariableData.cpp
+++ b/Engine/Plugins/Experimental/UAF/UAF/Source/UAF/Private/Variables/UAFInstanceVariableData.cpp
@@ -207,9 +207,12 @@ void FUAFInstanceVariableData::RebuildNameMaps()
 					check(Descs.Num() == NumVariablesInSet);
 					for (int32 VariableIndex = 0; VariableIndex < NumVariablesInSet; ++VariableIndex)
 					{
-						FUAFInstanceVariableData::FVariableMapping Mapping(SharedVariableSetIndex, VariableIndex);
-						VariablesMap.Add(FAnimNextVariableReference::FromName(Descs[VariableIndex].Name, Asset), Mapping);
-						LegacyVariableNameMap.Add(Descs[VariableIndex].Name, Mapping);
+						if (!IsInternalVariableName(Descs[VariableIndex].Name))
+						{
+							FUAFInstanceVariableData::FVariableMapping Mapping(SharedVariableSetIndex, VariableIndex);
+							VariablesMap.Add(FAnimNextVariableReference::FromName(Descs[VariableIndex].Name, Asset), Mapping);
+							LegacyVariableNameMap.Add(Descs[VariableIndex].Name, Mapping);
+						}
 					}
 				}
 				break;
@@ -224,10 +227,13 @@ void FUAFInstanceVariableData::RebuildNameMaps()
 					check(Properties.Num() == NumVariablesInSet);
 					for (int32 VariableIndex = 0; VariableIndex < NumVariablesInSet; ++VariableIndex)
 					{
-						FName Name = Properties[VariableIndex].Property->GetFName();
-						FUAFInstanceVariableData::FVariableMapping Mapping(SharedVariableSetIndex, VariableIndex);
-						VariablesMap.Add(FAnimNextVariableReference::FromProperty(Properties[VariableIndex].Property, Struct), Mapping);
-						LegacyVariableNameMap.Add(Name, Mapping);
+						const FName Name = Properties[VariableIndex].Property->GetFName();
+						if (!IsInternalVariableName(Name))
+						{
+							FUAFInstanceVariableData::FVariableMapping Mapping(SharedVariableSetIndex, VariableIndex);
+							VariablesMap.Add(FAnimNextVariableReference::FromProperty(Properties[VariableIndex].Property, Struct), Mapping);
+							LegacyVariableNameMap.Add(Name, Mapping);
+						}
 					}
 				}
 				break;
diff --git a/Engine/Plugins/Experimental/UAF/UAF/Source/UAF/Public/Variables/AnimNextSoftVariableReference.h b/Engine/Plugins/Experimental/UAF/UAF/Source/UAF/Public/Variables/AnimNextSoftVariableReference.h
--- a/Engine/Plugins/Experimental/UAF/UAF/Source/UAF/Public/Variables/AnimNextSoftVariableReference.h
+++ b/Engine/Plugins/Experimental/UAF/UAF/Source/UAF/Public/Variables/AnimNextSoftVariableReference.h
@@ -19,7 +19,10 @@ struct FAnimNextSoftVariableReference
 		return FAnimNextSoftVariableReference(InName, InSoftObjectPath);
 	}
 
-	explicit FAnimNextSoftVariableReference(const FAnimNextVariableReference& VariableReference) : Name(VariableReference.GetName()), SoftObjectPath(VariableReference.GetObject()) {}
+	explicit FAnimNextSoftVariableReference(const FAnimNextVariableReference& VariableReference) : Name(VariableReference.GetName()), SoftObjectPath(VariableReference.GetObject()) 
+	{
+		ensureAlwaysMsgf(!UE::UAF::IsInternalVariableName(Name), TEXT("Attempting to create soft variable reference using a variable name indicative of a programmatic variable %s"), *Name.ToString());
+	}
 
 	// Get the name of the variable
 	FName GetName() const
diff --git a/Engine/Plugins/Experimental/UAF/UAF/Source/UAF/Public/Variables/AnimNextVariableReference.h b/Engine/Plugins/Experimental/UAF/UAF/Source/UAF/Public/Variables/AnimNextVariableReference.h
--- a/Engine/Plugins/Experimental/UAF/UAF/Source/UAF/Public/Variables/AnimNextVariableReference.h
+++ b/Engine/Plugins/Experimental/UAF/UAF/Source/UAF/Public/Variables/AnimNextVariableReference.h
@@ -14,7 +14,14 @@ namespace UE::UAF::Tests
 class FVariableReferenceTest;
 }
 
-// A reference to an AnimNext variable
+namespace UE::UAF
+{
+static const FString InternalVariablePrefix(TEXT("__InternalVar_"));
+// Whether the variable is a function internal one
+static bool IsInternalVariableName(FName InVariableName) { return InVariableName.ToString().StartsWith(InternalVariablePrefix); }
+}
+
+// A reference to a specific variable (UAnimNextVariableEntry) on UUAFRigVMAsset derived class or FProperty on UScriptStruct (NativeSharedVariables).
 USTRUCT(BlueprintType)
 struct FAnimNextVariableReference
 {
diff --git a/Engine/Plugins/Experimental/UAF/UAF/Source/UAFUncookedOnly/Private/AnimNextUncookedOnlyModule.cpp b/Engine/Plugins/Experimental/UAF/UAF/Source/UAFUncookedOnly/Private/AnimNextUncookedOnlyModule.cpp
--- a/Engine/Plugins/Experimental/UAF/UAF/Source/UAFUncookedOnly/Private/AnimNextUncookedOnlyModule.cpp
+++ b/Engine/Plugins/Experimental/UAF/UAF/Source/UAFUncookedOnly/Private/AnimNextUncookedOnlyModule.cpp
@@ -53,6 +53,7 @@ TSharedPtr<IVariableBindingType> FModule::FindVariableBindingType(const UScriptS
 	return VariableBindingTypes.FindRef(*InStruct->GetPathName());
 }
 
+#if WITH_EDITOR
 FGuid FModule::GetVariableGuidByName(const FName InName, const UObject* InObject) const
 {
 	FGuid Guid;
@@ -139,6 +140,7 @@ FName FModule::GetVariableNameByGuid(const FGuid InGuid, const UObject* InObject
 	
 	return NAME_None;
 }
+#endif // WITH_EDITOR
 
 }
 
diff --git a/Engine/Plugins/Experimental/UAF/UAF/Source/UAFUncookedOnly/Private/UncookedOnlyUtils.cpp b/Engine/Plugins/Experimental/UAF/UAF/Source/UAFUncookedOnly/Private/UncookedOnlyUtils.cpp
--- a/Engine/Plugins/Experimental/UAF/UAF/Source/UAFUncookedOnly/Private/UncookedOnlyUtils.cpp
+++ b/Engine/Plugins/Experimental/UAF/UAF/Source/UAFUncookedOnly/Private/UncookedOnlyUtils.cpp
@@ -1560,7 +1560,7 @@ void FUtils::OpenProgrammaticGraphs(UUAFRigVMAssetEditorData* EditorData, const
 FString FUtils::MakeFunctionWrapperVariableName(FName InFunctionName, FName InVariableName)
 {
 	// We assume the function name is enough for variable name uniqueness in this graph (We don't yet desire global uniqueness).
-	return TEXT("__InternalVar_") + InFunctionName.ToString() + "_" + InVariableName.ToString();
+	return InternalVariablePrefix + InFunctionName.ToString() + "_" + InVariableName.ToString();
 }
 
 FString FUtils::MakeFunctionWrapperEventName(FName InFunctionName)
diff --git a/Engine/Plugins/Experimental/UAF/UAF/Source/UAFUncookedOnly/Public/IAnimNextUncookedOnlyModule.h b/Engine/Plugins/Experimental/UAF/UAF/Source/UAFUncookedOnly/Public/IAnimNextUncookedOnlyModule.h
--- a/Engine/Plugins/Experimental/UAF/UAF/Source/UAFUncookedOnly/Public/IAnimNextUncookedOnlyModule.h
+++ b/Engine/Plugins/Experimental/UAF/UAF/Source/UAFUncookedOnly/Public/IAnimNextUncookedOnlyModule.h
@@ -32,6 +32,7 @@ public:
 	// @param   InStruct         The struct type for this variable binding's data (must be a child struct of FAnimNextVariableBindingData)
 	virtual TSharedPtr<IVariableBindingType> FindVariableBindingType(const UScriptStruct* InInstanceIdStruct) const = 0;
 	
+#if WITH_EDITOR
 	// Try and retrieve a valid FGuid representing the provided variable (by name) within the source object
 	// @param	InName			Name of the variable to retrieve the FGuid for
 	// @param	InObject		Object, supposed to be, containing the variable (UScriptStruct, UUAFRigVMAsset or UUAFRigVMAssetEditorData)
@@ -41,6 +42,7 @@ public:
 	// @param	InGuid			FGuid of the variable to retrieve the FName for
 	// @param	InObject		Object, supposed to be, containing the variable (UScriptStruct, UUAFRigVMAsset or UUAFRigVMAssetEditorData)
 	virtual FName GetVariableNameByGuid(const FGuid InGuid, const UObject* InObject) const = 0;
+#endif // WITH_EDITOR
 	
 	static IAnimNextUncookedOnlyModule& Get() { return FModuleManager::Get().LoadModuleChecked<IAnimNextUncookedOnlyModule>("UAFUncookedOnly"); }
 };
