Missing header file from CL 50535463

#Synced-CL 50532483

diff --git a/Engine/Plugins/Experimental/UAF/UAF/Source/UAFUncookedOnly/Private/AnimNextUncookedOnlyModule.h b/Engine/Plugins/Experimental/UAF/UAF/Source/UAFUncookedOnly/Private/AnimNextUncookedOnlyModule.h
--- a/Engine/Plugins/Experimental/UAF/UAF/Source/UAFUncookedOnly/Private/AnimNextUncookedOnlyModule.h
+++ b/Engine/Plugins/Experimental/UAF/UAF/Source/UAFUncookedOnly/Private/AnimNextUncookedOnlyModule.h
@@ -24,8 +24,10 @@ private:
 	virtual void RegisterVariableBindingType(FName InStructName, TSharedPtr<IVariableBindingType> InType) override;
 	virtual void UnregisterVariableBindingType(FName InStructName) override;
 	virtual TSharedPtr<IVariableBindingType> FindVariableBindingType(const UScriptStruct* InInstanceIdStruct) const override;
+#if WITH_EDITOR
 	virtual FGuid GetVariableGuidByName(const FName InName, const UObject* InObject) const override;
 	virtual FName GetVariableNameByGuid(const FGuid InGuid, const UObject* InObject) const override;
+#endif // WITH_EDITOR
 
 	TMap<FName, TSharedPtr<IVariableBindingType>> VariableBindingTypes;
 
