Use sanitized internal variable name during RigVM compilation of UAF assets

#Synced-CL 50652508

diff --git a/Engine/Plugins/Experimental/UAF/UAF/Source/UAFUncookedOnly/Private/AnimNextRigVMAssetEditorData.cpp b/Engine/Plugins/Experimental/UAF/UAF/Source/UAFUncookedOnly/Private/AnimNextRigVMAssetEditorData.cpp
--- a/Engine/Plugins/Experimental/UAF/UAF/Source/UAFUncookedOnly/Private/AnimNextRigVMAssetEditorData.cpp
+++ b/Engine/Plugins/Experimental/UAF/UAF/Source/UAFUncookedOnly/Private/AnimNextRigVMAssetEditorData.cpp
@@ -858,7 +858,9 @@ void UUAFRigVMAssetEditorData::BuildFunctionWrapperEventVariables(FAnimNextRigVM
 					}
 				
 					FRigVMGraphFunctionArgument InternallyNamedArgument = Argument;
-					InternallyNamedArgument.Name = FName(FUtils::MakeFunctionWrapperVariableName(FunctionHeader.Name, Argument.Name));
+					FName InternalArgName = FName(FUtils::MakeFunctionWrapperVariableName(FunctionHeader.Name, Argument.Name));
+					FName SanitizedInternalArgName = FRigVMPropertyDescription::SanitizeName(InternalArgName);
+					InternallyNamedArgument.Name = SanitizedInternalArgName;
 					InContext.ProgrammaticVariables.Add(FAnimNextProgrammaticVariable::FromRigVMGraphFunctionArgument(InternallyNamedArgument));
 				}
 			}
