Reduce recoverable UAF variable mismatch checks to ensure

#tests editor

#Synced-CL 50659195

diff --git a/Engine/Plugins/Experimental/UAF/UAF/Source/UAF/Private/Script/UAFRigVMComponent.cpp b/Engine/Plugins/Experimental/UAF/UAF/Source/UAF/Private/Script/UAFRigVMComponent.cpp
--- a/Engine/Plugins/Experimental/UAF/UAF/Source/UAF/Private/Script/UAFRigVMComponent.cpp
+++ b/Engine/Plugins/Experimental/UAF/UAF/Source/UAF/Private/Script/UAFRigVMComponent.cpp
@@ -29,7 +29,12 @@ void FUAFRigVMComponent::OnBindToInstance()
 	// Hookup the runtime data ptrs
 	const TArray<FRigVMExternalVariableDef>& ExternalVariableDefs = VM->GetExternalVariableDefs();
 	const int32 NumExternalVariables = ExternalVariableDefs.Num();
-	check(AssetInstance.Variables.NumInternalVariables == NumExternalVariables);
+
+	if (!ensure(AssetInstance.Variables.NumInternalVariables == NumExternalVariables))
+	{
+		return;
+	}
+
 	if(NumExternalVariables > 0)
 	{
 		TArray<FRigVMExternalVariableRuntimeData> ExternalVariableRuntimeData;
@@ -59,7 +64,10 @@ void FUAFRigVMComponent::StaticCallEvent(TStructView<FUAFScriptComponent> InScri
 {
 	FUAFRigVMComponent& This = InScriptComponent.Get<FUAFRigVMComponent>();
 #if DO_CHECK
-	check(This.bIsInitialized);
+	if (!ensure(This.bIsInitialized))
+	{
+		return;
+	}
 #endif
 
 	if(This.VM == nullptr)
diff --git a/Engine/Plugins/Experimental/UAF/UAF/Source/UAF/Private/Variables/UAFInstanceVariableContainer.cpp b/Engine/Plugins/Experimental/UAF/UAF/Source/UAF/Private/Variables/UAFInstanceVariableContainer.cpp
--- a/Engine/Plugins/Experimental/UAF/UAF/Source/UAF/Private/Variables/UAFInstanceVariableContainer.cpp
+++ b/Engine/Plugins/Experimental/UAF/UAF/Source/UAF/Private/Variables/UAFInstanceVariableContainer.cpp
@@ -1146,7 +1146,12 @@ void FUAFInstanceVariableContainer::GetRigVMMemoryForVariables(TConstArrayView<F
 				{
 					const FPropertyBagPropertyDesc& Desc = Descs[VariableIndex];
 					const FRigVMExternalVariableDef& Def = InVariableDefs[VariableIndex];
-					check(Desc.CachedProperty->GetClass() == Def.GetProperty()->GetClass());
+
+					if (!ensure(Desc.CachedProperty->GetClass() == Def.GetProperty()->GetClass()))
+					{
+						OutRuntimeData.Empty();
+						return;
+					}
 					OutRuntimeData.Emplace(Desc.CachedProperty->ContainerPtrToValuePtr<uint8>(BasePtr));
 				}
 			}
@@ -1168,7 +1173,11 @@ void FUAFInstanceVariableContainer::GetRigVMMemoryForVariables(TConstArrayView<F
 				{
 					const FStructDataCache::FPropertyInfo& Property = Properties[VariableIndex];
 					const FRigVMExternalVariableDef& Def = InVariableDefs[VariableIndex];
-					check(Property.Property->GetClass() == Def.GetProperty()->GetClass());
+					if (!ensure(Property.Property->GetClass() == Def.GetProperty()->GetClass()))
+					{
+						OutRuntimeData.Empty();
+						return;
+					}
 					OutRuntimeData.Emplace(Property.Property->ContainerPtrToValuePtr<uint8>(BasePtr));
 				}
 			}
