Make class comparison failures log mismatched classes.
Reduce recoverable checks to ensure to prevent crashing.

#tests editor, PIE

#Synced-CL 50714302

diff --git a/Engine/Plugins/Experimental/UAF/UAF/Source/UAF/Private/Variables/RigVMDispatch_CallHoistedAccessorFunction.cpp b/Engine/Plugins/Experimental/UAF/UAF/Source/UAF/Private/Variables/RigVMDispatch_CallHoistedAccessorFunction.cpp
--- a/Engine/Plugins/Experimental/UAF/UAF/Source/UAF/Private/Variables/RigVMDispatch_CallHoistedAccessorFunction.cpp
+++ b/Engine/Plugins/Experimental/UAF/UAF/Source/UAF/Private/Variables/RigVMDispatch_CallHoistedAccessorFunction.cpp
@@ -105,7 +105,14 @@ void FRigVMDispatch_CallHoistedAccessorFunctionNative::Execute(FRigVMExtendedExe
 
 	const FProperty* TargetProperty = Handles[2].GetResolvedProperty();
 	checkSlow(TargetProperty);
-	check(TargetProperty->GetClass() == ReturnValueProperty->GetClass());
+
+	if (!ensureMsgf(TargetProperty->GetClass() == ReturnValueProperty->GetClass(), TEXT("FRigVMDispatch_CallHoistedAccessorFunctionNative::Execute - Type mismatch: %s, %s")
+		, *TargetProperty->GetClass()->GetFName().ToString()
+		, *ReturnValueProperty->GetClass()->GetFName().ToString()))
+	{
+		return;
+	}
+
 	uint8* TargetAddress = Handles[2].GetOutputData();
 	checkSlow(TargetAddress);
 
@@ -157,7 +164,14 @@ void FRigVMDispatch_CallHoistedAccessorFunctionScript::Execute(FRigVMExtendedExe
 
 	const FProperty* TargetProperty = Handles[2].GetResolvedProperty();
 	checkSlow(TargetProperty);
-	check(TargetProperty->GetClass() == ReturnValueProperty->GetClass());
+
+	if (!ensureMsgf(TargetProperty->GetClass() == ReturnValueProperty->GetClass(), TEXT("FRigVMDispatch_CallHoistedAccessorFunctionScript::Execute - Type mismatch: %s, %s")
+		, *TargetProperty->GetClass()->GetFName().ToString()
+		, *ReturnValueProperty->GetClass()->GetFName().ToString()))
+	{
+		return;
+	}
+
 	uint8* TargetAddress = Handles[2].GetOutputData();
 	checkSlow(TargetAddress);
 
diff --git a/Engine/Plugins/Experimental/UAF/UAF/Source/UAF/Private/Variables/RigVMDispatch_CallObjectAccessorFunction.cpp b/Engine/Plugins/Experimental/UAF/UAF/Source/UAF/Private/Variables/RigVMDispatch_CallObjectAccessorFunction.cpp
--- a/Engine/Plugins/Experimental/UAF/UAF/Source/UAF/Private/Variables/RigVMDispatch_CallObjectAccessorFunction.cpp
+++ b/Engine/Plugins/Experimental/UAF/UAF/Source/UAF/Private/Variables/RigVMDispatch_CallObjectAccessorFunction.cpp
@@ -105,7 +105,14 @@ void FRigVMDispatch_CallObjectAccessorFunctionNative::Execute(FRigVMExtendedExec
 
 	const FProperty* TargetProperty = Handles[2].GetResolvedProperty();
 	checkSlow(TargetProperty);
-	check(TargetProperty->GetClass() == ReturnValueProperty->GetClass());
+
+	if (!ensureMsgf(TargetProperty->GetClass() == ReturnValueProperty->GetClass(), TEXT("FRigVMDispatch_CallObjectAccessorFunctionNative::Execute - Type mismatch: %s, %s")
+		, *TargetProperty->GetClass()->GetFName().ToString()
+		, *ReturnValueProperty->GetClass()->GetFName().ToString()))
+	{
+		return;
+	}
+
 	uint8* TargetAddress = Handles[2].GetOutputData();
 	checkSlow(TargetAddress);
 
@@ -142,7 +149,14 @@ void FRigVMDispatch_CallObjectAccessorFunctionScript::Execute(FRigVMExtendedExec
 
 	const FProperty* TargetProperty = Handles[2].GetResolvedProperty();
 	checkSlow(TargetProperty);
-	check(TargetProperty->GetClass() == ReturnValueProperty->GetClass());
+
+	if (!ensureMsgf(TargetProperty->GetClass() == ReturnValueProperty->GetClass(), TEXT("FRigVMDispatch_CallObjectAccessorFunctionScript::Execute - Type mismatch: %s, %s")
+		, *TargetProperty->GetClass()->GetFName().ToString()
+		, *ReturnValueProperty->GetClass()->GetFName().ToString()))
+	{
+		return;
+	}
+
 	uint8* TargetAddress = Handles[2].GetOutputData();
 	checkSlow(TargetAddress);
 
diff --git a/Engine/Plugins/Experimental/UAF/UAF/Source/UAF/Private/Variables/RigVMDispatch_GetObjectProperty.cpp b/Engine/Plugins/Experimental/UAF/UAF/Source/UAF/Private/Variables/RigVMDispatch_GetObjectProperty.cpp
--- a/Engine/Plugins/Experimental/UAF/UAF/Source/UAF/Private/Variables/RigVMDispatch_GetObjectProperty.cpp
+++ b/Engine/Plugins/Experimental/UAF/UAF/Source/UAF/Private/Variables/RigVMDispatch_GetObjectProperty.cpp
@@ -97,7 +97,14 @@ void FRigVMDispatch_GetObjectProperty::Execute(FRigVMExtendedExecuteContext& InC
 
 	const FProperty* TargetProperty = Handles[2].GetResolvedProperty();
 	checkSlow(TargetProperty);
-	checkSlow(TargetProperty->GetClass() == SourceProperty->GetClass());
+
+	if (!ensureMsgf(TargetProperty->GetClass() == SourceProperty->GetClass(), TEXT("FRigVMDispatch_GetObjectProperty::Execute - Type mismatch: %s, %s")
+		, *TargetProperty->GetClass()->GetFName().ToString()
+		, *SourceProperty->GetClass()->GetFName().ToString()))
+	{
+		return;
+	}
+
 	uint8* TargetAddress = Handles[2].GetOutputData();
 	checkSlow(TargetAddress);
 
