UAF: Variable reference picker improvements

Allow variable reference pickers to function correctly on pins
Add the ability to filter to the 'current' asset (the default when inside a UAF asset context), or ShowAll (show variables from all assets - the default when not in a UAF asset context)

#Synced-CL 50778421

diff --git a/Engine/Plugins/Experimental/UAF/UAF/Source/UAFEditor/Internal/Variables/VariablePickerArgs.h b/Engine/Plugins/Experimental/UAF/UAF/Source/UAFEditor/Internal/Variables/VariablePickerArgs.h
--- a/Engine/Plugins/Experimental/UAF/UAF/Source/UAFEditor/Internal/Variables/VariablePickerArgs.h
+++ b/Engine/Plugins/Experimental/UAF/UAF/Source/UAFEditor/Internal/Variables/VariablePickerArgs.h
@@ -34,6 +34,9 @@ using FOnIsContextSensitive = TDelegate<bool()>;
 // Delegate called when context sensitivity changes
 using FOnContextSensitivityChanged = TDelegate<void(bool)>;
 
+// Delegate called to get the inclusion/exclusion flags
+using FOnGetIncludeExcludeFilter = TDelegate<void(EAnimNextExportedVariableFlags& OutFlagInclusionFilter, EAnimNextExportedVariableFlags& OutFlagExclusionFilter)>;
+
 struct FVariablePickerArgs
 {
 	FVariablePickerArgs() = default;
@@ -56,10 +59,13 @@ struct FVariablePickerArgs
 	// Delegate called when context sensitivity changes
 	FOnContextSensitivityChanged OnContextSensitivityChanged;
 
-	// Default filter for AnimNext export flags to be used for variable list population
+	// Delegate called to get the inclusion/exclusion flags
+	FOnGetIncludeExcludeFilter OnGetIncludeExcludeFilter;
+
+	// Default filter for AnimNext export flags to be used for variable list population (used if OnGetIncludeExcludeFilter is not overriden)
 	EAnimNextExportedVariableFlags FlagInclusionFilter = EAnimNextExportedVariableFlags::Declared | EAnimNextExportedVariableFlags::Public;
 	EAnimNextExportedVariableFlags FlagExclusionFilter = EAnimNextExportedVariableFlags::Referenced;
-	
+
 	// Whether the search box should be focused on widget creation
 	bool bFocusSearchWidget = true;
 };
diff --git a/Engine/Plugins/Experimental/UAF/UAF/Source/UAFEditor/Private/Variables/SGraphPinVariableReference.cpp b/Engine/Plugins/Experimental/UAF/UAF/Source/UAFEditor/Private/Variables/SGraphPinVariableReference.cpp
--- a/Engine/Plugins/Experimental/UAF/UAF/Source/UAFEditor/Private/Variables/SGraphPinVariableReference.cpp
+++ b/Engine/Plugins/Experimental/UAF/UAF/Source/UAFEditor/Private/Variables/SGraphPinVariableReference.cpp
@@ -37,11 +37,18 @@ void SGraphPinVariableReference::Construct(const FArguments& InArgs, UEdGraphPin
 		FilterType = FAnimNextParamType::FromString(AllowedType);
 	}
 
+	bIsContextSensitive = Node->GetPinMetaData(InPin->GetFName(),"ShowAll") != TEXT("true");
+
+	if (UUAFRigVMAsset* Asset = Node->GetTypedOuter<UUAFRigVMAsset>())
+	{
+		WeakFilterAsset = Asset;
+	}
+
 	check(CastChecked<UScriptStruct>(InPin->PinType.PinSubCategoryObject.Get()) == FAnimNextVariableReference::StaticStruct());
 	if (URigVMEdGraphNode* RigVMEdGraphNode = Cast<URigVMEdGraphNode>(Node))
 	{
-		ModelPin = RigVMEdGraphNode->GetModelNode()->FindPin(InPin->GetFName().ToString());
-		check(ModelPin);
+		ModelPin = RigVMEdGraphNode->GetModelPinFromPinPath(InPin->GetName());
+		ensure(ModelPin);
 		RigVMEdGraphNode->GetModel()->OnModified().AddSP(this, &SGraphPinVariableReference::HandleGraphModified);
 	}
 
@@ -71,7 +78,56 @@ TSharedRef<SWidget> SGraphPinVariableReference::GetDefaultValueWidget()
 
 		UpdateCachedData();
 	});
-	
+
+	if (WeakFilterAsset.Get() != nullptr)
+	{
+		OnIsContextSensitiveDelegate = FOnIsContextSensitive::CreateLambda([this]()
+		{
+			return bIsContextSensitive;
+		});
+
+		Args.OnIsContextSensitive = &OnIsContextSensitiveDelegate;
+		Args.OnContextSensitivityChanged = FOnContextSensitivityChanged::CreateLambda([this](bool bInIsContextSensitive)
+		{
+			bIsContextSensitive = bInIsContextSensitive;
+		});
+
+		Args.OnGetIncludeExcludeFilter = FOnGetIncludeExcludeFilter::CreateLambda([this](EAnimNextExportedVariableFlags& OutFlagInclusionFilter, EAnimNextExportedVariableFlags& OutFlagExclusionFilter)
+		{
+			// If we are context-sensitive and we have an asset, we allow private variable references from that asset
+			if (bIsContextSensitive && WeakFilterAsset.Get())
+			{
+				OutFlagInclusionFilter = EAnimNextExportedVariableFlags::Declared;
+				OutFlagExclusionFilter = EAnimNextExportedVariableFlags::Referenced;
+			}
+			else
+			{
+				OutFlagInclusionFilter = EAnimNextExportedVariableFlags::Declared | EAnimNextExportedVariableFlags::Public;
+				OutFlagExclusionFilter = EAnimNextExportedVariableFlags::Referenced;
+			}
+		});
+	}
+
+	Args.OnFilterVariable = FOnFilterVariable::CreateLambda([this](const FAnimNextSoftVariableReference& InVariableReference)
+	{
+		const UUAFRigVMAsset* FilterAsset = WeakFilterAsset.Get();
+		if (FilterAsset == nullptr || !bIsContextSensitive)
+		{
+			return EFilterVariableResult::Include;
+		}
+
+		if (FilterAsset != nullptr)
+		{
+			FSoftObjectPath FilterAssetPath(FilterAsset);
+			if (InVariableReference.GetSoftObjectPath() == FilterAssetPath)
+			{
+				return EFilterVariableResult::Include;
+			}
+		}
+
+		return EFilterVariableResult::Exclude;
+	});
+
 	Args.OnFilterVariableType = FOnFilterVariableType::CreateLambda([this](const FAnimNextParamType& InParamType)-> EFilterVariableResult
 	{
 		if(FilterType.IsValid())
diff --git a/Engine/Plugins/Experimental/UAF/UAF/Source/UAFEditor/Private/Variables/SGraphPinVariableReference.h b/Engine/Plugins/Experimental/UAF/UAF/Source/UAFEditor/Private/Variables/SGraphPinVariableReference.h
--- a/Engine/Plugins/Experimental/UAF/UAF/Source/UAFEditor/Private/Variables/SGraphPinVariableReference.h
+++ b/Engine/Plugins/Experimental/UAF/UAF/Source/UAFEditor/Private/Variables/SGraphPinVariableReference.h
@@ -5,6 +5,7 @@
 #include "SGraphPin.h"
 #include "Param/ParamType.h"
 #include "Variables/AnimNextSoftVariableReference.h"
+#include "Variables/VariablePickerArgs.h"
 
 class URigVMGraph;
 enum class ERigVMGraphNotifType : uint8;
@@ -52,6 +53,12 @@ private:
 	FAnimNextParamType CachedType;
 
 	TSharedPtr<SVariablePickerCombo> PickerCombo;
+
+	TWeakObjectPtr<UUAFRigVMAsset> WeakFilterAsset;
+
+	FOnIsContextSensitive OnIsContextSensitiveDelegate;
+
+	bool bIsContextSensitive = true;
 };
 
 }
diff --git a/Engine/Plugins/Experimental/UAF/UAF/Source/UAFEditor/Private/Variables/SVariablePicker.cpp b/Engine/Plugins/Experimental/UAF/UAF/Source/UAFEditor/Private/Variables/SVariablePicker.cpp
--- a/Engine/Plugins/Experimental/UAF/UAF/Source/UAFEditor/Private/Variables/SVariablePicker.cpp
+++ b/Engine/Plugins/Experimental/UAF/UAF/Source/UAFEditor/Private/Variables/SVariablePicker.cpp
@@ -83,6 +83,10 @@ void SVariablePicker::RefreshEntries()
 	CachedContainers.Reset();
 	ContainerMap.Reset();
 
+	EAnimNextExportedVariableFlags FlagInclusionFilter = Args.FlagInclusionFilter;
+	EAnimNextExportedVariableFlags FlagExclusionFilter = Args.FlagExclusionFilter;
+	Args.OnGetIncludeExcludeFilter.ExecuteIfBound(FlagInclusionFilter, FlagExclusionFilter);
+
 	// Add variables exposed on assets
 	TMap<FAssetData, FAnimNextAssetRegistryExports> Exports;
 	if(UncookedOnly::FUtils::GetExportedVariablesFromAssetRegistry(Exports))
@@ -98,8 +102,8 @@ void SVariablePicker::RefreshEntries()
 				{
 					if (const FAnimNextVariableDeclarationData* VariableEntry = Export.Data.GetPtr<FAnimNextVariableDeclarationData>())
 					{
-						const bool bHasInclusionFlags = EnumHasAllFlags(static_cast<EAnimNextExportedVariableFlags>(VariableEntry->Flags), Args.FlagInclusionFilter);
-						const bool bHasExclusionFlags = EnumHasAnyFlags(static_cast<EAnimNextExportedVariableFlags>(VariableEntry->Flags), Args.FlagExclusionFilter);
+						const bool bHasInclusionFlags = EnumHasAllFlags(static_cast<EAnimNextExportedVariableFlags>(VariableEntry->Flags), FlagInclusionFilter);
+						const bool bHasExclusionFlags = EnumHasAnyFlags(static_cast<EAnimNextExportedVariableFlags>(VariableEntry->Flags), FlagExclusionFilter);
 						
 						if(bHasInclusionFlags && !bHasExclusionFlags && Export.Identifier != NAME_None && VariableEntry->Type.IsValid())
 						{
diff --git a/Engine/Plugins/Experimental/UAF/UAF/Source/UAFEditor/Private/Variables/VariableReferencePropertyCustomization.cpp b/Engine/Plugins/Experimental/UAF/UAF/Source/UAFEditor/Private/Variables/VariableReferencePropertyCustomization.cpp
--- a/Engine/Plugins/Experimental/UAF/UAF/Source/UAFEditor/Private/Variables/VariableReferencePropertyCustomization.cpp
+++ b/Engine/Plugins/Experimental/UAF/UAF/Source/UAFEditor/Private/Variables/VariableReferencePropertyCustomization.cpp
@@ -27,15 +27,27 @@ void FVariableReferencePropertyCustomization::CustomizeHeader(TSharedRef<IProper
 	const FString AllowedType = PropertyHandle->GetMetaData("AllowedType");
 	FilterType = FAnimNextParamType::FromString(AllowedType);
 
-	bIsContextSensitive = !PropertyHandle->HasMetaData("ShowAll");
+	bIsContextSensitive = PropertyHandle->GetMetaData("ShowAll") != TEXT("true");
 
-	TArray<UObject*> Objects;
-	InPropertyHandle->GetOuterObjects(Objects);
-	for (UObject* Object : Objects)
+	// Grab the objects from the details view itself as we will often be customizing traits, which have no object associated with their property handle
+	TArray<TWeakObjectPtr<UObject>> Objects = InCustomizationUtils.GetPropertyUtilities()->GetSelectedObjects();
+	for (TWeakObjectPtr<UObject> WeakObject : Objects)
 	{
-		if (UUAFRigVMAsset* Asset = Object->GetTypedOuter<UUAFRigVMAsset>())
+		if (UObject* Object = WeakObject.Get())
 		{
-			FilterAssets.Add(Asset);
+			UUAFRigVMAsset* Asset = Cast<UUAFRigVMAsset>(Object);
+			if (Asset != nullptr)
+			{
+				FilterAssets.Add(Asset);
+			}
+			else
+			{
+				UUAFRigVMAsset* OuterAsset = Object->GetTypedOuter<UUAFRigVMAsset>();
+				if (OuterAsset != nullptr)
+				{
+					FilterAssets.Add(OuterAsset);
+				}
+			}
 		}
 	}
 
@@ -90,6 +102,21 @@ void FVariableReferencePropertyCustomization::CustomizeHeader(TSharedRef<IProper
 		{
 			bIsContextSensitive = bInIsContextSensitive;
 		});
+		
+		Args.OnGetIncludeExcludeFilter = FOnGetIncludeExcludeFilter::CreateLambda([this](EAnimNextExportedVariableFlags& OutFlagInclusionFilter, EAnimNextExportedVariableFlags& OutFlagExclusionFilter)
+		{
+			// If we are context-sensitive and we have filter assets, we allow private variable references from those assets
+			if (bIsContextSensitive && FilterAssets.Num() > 0)
+			{
+				OutFlagInclusionFilter = EAnimNextExportedVariableFlags::Declared;
+				OutFlagExclusionFilter = EAnimNextExportedVariableFlags::Referenced;
+			}
+			else
+			{
+				OutFlagInclusionFilter = EAnimNextExportedVariableFlags::Declared | EAnimNextExportedVariableFlags::Public;
+				OutFlagExclusionFilter = EAnimNextExportedVariableFlags::Referenced;
+			}
+		});
 	}
 
 	Args.OnFilterVariable = FOnFilterVariable::CreateLambda([this](const FAnimNextSoftVariableReference& InVariableReference)
diff --git a/Engine/Plugins/Experimental/UAF/UAF/Source/UAFUncookedOnly/Internal/AnimNextEdGraphSchema.h b/Engine/Plugins/Experimental/UAF/UAF/Source/UAFUncookedOnly/Internal/AnimNextEdGraphSchema.h
--- a/Engine/Plugins/Experimental/UAF/UAF/Source/UAFUncookedOnly/Internal/AnimNextEdGraphSchema.h
+++ b/Engine/Plugins/Experimental/UAF/UAF/Source/UAFUncookedOnly/Internal/AnimNextEdGraphSchema.h
@@ -19,6 +19,7 @@ class UAnimNextEdGraphSchema : public URigVMEdGraphSchema
 	// URigVMEdGraphSchema interface
 	virtual TSubclassOf<URigVMEdGraphNode> GetGraphNodeClass(const URigVMEdGraph* InGraph) const override { return UAnimNextEdGraphNode::StaticClass(); }
 	UE_API virtual bool ShouldHidePinOnNode(const URigVMPin* InModelPin) const override;
+	UE_API virtual bool IsStructEditable(UStruct* InStruct) const override;
 
 	// UEdGraphSchema interface
 	UE_API virtual void GetGraphDisplayInformation(const UEdGraph& Graph, /*out*/ FGraphDisplayInfo& DisplayInfo) const override;
diff --git a/Engine/Plugins/Experimental/UAF/UAF/Source/UAFUncookedOnly/Private/AnimNextEdGraphSchema.cpp b/Engine/Plugins/Experimental/UAF/UAF/Source/UAFUncookedOnly/Private/AnimNextEdGraphSchema.cpp
--- a/Engine/Plugins/Experimental/UAF/UAF/Source/UAFUncookedOnly/Private/AnimNextEdGraphSchema.cpp
+++ b/Engine/Plugins/Experimental/UAF/UAF/Source/UAFUncookedOnly/Private/AnimNextEdGraphSchema.cpp
@@ -55,4 +55,14 @@ TSharedPtr<IAssetReferenceFilter> UAnimNextEdGraphSchema::MakeAssetReferenceFilt
 	return {};
 }
 
+
+bool UAnimNextEdGraphSchema::IsStructEditable(UStruct* InStruct) const
+{
+	if (InStruct == FAnimNextVariableReference::StaticStruct())
+	{
+		return true;
+	}
+	return Super::IsStructEditable(InStruct);
+}
+
 #undef LOCTEXT_NAMESPACE
