UAF: Avoid recompilation during the active update of a slider UI interaction in the trait stack

#Synced-CL 50858266

diff --git a/Engine/Plugins/Experimental/UAF/UAF/Source/UAFEditor/Private/Common/AnimNextEdGraphNodeCustomization.cpp b/Engine/Plugins/Experimental/UAF/UAF/Source/UAFEditor/Private/Common/AnimNextEdGraphNodeCustomization.cpp
--- a/Engine/Plugins/Experimental/UAF/UAF/Source/UAFEditor/Private/Common/AnimNextEdGraphNodeCustomization.cpp
+++ b/Engine/Plugins/Experimental/UAF/UAF/Source/UAFEditor/Private/Common/AnimNextEdGraphNodeCustomization.cpp
@@ -527,7 +527,14 @@ void FAnimNextEdGraphNodeCustomization::PopulateCategory(IDetailLayoutBuilder& D
 			{
 				return;
 			}
-
+				
+			// Early out if the change type is interactive e.g. from dragging of a slier, to avoid recompilation spam
+			// This will be followed by a ValueSet when finished
+			if ((InPropertyChangedEvent.ChangeType & EPropertyChangeType::Interactive) != 0)
+			{
+				return;
+			}	
+				
 			// Avoid VM recompilation for each Set Default Value
 			URigVMController* Controller = PinnedTraitData->EdGraphNodes[0]->GetController();
 			FRigVMControllerCompileBracketScope CompileScope(Controller);
