UAF: Variable reference picker improvements

Add a transaction so that undo/redo works with pickers and so that var ref properties work correctly on components.
Add referenced assets and component data to filters
Add Identical() to struct to allow struct comparisons to exclude the GUID

#Synced-CL 50814747

diff --git a/Engine/Plugins/Experimental/UAF/UAF/Source/UAF/Private/Variables/AnimNextVariableReference.cpp b/Engine/Plugins/Experimental/UAF/UAF/Source/UAF/Private/Variables/AnimNextVariableReference.cpp
--- a/Engine/Plugins/Experimental/UAF/UAF/Source/UAF/Private/Variables/AnimNextVariableReference.cpp
+++ b/Engine/Plugins/Experimental/UAF/UAF/Source/UAF/Private/Variables/AnimNextVariableReference.cpp
@@ -185,6 +185,11 @@ void FAnimNextVariableReference::ValidateVariableNameAndGuid(FName& InOutName, F
 }
 #endif // WITH_EDITOR
 
+bool FAnimNextVariableReference::Identical(const FAnimNextVariableReference* Other, uint32 PortFlags) const
+{
+	return Object != nullptr && (Name == Other->Name && Object == Other->Object);
+}
+
 bool FAnimNextVariableReference::IsValid() const
 {
 	// TODO: Remove legacy name-based lookup path when deprecations are all fixed up
diff --git a/Engine/Plugins/Experimental/UAF/UAF/Source/UAF/Public/Variables/AnimNextVariableReference.h b/Engine/Plugins/Experimental/UAF/UAF/Source/UAF/Public/Variables/AnimNextVariableReference.h
--- a/Engine/Plugins/Experimental/UAF/UAF/Source/UAF/Public/Variables/AnimNextVariableReference.h
+++ b/Engine/Plugins/Experimental/UAF/UAF/Source/UAF/Public/Variables/AnimNextVariableReference.h
@@ -46,6 +46,8 @@ struct FAnimNextVariableReference
 	UAF_API void PostSerialize(const FArchive& Ar);
 #endif // WITH_EDITOR
 
+	UAF_API bool Identical(const FAnimNextVariableReference* Other, uint32 PortFlags) const;
+	
 	// Legacy constructor from an FName
 	UE_DEPRECATED(5.6, "Constructor should no longer be used, please use constructor that takes an asset or a struct")
 	explicit FAnimNextVariableReference(FName InName)
@@ -151,6 +153,7 @@ struct TStructOpsTypeTraits<FAnimNextVariableReference> : public TStructOpsTypeT
 	{
 		WithSerializer = true,
 		WithPostSerialize = true,
+		WithIdentical = true,
 	};
 #endif
 };
diff --git a/Engine/Plugins/Experimental/UAF/UAF/Source/UAFEditor/Private/Variables/VariableReferencePropertyCustomization.cpp b/Engine/Plugins/Experimental/UAF/UAF/Source/UAFEditor/Private/Variables/VariableReferencePropertyCustomization.cpp
--- a/Engine/Plugins/Experimental/UAF/UAF/Source/UAFEditor/Private/Variables/VariableReferencePropertyCustomization.cpp
+++ b/Engine/Plugins/Experimental/UAF/UAF/Source/UAFEditor/Private/Variables/VariableReferencePropertyCustomization.cpp
@@ -13,6 +13,11 @@
 #include "Variables/AnimNextVariableReference.h"
 #include "Variables/VariablePickerArgs.h"
 #include "AnimNextRigVMAsset.h"
+#include "ScopedTransaction.h"
+#include "Component/AnimNextComponent.h"
+#include "Factory/SystemFactory.h"
+#include "Module/UAFSystemAssetData.h"
+#include "UAF/UAFAssetFactory.h"
 #include "Variables/AnimNextSoftVariableReference.h"
 
 #define LOCTEXT_NAMESPACE "VariableReferencePropertyCustomization"
@@ -40,20 +45,56 @@ void FVariableReferencePropertyCustomization::CustomizeHeader(TSharedRef<IProper
 			{
 				FilterAssets.Add(Asset);
 			}
-			else
+
+			UUAFRigVMAsset* OuterAsset = Object->GetTypedOuter<UUAFRigVMAsset>();
+			if (OuterAsset != nullptr)
+			{
+				FilterAssets.Add(OuterAsset);
+			}
+			
+			UUAFComponent* Component = Cast<UUAFComponent>(Object);
+			if (Component != nullptr)
 			{
-				UUAFRigVMAsset* OuterAsset = Object->GetTypedOuter<UUAFRigVMAsset>();
-				if (OuterAsset != nullptr)
+				const UUAFSystem* System = nullptr;
+				
+				if (Component->GetAssetData().IsValid())
+				{
+					if (Component->GetAssetData().GetScriptStruct()->IsChildOf<FUAFSystemFactoryAsset_System>())
+					{
+						System = Component->GetAssetData().Get<FUAFSystemFactoryAsset_System>().System.Get();
+					}
+					else
+					{
+						FUAFSystemFactoryParams FactoryParams = FSystemFactory::GetDefaultParamsForAsset(Component->GetAssetData());
+						System = FSystemFactory::BuildSystem(FactoryParams.GetBuilder());
+					}
+				}
+				
+				if (System)
 				{
-					FilterAssets.Add(OuterAsset);
+					FilterAssets.Add(System);
 				}
 			}
 		}
 	}
+	
+	// Add any referenced assets
+	TArray<TWeakObjectPtr<const UUAFRigVMAsset>> ReferencedAssets;
+	for (TWeakObjectPtr<const UUAFRigVMAsset> FilterAsset : FilterAssets)
+	{
+		for (const UUAFRigVMAsset* ReferencedAsset : FilterAsset->GetReferencedVariableAssets())
+		{
+			ReferencedAssets.Add(ReferencedAsset);
+		}
+	}
+	
+	FilterAssets.Append(ReferencedAssets);
 
 	FVariablePickerArgs Args;
 	Args.OnVariablePicked = FOnVariablePicked::CreateLambda([this](const FAnimNextSoftVariableReference& InVariableReference, const FAnimNextParamType& InType)
 	{
+		FScopedTransaction Transaction(LOCTEXT("PickVariableReference", "Pick Variable Reference"));
+
 		PropertyHandle->NotifyPreChange();
 
 		PropertyHandle->EnumerateRawData([this, &InVariableReference](void* RawData, const int32 DataIndex, const int32 NumDatas)
@@ -126,7 +167,7 @@ void FVariableReferencePropertyCustomization::CustomizeHeader(TSharedRef<IProper
 			return EFilterVariableResult::Include;
 		}
 
-		for (TWeakObjectPtr<UUAFRigVMAsset> WeakFilterAsset : FilterAssets)
+		for (TWeakObjectPtr<const UUAFRigVMAsset> WeakFilterAsset : FilterAssets)
 		{
 			if (const UUAFRigVMAsset* FilterAsset = WeakFilterAsset.Get())
 			{
diff --git a/Engine/Plugins/Experimental/UAF/UAF/Source/UAFEditor/Private/Variables/VariableReferencePropertyCustomization.h b/Engine/Plugins/Experimental/UAF/UAF/Source/UAFEditor/Private/Variables/VariableReferencePropertyCustomization.h
--- a/Engine/Plugins/Experimental/UAF/UAF/Source/UAFEditor/Private/Variables/VariableReferencePropertyCustomization.h
+++ b/Engine/Plugins/Experimental/UAF/UAF/Source/UAFEditor/Private/Variables/VariableReferencePropertyCustomization.h
@@ -35,7 +35,7 @@ private:
 
 	FAnimNextParamType FilterType;
 
-	TArray<TWeakObjectPtr<UUAFRigVMAsset>> FilterAssets;
+	TArray<TWeakObjectPtr<const UUAFRigVMAsset>> FilterAssets;
 
 	FOnIsContextSensitive OnIsContextSensitiveDelegate;
 
