Recursively referenced SharedVariable Assets/Structs:
- Variable compilation now traverses SharedVariables and includes any recursively referenced assets/structs into the property bag
- GetExternalVariablesImpl for UAF assets now has behaviour to include editor-only definition of variables (this covers/handles uncompiled/non-up-to-date compiled assets)
- Added util function for retrieving recursive shared variable assets/struct
- Added util function for retrieving all referenced/local variables
- Added module functions to access editor-only variable definition from UAF module and check asset variable compilation state

Example:

[System references a shared variables its function]
System --> SharedVariables_One
System --> SharedVariables_One::TestFunction

[Referenced shared variables itself references a variable from another shared variables]
SharedVariables_One --> SharedVariables_Two
SharedVariables_One::TestFunction --> SharedVariables_Two::IntVariable

[System does not reference SharedVariables_Two directly]

#Synced-CL 50930993

diff --git a/Engine/Plugins/Experimental/UAF/UAF/Source/UAF/Private/AnimNextRigVMAsset.cpp b/Engine/Plugins/Experimental/UAF/UAF/Source/UAF/Private/AnimNextRigVMAsset.cpp
--- a/Engine/Plugins/Experimental/UAF/UAF/Source/UAF/Private/AnimNextRigVMAsset.cpp
+++ b/Engine/Plugins/Experimental/UAF/UAF/Source/UAF/Private/AnimNextRigVMAsset.cpp
@@ -9,9 +9,12 @@
 #include "AnimNextRigVMFunctionData.h"
 #include "RigVMRuntimeDataRegistry.h"
 #include "Graph/RigUnit_AnimNextBeginExecution.h"
-#include "Modules/ModuleManager.h"
 #include "UObject/AssetRegistryTagsContext.h"
 
+#if WITH_EDITOR
+#include "IAnimNextUncookedOnlyModule.h"
+#endif // WITH_EDITOR
+
 #include UE_INLINE_GENERATED_CPP_BY_NAME(AnimNextRigVMAsset)
 
 namespace UE::UAF::Private
@@ -115,6 +118,19 @@ void UUAFRigVMAsset::PreDuplicate(FObjectDuplicationParameters& DupParams)
 TArray<FRigVMExternalVariable> UUAFRigVMAsset::GetExternalVariablesImpl(bool bFallbackToBlueprint) const
 {
 	TArray<FRigVMExternalVariable> ExternalVariables;
+#if WITH_EDITOR
+	UE::UAF::UncookedOnly::IAnimNextUncookedOnlyModule& UncookedModule = UE::UAF::UncookedOnly::IAnimNextUncookedOnlyModule::Get();
+	const bool bCompiledVariableDataOutOfDate = UncookedModule.DoesAssetVariablesRequireCompilation(this);
+	if (bCompiledVariableDataOutOfDate && bFallbackToBlueprint)
+	{
+		UncookedModule.GetExternalVariablesForAsset(this, ExternalVariables);
+		// Cannot rely on compiled property bag below as it might contain out-of-date information (e.g. removed or renamed variables, yet to be re-compiled)
+		return ExternalVariables;
+	}
+
+	// Using of CombinedPropertyBag data is only valid if its up-to-date (or assumed up to date bFallbackToBlueprint = false)
+	check(!bCompiledVariableDataOutOfDate || !bFallbackToBlueprint);
+#endif // WITH_EDITOR
 
 	if(const UPropertyBag* PropertyBag = CombinedPropertyBag.GetPropertyBagStruct())
 	{
diff --git a/Engine/Plugins/Experimental/UAF/UAF/Source/UAF/Private/Param/ParamType.cpp b/Engine/Plugins/Experimental/UAF/UAF/Source/UAF/Private/Param/ParamType.cpp
--- a/Engine/Plugins/Experimental/UAF/UAF/Source/UAF/Private/Param/ParamType.cpp
+++ b/Engine/Plugins/Experimental/UAF/UAF/Source/UAF/Private/Param/ParamType.cpp
@@ -498,7 +498,7 @@ void FAnimNextParamType::ToString(FStringBuilderBase& InStringBuilder) const
 		case EValueType::Enum:
 			if(const UEnum* Enum = Cast<UEnum>(ValueTypeObject.Get()))
 			{
-				InStringBuilder.Append(TEXT("U"));
+				InStringBuilder.Append(TEXT("E"));
 				InStringBuilder.Append(Enum->GetName());
 			}
 			else
diff --git a/Engine/Plugins/Experimental/UAF/UAF/Source/UAFUncookedOnly/Internal/AnimNextRigVMAssetEditorData.h b/Engine/Plugins/Experimental/UAF/UAF/Source/UAFUncookedOnly/Internal/AnimNextRigVMAssetEditorData.h
--- a/Engine/Plugins/Experimental/UAF/UAF/Source/UAFUncookedOnly/Internal/AnimNextRigVMAssetEditorData.h
+++ b/Engine/Plugins/Experimental/UAF/UAF/Source/UAFUncookedOnly/Internal/AnimNextRigVMAssetEditorData.h
@@ -774,6 +774,7 @@ protected:
 	bool bSuspendCompilerReports = false;
 	bool bUpgradeDataInterfacesOnLoad = false;
 	bool bFunctionsPatched = false;
+	bool bHasCompiledVariables = false;
 };
 
 #undef UE_API
diff --git a/Engine/Plugins/Experimental/UAF/UAF/Source/UAFUncookedOnly/Internal/UncookedOnlyUtils.h b/Engine/Plugins/Experimental/UAF/UAF/Source/UAFUncookedOnly/Internal/UncookedOnlyUtils.h
--- a/Engine/Plugins/Experimental/UAF/UAF/Source/UAFUncookedOnly/Internal/UncookedOnlyUtils.h
+++ b/Engine/Plugins/Experimental/UAF/UAF/Source/UAFUncookedOnly/Internal/UncookedOnlyUtils.h
@@ -240,7 +240,10 @@ struct FUtils
 	// Generate a (asset reference) workspace export for the provided asset path
 	static UE_API FWorkspaceOutlinerItemExport MakeAssetReferenceExport(const FSoftObjectPath& ReferencedAssetPath, const FWorkspaceOutlinerItemExport& ParentExport, bool bShouldExpand = true);
 
+	static UE_API void GetExternalVariables(const UUAFRigVMAssetEditorData* EditorData, TArray<FRigVMExternalVariable>& OutVariables);
 private:
+	static void GetInternalVariables(const UUAFRigVMAssetEditorData* EditorData, TArray<FRigVMExternalVariable>& OutVariables);
+
 	static UE_API void CompileVariableBindingsInternal(const FRigVMCompileSettings& InSettings, UUAFRigVMAsset* InAsset, TArray<URigVMGraph*>& OutGraphs, bool bInThreadSafe);
 	
 	static void GetSubGraphWorkspaceExportsRecursive(const UUAFRigVMAssetEditorData* EditorData, FWorkspaceOutlinerItemExports& OutExports, const FWorkspaceOutlinerItemExport& RootExport, int32 ParentExportIndex, URigVMEdGraph* RigVMEdGraph, FAssetRegistryTagsContext Context);
@@ -251,6 +254,7 @@ private:
 	static void GetStructVariableExports(const UScriptStruct* Struct, TSet<FAnimNextExport>& OutExports);
 	static void GetSubGraphVariableExportsRecursive(const URigVMEdGraph* RigVMEdGraph, TSet<FAnimNextExport>& OutExports);
 	static void ProcessPinAssetReferences(const URigVMPin* InPin, FWorkspaceOutlinerItemExports& OutExports, const FWorkspaceOutlinerItemExport& RootExport, int32 ParentExportIndex);
+	static void RetrieveReferencedSharedVariableSources(const UUAFRigVMAsset* InAsset, TArray<const UUAFRigVMAsset*>& ReferencedVariableAssets, TArray<const UScriptStruct*>& ReferencedVariableStructs);
 };
 
 }
diff --git a/Engine/Plugins/Experimental/UAF/UAF/Source/UAFUncookedOnly/Private/AnimNextRigVMAssetEditorData.cpp b/Engine/Plugins/Experimental/UAF/UAF/Source/UAFUncookedOnly/Private/AnimNextRigVMAssetEditorData.cpp
--- a/Engine/Plugins/Experimental/UAF/UAF/Source/UAFUncookedOnly/Private/AnimNextRigVMAssetEditorData.cpp
+++ b/Engine/Plugins/Experimental/UAF/UAF/Source/UAFUncookedOnly/Private/AnimNextRigVMAssetEditorData.cpp
@@ -1096,6 +1096,7 @@ void UUAFRigVMAssetEditorData::RecompileVM()
 
 	bWarningsDuringCompilation = false;
 	bErrorsDuringCompilation = false;
+	bHasCompiledVariables = false;
 
 	RigGraphDisplaySettings.MinMicroSeconds = RigGraphDisplaySettings.LastMinMicroSeconds = DBL_MAX;
 	RigGraphDisplaySettings.MaxMicroSeconds = RigGraphDisplaySettings.LastMaxMicroSeconds = (double)INDEX_NONE;
@@ -1204,6 +1205,8 @@ void UUAFRigVMAssetEditorData::RecompileVM()
 			AssetRegistry->AssetUpdateTags(Asset, EAssetRegistryTagsCaller::Fast);
 		}
 	}
+
+	checkf(bHasCompiledVariables, TEXT("Asset variables are expected to be compiled as part of compilation process"));
 }
 
 void UUAFRigVMAssetEditorData::GatherDefaultComponents(UUAFRigVMAsset* InAsset)
@@ -2346,6 +2349,12 @@ UAnimNextVariableEntry* UUAFRigVMAssetEditorData::AddVariable(FName InName, FAni
 		ReportError(*FString::Printf(TEXT("UUAFRigVMAssetEditorData::AddVariable: Cannot add a variable to this asset - type (%s) is not valid."), *InType.ToString()));
 		return nullptr;
 	}
+	
+	if (InType.ToRigVMTemplateArgument().IsUnknownType() || InType.ToRigVMTemplateArgument().IsWildCard())
+	{
+		ReportError(*FString::Printf(TEXT("UUAFRigVMAssetEditorData::AddVariable: Cannot add a variable to this asset - type (%s) is not supported by RigVM."), *InType.ToString()));
+		return nullptr;
+	}
 
 	// Check for duplicate name
 	FName NewParameterName = InName;
@@ -3118,7 +3127,7 @@ FInstancedPropertyBag UUAFRigVMAssetEditorData::GenerateCombinedPropertyBag(cons
 		{
 			InCompileContext.GetAssetCompileContext().Error(
 				Asset,
-				LOCTEXT("FailedToGenerateCombinedPropertyBag", "Failed to generate combined property bag for %s"), FText::FromName(Asset->GetFName())
+				LOCTEXT("FailedToGenerateCombinedPropertyBag", "Failed to generate combined property bag for {0}"), FText::FromName(Asset->GetFName())
 			);
 		}
 	 }
diff --git a/Engine/Plugins/Experimental/UAF/UAF/Source/UAFUncookedOnly/Private/AnimNextUncookedOnlyModule.cpp b/Engine/Plugins/Experimental/UAF/UAF/Source/UAFUncookedOnly/Private/AnimNextUncookedOnlyModule.cpp
--- a/Engine/Plugins/Experimental/UAF/UAF/Source/UAFUncookedOnly/Private/AnimNextUncookedOnlyModule.cpp
+++ b/Engine/Plugins/Experimental/UAF/UAF/Source/UAFUncookedOnly/Private/AnimNextUncookedOnlyModule.cpp
@@ -140,8 +140,22 @@ FName FModule::GetVariableNameByGuid(const FGuid InGuid, const UObject* InObject
 	
 	return NAME_None;
 }
-#endif // WITH_EDITOR
 
+void FModule::GetExternalVariablesForAsset(const UUAFRigVMAsset* InAsset, TArray<FRigVMExternalVariable>& OutVariables) const
+{
+	FUtils::GetExternalVariables(FUtils::GetEditorData<UUAFRigVMAssetEditorData>(InAsset), OutVariables);
+}
+
+bool FModule::DoesAssetVariablesRequireCompilation(const UUAFRigVMAsset* InAsset) const
+{
+	UUAFRigVMAssetEditorData* EditorData = FUtils::GetEditorData<UUAFRigVMAssetEditorData>(InAsset);
+	
+	const bool bNeedsCompilation = !EditorData->bIsCompiling && (EditorData->bVMRecompilationRequired || EditorData->bVMRecompilationRequested);
+	const bool bVariablesNotYetCompiled = EditorData->bIsCompiling && !EditorData->bHasCompiledVariables;
+	
+	return bNeedsCompilation || bVariablesNotYetCompiled;
+}
+#endif // WITH_EDITOR
 }
 
 #undef LOCTEXT_NAMESPACE 
diff --git a/Engine/Plugins/Experimental/UAF/UAF/Source/UAFUncookedOnly/Private/AnimNextUncookedOnlyModule.h b/Engine/Plugins/Experimental/UAF/UAF/Source/UAFUncookedOnly/Private/AnimNextUncookedOnlyModule.h
--- a/Engine/Plugins/Experimental/UAF/UAF/Source/UAFUncookedOnly/Private/AnimNextUncookedOnlyModule.h
+++ b/Engine/Plugins/Experimental/UAF/UAF/Source/UAFUncookedOnly/Private/AnimNextUncookedOnlyModule.h
@@ -27,6 +27,9 @@ private:
 #if WITH_EDITOR
 	virtual FGuid GetVariableGuidByName(const FName InName, const UObject* InObject) const override;
 	virtual FName GetVariableNameByGuid(const FGuid InGuid, const UObject* InObject) const override;
+
+	virtual void GetExternalVariablesForAsset(const UUAFRigVMAsset* InAsset, TArray<FRigVMExternalVariable>& OutVariables) const override;
+	virtual bool DoesAssetVariablesRequireCompilation(const UUAFRigVMAsset* InAsset) const override;
 #endif // WITH_EDITOR
 
 	TMap<FName, TSharedPtr<IVariableBindingType>> VariableBindingTypes;
diff --git a/Engine/Plugins/Experimental/UAF/UAF/Source/UAFUncookedOnly/Private/UncookedOnlyUtils.cpp b/Engine/Plugins/Experimental/UAF/UAF/Source/UAFUncookedOnly/Private/UncookedOnlyUtils.cpp
--- a/Engine/Plugins/Experimental/UAF/UAF/Source/UAFUncookedOnly/Private/UncookedOnlyUtils.cpp
+++ b/Engine/Plugins/Experimental/UAF/UAF/Source/UAFUncookedOnly/Private/UncookedOnlyUtils.cpp
@@ -43,6 +43,8 @@
 #include "Misc/UObjectToken.h"
 #include "RigVMModel/Nodes/RigVMAggregateNode.h"
 #include "Widgets/Notifications/SNotificationList.h"
+#include "Compilation/AnimNextGetFunctionHeaderCompileContext.h"
+#include "UObject/PropertyIterator.h"
 
 #define LOCTEXT_NAMESPACE "AnimNextUncookedOnlyUtils"
 
@@ -53,6 +55,11 @@ TAutoConsoleVariable<bool> CVarDumpProgrammaticGraphs(
 	TEXT("UAF.DumpProgrammaticGraphs"),
 	false,
 	TEXT("When true the transient programmatic graphs will be automatically opened for any that are generated."));
+	
+TAutoConsoleVariable<bool> CVarValidatePropertyBagAgainstExternalVariables(
+	TEXT("UAF.ValidatePropertyBagAgainstExternalVariables"),
+	true,
+	TEXT("When true generated (compiled) property bag its properties will be validated against result of calling GetExternalVariables to enforce variable order."));
 
 void FUtils::RecreateVM(UUAFRigVMAsset* InAsset)
 {
@@ -202,13 +209,18 @@ void FUtils::CompileVariables(const FRigVMCompileSettings& InSettings, UUAFRigVM
 	UUAFRigVMAssetEditorData* EditorData = GetEditorData<UUAFRigVMAssetEditorData>(InAsset);
 
 	// Gather programmatic variables regenerated each compile
-	EditorData->OnPreCompileGetProgrammaticVariables(InSettings, OutCompileContext);
+	// TODO when we need support for this, the variables will have to be available for FUtils::GetExternalVariables as well to keep them in sync with property bag
+	//EditorData->OnPreCompileGetProgrammaticVariables(InSettings, OutCompileContext);
 
 	// Generate the internal property bag
 	InAsset->VariableDefaults = MakePropertyBagForEditorData(EditorData, OutCompileContext);
 
 	TSet<const UUAFRigVMAsset*> SharedVariableAssets;
 	TSet<const UScriptStruct*> SharedVariableStructs;
+	
+	// Set of SharedVariable Assets and structs, retrieved by recursively following references
+	TArray<const UUAFRigVMAsset*> RecursiveSharedVariableAssets;
+	TArray<const UScriptStruct*> RecursiveSharedVariableStructs;
 
 	for(UUAFRigVMAssetEntry* Entry : EditorData->Entries)
 	{
@@ -254,6 +266,13 @@ void FUtils::CompileVariables(const FRigVMCompileSettings& InSettings, UUAFRigVM
 						SharedAssetEntryError(FText::Format(LOCTEXT("DuplicateSharedVariableAssetFound", "Duplicate shared variables asset found ({0})"), FText::FromString(SharedVariablesEntry->GetObjectPath().ToString())));
 						continue;
 					}
+					
+					RetrieveReferencedSharedVariableSources(SharedVariablesAsset, RecursiveSharedVariableAssets, RecursiveSharedVariableStructs);
+					
+					if (RecursiveSharedVariableAssets.Remove(InAsset) != 0)
+					{
+						OutCompileContext.GetAssetCompileContext().Warning(InAsset, LOCTEXT("RecursiveCircularReference", "@@ is recursively referenced by its own Shared Variable Asset reference {0} - chain should be broken."), FText::FromString(SharedVariablesEntry->GetObjectPath().ToString()));
+					}
 
 					SharedVariableAssets.Add(SharedVariablesAsset);
 					break;
@@ -279,7 +298,10 @@ void FUtils::CompileVariables(const FRigVMCompileSettings& InSettings, UUAFRigVM
 			}
 		}
 	}
-	
+
+	SharedVariableAssets.Append(RecursiveSharedVariableAssets);
+	SharedVariableStructs.Append(RecursiveSharedVariableStructs);
+
 	// Set shared variable assets/structs before we setup variables as GetExternalVariablesImpl relies on these
 	if (SharedVariableAssets.Num() > 0)
 	{
@@ -312,6 +334,30 @@ void FUtils::CompileVariables(const FRigVMCompileSettings& InSettings, UUAFRigVM
 		InAsset->VM->ClearExternalVariables(InAsset->ExtendedExecuteContext);
 	}
 
+	EditorData->bHasCompiledVariables = true;
+
+	if (CVarValidatePropertyBagAgainstExternalVariables.GetValueOnAnyThread())
+    {
+	    // Compilation might have failed
+	    if (const UPropertyBag* PropertyBag = InAsset->CombinedPropertyBag.GetPropertyBagStruct())
+	    {
+		    TConstArrayView<FPropertyBagPropertyDesc> PropertyDescs = PropertyBag->GetPropertyDescs();
+		    TArray<FRigVMExternalVariable> ExternalVariables;
+		   	GetExternalVariables(EditorData, ExternalVariables);
+		    
+		    const int32 NumDescriptions = PropertyDescs.Num();
+		    if (ensure(NumDescriptions == ExternalVariables.Num()))
+		    {
+		    	for (int32 Index = 0; Index < NumDescriptions; ++Index)
+		    	{
+		    		const FPropertyBagPropertyDesc& Desc = PropertyDescs[Index];
+		    		const FRigVMExternalVariable& Variable = ExternalVariables[Index];
+		    		ensureMsgf(Desc.Name == Variable.GetName(), TEXT("Mismatch between compiled property bag variable name %s and external variable entry at same index %s"), *Desc.Name.ToString(), *Variable.GetName().ToString());
+		    	}
+		    }
+	    }
+    }
+
 	// Validate injection site
 	InAsset->DefaultInjectionSite.Reset();
 
@@ -334,6 +380,179 @@ void FUtils::CompileVariables(const FRigVMCompileSettings& InSettings, UUAFRigVM
 	}
 }
 
+// Helper function to sort ExternalVariables according to property-size/name without needing to rely on RigVMRegistry (which is too slow)
+template<typename SortType, typename ValueType, class SortPredicate>
+void SortByValuesArray(TArrayView<SortType> ToSort, TConstArrayView<ValueType> Values, const SortPredicate& SortPred)
+{
+	const int32 NumValues = Values.Num();
+	if (!ensureMsgf(ToSort.Num() == NumValues, TEXT("Values array must have same length as ToSort array, but did not (%d vs %d)"), NumValues, ToSort.Num()))
+	{
+		return;
+	}
+
+	// make a reference array of indices, that we will sort based on values
+	TArray<int32> Indices;
+	Indices.SetNumUninitialized(NumValues);
+	for (int32 Idx = 0; Idx < NumValues; ++Idx)
+	{
+		Indices[Idx] = Idx;
+	}
+	
+	Indices.Sort([&SortPred, &Values](int32 IdxA, int32 IdxB)
+	{
+		return SortPred(Values[IdxA], Values[IdxB]);
+	});
+
+	// swap the ToSort array to follow the ordering of the reference array
+	for (int32 Idx = 0; Idx < NumValues; ++Idx)
+	{
+		int32 SwapFromIndex = Indices[Idx];
+		while (SwapFromIndex < Idx)
+		{
+			SwapFromIndex = Indices[SwapFromIndex];
+		}
+
+		if (SwapFromIndex != Idx)
+		{
+			Swap(ToSort[Idx], ToSort[SwapFromIndex]);
+		}
+	}
+}
+
+void FUtils::GetExternalVariables(const UUAFRigVMAssetEditorData* EditorData, TArray<FRigVMExternalVariable>& OutVariables)
+{
+	GetInternalVariables(EditorData, OutVariables);
+
+	TSet<const UScriptStruct*> ReferencedVariableStructs;
+	TSet<const UUAFRigVMAsset*> ReferencedSharedVariableAssets;
+	TArray<const UUAFRigVMAsset*> RecursiveSharedVariableAssets;
+	TArray<const UScriptStruct*> RecursiveSharedVariableStructs;
+	for(UUAFRigVMAssetEntry* Entry : EditorData->Entries)
+	{
+		if(UUAFSharedVariablesEntry* SharedVariablesEntry = Cast<UUAFSharedVariablesEntry>(Entry))
+		{
+			switch (SharedVariablesEntry->GetType())
+			{
+			case EAnimNextSharedVariablesType::Asset:
+				{
+					const UUAFSharedVariables* SharedVariablesAsset = SharedVariablesEntry->GetAsset();
+					RetrieveReferencedSharedVariableSources(SharedVariablesAsset, RecursiveSharedVariableAssets, RecursiveSharedVariableStructs);
+					ReferencedSharedVariableAssets.Add(SharedVariablesAsset);
+					break;
+				}
+			case EAnimNextSharedVariablesType::Struct:
+				{
+					const UScriptStruct* SharedVariablesStruct = SharedVariablesEntry->GetStruct();
+					ReferencedVariableStructs.Add(SharedVariablesStruct);
+					break;
+				}
+			}
+		}
+	}
+	
+	ReferencedSharedVariableAssets.Append(RecursiveSharedVariableAssets);
+	ReferencedVariableStructs.Append(RecursiveSharedVariableStructs);
+	
+	// Append each _sorted_ set of referenced shared asset variables
+	for (const UUAFRigVMAsset* ReferencedVariableAsset : ReferencedSharedVariableAssets.ArrayView())
+	{
+		GetInternalVariables(GetEditorData<const UUAFRigVMAssetEditorData>(ReferencedVariableAsset), OutVariables);
+	}
+	
+	// Append each set of referenced shared struct variables, sorted by their property offset (can differ from its declaration)
+	for (const UScriptStruct* ReferencedVariableStruct : ReferencedVariableStructs.ArrayView())
+	{
+		TArray<int32> PropertyOffsets;
+		TArray<FRigVMExternalVariable> StructVariables;
+		for (TFieldIterator<FProperty> It(ReferencedVariableStruct); It; ++It)
+		{
+			const FProperty* Property = *It;
+			const FRigVMExternalVariable Variable = FRigVMExternalVariable::Make(Property, nullptr);
+			if (!Variable.IsValid(true))
+			{
+				UE_LOG(LogRigVM, Error, TEXT("%s: Property '%s' of type '%s' is not valid."), *ReferencedVariableStruct->GetName(), *Property->GetName(), *Property->GetCPPType());
+				continue;
+			}
+
+			StructVariables.Add(Variable);
+			PropertyOffsets.Add(It->GetOffset_ForInternal());
+		}
+
+		SortByValuesArray(MakeArrayView(StructVariables), MakeConstArrayView(PropertyOffsets), 
+		    [](const int32 OffsetA, const int32 OffsetB) -> bool
+		    {
+			    return OffsetA < OffsetB;
+		    });
+
+		OutVariables.Append(StructVariables);
+	}
+}
+
+void FUtils::GetInternalVariables(const UUAFRigVMAssetEditorData* EditorData, TArray<FRigVMExternalVariable>& OutVariables)
+{
+	TArray<TPair<FName, FAnimNextParamType>> VariableInfo;
+	TArray<FRigVMExternalVariable> AssetVariables;
+	EditorData->ForEachEntryOfType<UAnimNextVariableEntry>([AssetName = EditorData->GetAssetName(), &AssetVariables, &VariableInfo](const UAnimNextVariableEntry* Entry)
+	{
+		const FRigVMExternalVariable Variable = FRigVMExternalVariable::Make(Entry->GetVariableName(), Entry->GetType().ToString(), const_cast<UObject*>(Entry->GetType().GetValueTypeObject()), Entry->Access == EAnimNextExportAccessSpecifier::Public, false);
+		if (!Variable.IsValid(true))
+		{
+			UE_LOG(LogRigVM, Error, TEXT("%s: Property '%s' of type '%s' is not valid."), *AssetName, *Entry->GetName(), *Variable.GetExtendedCPPType().ToString());
+			return true;
+		}
+		
+		VariableInfo.Add({Entry->GetVariableName(), Entry->GetType() });
+
+		AssetVariables.Add(Variable);
+		return true;
+	});
+
+	// Handle programmatically generated variables
+	{
+		FAnimNextRigVMAssetCompileContext CompileContext = { EditorData };
+		FRigVMCompileSettings Settings;
+		const FAnimNextGetVariableCompileContext GetVariableCompileContext = EditorData->GetVariableCompileContext(Settings, CompileContext);
+		for (const FAnimNextProgrammaticVariable& ProgrammaticVariable : GetVariableCompileContext.GetProgrammaticVariables())
+		{
+			VariableInfo.Add({ProgrammaticVariable.GetVariableName(), ProgrammaticVariable.GetType()});
+			AssetVariables.Add(FRigVMExternalVariable::Make(ProgrammaticVariable.GetVariableName(), ProgrammaticVariable.GetType().ToString(), const_cast<UObject*>(ProgrammaticVariable.GetType().GetValueTypeObject()), false, false));
+		}
+	}
+	
+	SortByValuesArray(MakeArrayView(AssetVariables), MakeConstArrayView(VariableInfo), 
+		[](const TPair<FName, FAnimNextParamType>& A, const TPair<FName, FAnimNextParamType>& B)-> bool
+	    {
+		    if(A.Value.GetSize() != B.Value.GetSize())
+		    {
+			    return A.Value.GetSize() > B.Value.GetSize();
+		    }
+    
+		    return A.Key.LexicalLess(B.Key);
+	    });
+	OutVariables.Append(AssetVariables);
+}
+
+void FUtils::RetrieveReferencedSharedVariableSources(const UUAFRigVMAsset* InAsset, TArray<const UUAFRigVMAsset*>& ReferencedVariableAssets, TArray<const UScriptStruct*>& ReferencedVariableStructs)
+{
+	const UUAFRigVMAssetEditorData* EditorData = GetEditorData<UUAFRigVMAssetEditorData>(InAsset);
+	EditorData->ForEachEntryOfType<UUAFSharedVariablesEntry>([&ReferencedVariableAssets, &ReferencedVariableStructs](const UUAFSharedVariablesEntry* Entry)
+	{
+		if (Entry->Asset)
+		{
+			ensure(!Entry->HasAllFlags(RF_NeedPostLoad));
+			if (ReferencedVariableAssets.AddUnique(Entry->Asset) != INDEX_NONE)
+			{
+				RetrieveReferencedSharedVariableSources(Entry->Asset.Get(), ReferencedVariableAssets, ReferencedVariableStructs);
+			}
+		 }
+		 else if (Entry->Struct)
+		 {
+			ReferencedVariableStructs.AddUnique(Entry->Struct);
+		 }
+		return true;
+	});
+}
+
 void FUtils::CompileVariableBindings(const FRigVMCompileSettings& InSettings, UUAFRigVMAsset* InAsset, TArray<URigVMGraph*>& OutGraphs)
 {
 	CompileVariableBindingsInternal(InSettings, InAsset, OutGraphs, true);
diff --git a/Engine/Plugins/Experimental/UAF/UAF/Source/UAFUncookedOnly/Public/IAnimNextUncookedOnlyModule.h b/Engine/Plugins/Experimental/UAF/UAF/Source/UAFUncookedOnly/Public/IAnimNextUncookedOnlyModule.h
--- a/Engine/Plugins/Experimental/UAF/UAF/Source/UAFUncookedOnly/Public/IAnimNextUncookedOnlyModule.h
+++ b/Engine/Plugins/Experimental/UAF/UAF/Source/UAFUncookedOnly/Public/IAnimNextUncookedOnlyModule.h
@@ -5,9 +5,12 @@
 #include "CoreMinimal.h"
 #include "Modules/ModuleInterface.h"
 #include "Modules/ModuleManager.h"
+#include "RigVMCore/RigVMExternalVariable.h"
 
 template<typename T> struct TInstancedStruct;
 
+class UUAFRigVMAsset;
+
 namespace UE::UAF::UncookedOnly
 {
 	class IVariableBindingType;
@@ -42,6 +45,15 @@ public:
 	// @param	InGuid			FGuid of the variable to retrieve the FName for
 	// @param	InObject		Object, supposed to be, containing the variable (UScriptStruct, UUAFRigVMAsset or UUAFRigVMAssetEditorData)
 	virtual FName GetVariableNameByGuid(const FGuid InGuid, const UObject* InObject) const = 0;
+	
+	// Try and retrieve all local and referenced (through shared variable asset/struct(s)) variables for provided asset
+	// @param	InAsset			Asset to retrieve variables for
+	// @param	OutVariables	Output array of (referenced) external variables
+	virtual void GetExternalVariablesForAsset(const UUAFRigVMAsset* InAsset, TArray<FRigVMExternalVariable>& OutVariables) const = 0;
+	
+	// Whether an Asset its compiled data is out-of-date, indicating it should not be relied upon
+	// @param	InAsset			Asset to check for
+	virtual bool DoesAssetVariablesRequireCompilation(const UUAFRigVMAsset* InAsset) const = 0;
 #endif // WITH_EDITOR
 	
 	static IAnimNextUncookedOnlyModule& Get() { return FModuleManager::Get().LoadModuleChecked<IAnimNextUncookedOnlyModule>("UAFUncookedOnly"); }
diff --git a/Engine/Plugins/Experimental/UAF/UAFTestSuites/Source/UAFAnimGraphTestSuite/Private/Module.cpp b/Engine/Plugins/Experimental/UAF/UAFTestSuites/Source/UAFAnimGraphTestSuite/Private/Module.cpp
--- a/Engine/Plugins/Experimental/UAF/UAFTestSuites/Source/UAFAnimGraphTestSuite/Private/Module.cpp
+++ b/Engine/Plugins/Experimental/UAF/UAFTestSuites/Source/UAFAnimGraphTestSuite/Private/Module.cpp
@@ -1,6 +1,7 @@
 // Copyright Epic Games, Inc. All Rights Reserved.
 
 #include "CoreMinimal.h"
+#include "AnimNextAnimGraphTraitGraphTest.h"
 #include "Modules/ModuleInterface.h"
 #include "Modules/ModuleManager.h"
 
@@ -8,6 +9,26 @@ namespace UE::UAF::AnimGraph::Tests
 {
 	class FModule : public IModuleInterface
 	{
+	public:
+		virtual void StartupModule() override
+		{
+			// Register struct/class types required for contained tests
+			static UScriptStruct* const AllowedStructTypes[] =
+			{
+				FTestDerivedVector::StaticStruct()
+			};
+
+			FRigVMRegistry& RigVMRegistry = FRigVMRegistry::Get();
+			RigVMRegistry.RegisterStructTypes(AllowedStructTypes);
+			
+			static TPair<UClass*, FRigVMRegistry::ERegisterObjectOperation> const AllowedObjectTypes[] =
+			{
+				{ UUAFSharedVariables::StaticClass(), FRigVMRegistry::ERegisterObjectOperation::Class },
+			};
+
+			
+			RigVMRegistry.RegisterObjectTypes(AllowedObjectTypes);
+		}
 	};
 }
 
