Support BlendSpace1D variables in UAF.

#Synced-CL 51091996

diff --git a/Engine/Plugins/Experimental/UAF/UAF/Source/UAF/Private/AnimNextModuleImpl.cpp b/Engine/Plugins/Experimental/UAF/UAF/Source/UAF/Private/AnimNextModuleImpl.cpp
--- a/Engine/Plugins/Experimental/UAF/UAF/Source/UAF/Private/AnimNextModuleImpl.cpp
+++ b/Engine/Plugins/Experimental/UAF/UAF/Source/UAF/Private/AnimNextModuleImpl.cpp
@@ -9,6 +9,7 @@
 #include "UniversalObjectLocator.h"
 #include "Animation/AnimSequence.h"
 #include "Animation/BlendSpace.h"
+#include "Animation/BlendSpace1D.h"
 #include "Animation/BlendProfile.h"
 #include "Component/AnimNextComponent.h"
 #include "Curves/CurveFloat.h"
@@ -59,6 +60,7 @@ namespace UE::UAF
 		{
 			{ UAnimSequence::StaticClass(), FRigVMRegistry::ERegisterObjectOperation::Class },
 			{ UBlendSpace::StaticClass(), FRigVMRegistry::ERegisterObjectOperation::Class },
+			{ UBlendSpace1D::StaticClass(), FRigVMRegistry::ERegisterObjectOperation::Class },
 			{ UScriptStruct::StaticClass(), FRigVMRegistry::ERegisterObjectOperation::Class },
 			{ UBlendProfile::StaticClass(), FRigVMRegistry::ERegisterObjectOperation::Class },
 			{ UCurveFloat::StaticClass(), FRigVMRegistry::ERegisterObjectOperation::Class },
