Support IRigVMRuntimeAsset in UAF variables. This allows control rig public variables to be used in the same way as UAF shared variables assets. We can now set control rig variables via the UAF set variable API (e.g. from blueprint, c++), and we can also control control rig variables from within UAF systems

#Synced-CL 51100498

diff --git a/Engine/Plugins/Experimental/UAF/UAF/Source/UAF/Internal/Variables/UAFInstanceVariableContainer.h b/Engine/Plugins/Experimental/UAF/UAF/Source/UAF/Internal/Variables/UAFInstanceVariableContainer.h
--- a/Engine/Plugins/Experimental/UAF/UAF/Source/UAF/Internal/Variables/UAFInstanceVariableContainer.h
+++ b/Engine/Plugins/Experimental/UAF/UAF/Source/UAF/Internal/Variables/UAFInstanceVariableContainer.h
@@ -6,6 +6,7 @@
 #include "StructUtils/PropertyBag.h"
 #include "UAFInstanceVariableContainer.generated.h"
 
+class IRigVMRuntimeAssetInterface;
 struct FRigVMExternalVariableRuntimeData;
 struct FRigVMExternalVariableDef;
 struct FUAFInstanceVariableContainerProxy;
@@ -56,6 +57,9 @@ private:
 	// Construct from an asset, copying the variables. Optional overrides provided.
 	explicit FUAFInstanceVariableContainer(const FUAFAssetInstance& InHostInstance, const UUAFRigVMAsset* InAsset, const TSharedPtr<const UE::UAF::FVariableOverridesCollection>& InOverrides = nullptr, int32 InOverridesIndex = INDEX_NONE);
 
+	// Construct from a IRigVMRuntimeAssetInterface, copying the variables. Optional overrides provided.
+	explicit FUAFInstanceVariableContainer(const FUAFAssetInstance& InHostInstance, const IRigVMRuntimeAssetInterface* InAsset, const TSharedPtr<const UE::UAF::FVariableOverridesCollection>& InOverrides = nullptr, int32 InOverridesIndex = INDEX_NONE);
+
 	// Gets the value of the specified variable. If the type does not match exactly then a conversion will be attempted.
 	EPropertyBagResult GetVariable(int32 InVariableIndex, const FAnimNextParamType& InType, TArrayView<uint8> OutResult) const;
 
@@ -77,6 +81,9 @@ private:
 	// Unchecked, non-type converting version of AccessVariable
 	void AccessVariableUnchecked(int32 InVariableIndex, const FProperty*& OutProperty, TArrayView<uint8>& OutResult);
 
+	// Get the name of the underlying object (asset, struct etc) that this instance variable container is initialized from
+	FName GetFName() const;
+
 #if WITH_EDITOR
 	// Used during compilation, migrate all variables to new bags according to new defaults.
 	void Migrate();
@@ -94,7 +101,8 @@ private:
 	// The asset or struct these variables came from, if any
 	using FAssetType = TObjectPtr<const UUAFRigVMAsset>;
 	using FStructType = TObjectPtr<const UScriptStruct>;
-	using FAssetOrStructType = TVariant<FAssetType, FStructType>;
+	using FRigVMAssetType = TScriptInterface<const IRigVMRuntimeAssetInterface>;
+	using FAssetOrStructType = TVariant<FAssetType, FStructType, FRigVMAssetType>;
 	FAssetOrStructType AssetOrStructData;
 
 	// Instance of the asset's variables or resolved overrides for the asset's variables
diff --git a/Engine/Plugins/Experimental/UAF/UAF/Source/UAF/Private/Param/ParamType.cpp b/Engine/Plugins/Experimental/UAF/UAF/Source/UAF/Private/Param/ParamType.cpp
--- a/Engine/Plugins/Experimental/UAF/UAF/Source/UAF/Private/Param/ParamType.cpp
+++ b/Engine/Plugins/Experimental/UAF/UAF/Source/UAF/Private/Param/ParamType.cpp
@@ -661,7 +661,7 @@ FAnimNextParamType FAnimNextParamType::FromString(const FString& InString)
 			ObjectValueType = EValueType::Enum;
 		}
 
-		if(UObject* ObjectType = FindFirstObject<UObject>(*ObjectInnerString, EFindFirstObjectOptions::NativeFirst))
+		if(UObject* ObjectType = FindFirstObjectSafe<UObject>(*ObjectInnerString, EFindFirstObjectOptions::NativeFirst))
 		{
 			OutType.ValueType = ObjectValueType;
 			OutType.ValueTypeObject = nullptr;
diff --git a/Engine/Plugins/Experimental/UAF/UAF/Source/UAF/Private/Variables/AnimNextVariableReference.cpp b/Engine/Plugins/Experimental/UAF/UAF/Source/UAF/Private/Variables/AnimNextVariableReference.cpp
--- a/Engine/Plugins/Experimental/UAF/UAF/Source/UAF/Private/Variables/AnimNextVariableReference.cpp
+++ b/Engine/Plugins/Experimental/UAF/UAF/Source/UAF/Private/Variables/AnimNextVariableReference.cpp
@@ -9,6 +9,8 @@
 #include "IAnimNextUncookedOnlyModule.h"
 #endif // WITH_EDITOR
 
+#include "RigVMRuntimeAsset.h"
+
 #include UE_INLINE_GENERATED_CPP_BY_NAME(AnimNextVariableReference)
 
 FAnimNextVariableReference FAnimNextVariableReference::FromProperty(const FProperty* InProperty, const UScriptStruct* InStruct)
@@ -67,7 +69,7 @@ FAnimNextVariableReference::FAnimNextVariableReference(const FAnimNextSoftVariab
 	: Name(InSoftReference.GetName())
 	, Object(InSoftReference.GetSoftObjectPath().TryLoad())
 {
-	check(Object == nullptr || Object->IsA<UUAFRigVMAsset>() || Object->IsA<UScriptStruct>());
+	check(Object == nullptr || Object->IsA<UUAFRigVMAsset>() || Object->IsA<UScriptStruct>() || Object->Implements<URigVMRuntimeAssetInterface>());
 
 #if WITH_EDITOR
 	ensureAlwaysMsgf(!UE::UAF::IsInternalVariableName(Name), TEXT("Attempting to create variable reference using a variable name indicative of a programmatic variable %s"), *Name.ToString());
@@ -219,6 +221,13 @@ const FProperty* FAnimNextVariableReference::ResolveProperty() const
 				CachedProperty = FoundProperty;
 			}
 		}
+		else if (const IRigVMRuntimeAssetInterface* RigVMAsset = Cast<IRigVMRuntimeAssetInterface>(Object))
+		{
+			if (const FProperty* FoundProperty = RigVMAsset->FindGeneratedPropertyByName(Name))
+			{
+				CachedProperty = FoundProperty;
+			}
+		}
 	}
 	return CachedProperty.Get();
 }
diff --git a/Engine/Plugins/Experimental/UAF/UAF/Source/UAF/Private/Variables/UAFInstanceVariableContainer.cpp b/Engine/Plugins/Experimental/UAF/UAF/Source/UAF/Private/Variables/UAFInstanceVariableContainer.cpp
--- a/Engine/Plugins/Experimental/UAF/UAF/Source/UAF/Private/Variables/UAFInstanceVariableContainer.cpp
+++ b/Engine/Plugins/Experimental/UAF/UAF/Source/UAF/Private/Variables/UAFInstanceVariableContainer.cpp
@@ -3,6 +3,8 @@
 #include "Variables/UAFInstanceVariableContainer.h"
 
 #include "AnimNextRigVMAsset.h"
+
+#include "RigVMRuntimeAsset.h"
 #include "Variables/StructDataCache.h"
 #include "VariableOverrides.h"
 #include "Param/ParamType.h"
@@ -430,6 +432,37 @@ FUAFInstanceVariableContainer::FUAFInstanceVariableContainer(const FUAFAssetInst
 	ResolveOverrides();
 }
 
+FUAFInstanceVariableContainer::FUAFInstanceVariableContainer(const FUAFAssetInstance& InHostInstance, const IRigVMRuntimeAssetInterface* InAsset, const TSharedPtr<const UE::UAF::FVariableOverridesCollection>& InOverrides, int32 InOverridesIndex)
+	: HostInstance(&InHostInstance)
+{
+	TArray<FRigVMExternalVariable> ExternalVariables = const_cast<IRigVMRuntimeAssetInterface*>(InAsset)->GetExternalVariables();
+	FInstancedPropertyBag DefaultProperties;
+	
+	for (const FRigVMExternalVariable& Variable : ExternalVariables)
+	{
+		// Add to schema
+		DefaultProperties.AddProperty(Variable.GetName(), Variable.GetProperty());
+
+		// Copy in the default value
+		DefaultProperties.SetValue(Variable.GetName(), Variable.GetProperty(), Variable.GetMemory());
+	}
+	
+	NumVariables = DefaultProperties.GetNumPropertiesInBag();
+
+	// Copy the variables into the container
+	VariablesContainer.Set<FInstancedPropertyBag>(MoveTemp(DefaultProperties));
+
+	const UObject* InterfaceObject = Cast< const UObject>(InAsset);
+	TScriptInterface<const IRigVMRuntimeAssetInterface> Interface = InterfaceObject;
+	AssetOrStructData.Set<FRigVMAssetType>(Interface);
+
+	Overrides = InOverrides;
+	OverridesIndex = InOverridesIndex;
+	ResolveOverrides();
+	
+}
+
+
 EPropertyBagResult FUAFInstanceVariableContainer::GetVariable(int32 InVariableIndex, const FAnimNextParamType& InType, TArrayView<uint8> OutResult) const
 {
 	using namespace UE::UAF;
@@ -442,7 +475,10 @@ EPropertyBagResult FUAFInstanceVariableContainer::GetVariable(int32 InVariableIn
 	{
 	case FVariablesContainerType::IndexOfType<FInstancedPropertyBag>():
 		{
-			const UPropertyBag* PropertyBag = AssetOrStructData.Get<FAssetType>()->GetVariableDefaults().GetPropertyBagStruct();
+			const FInstancedPropertyBag& Variables = VariablesContainer.Get<FInstancedPropertyBag>();
+			
+			const UPropertyBag* PropertyBag = Variables.GetPropertyBagStruct();
+			
 			check(PropertyBag != nullptr);
 			TConstArrayView<FPropertyBagPropertyDesc> Descs = PropertyBag->GetPropertyDescs();
 			const FPropertyBagPropertyDesc& Desc = Descs[InVariableIndex];
@@ -454,7 +490,6 @@ EPropertyBagResult FUAFInstanceVariableContainer::GetVariable(int32 InVariableIn
 			}
 			else
 			{
-				const FInstancedPropertyBag& Variables = VariablesContainer.Get<FInstancedPropertyBag>();
 				Memory = Desc.CachedProperty->ContainerPtrToValuePtr<uint8>(Variables.GetValue().GetMemory());
 			}
 			break;
@@ -509,7 +544,10 @@ EPropertyBagResult FUAFInstanceVariableContainer::AccessVariable(int32 InVariabl
 	{
 	case FVariablesContainerType::IndexOfType<FInstancedPropertyBag>():
 		{
-			const UPropertyBag* PropertyBag = AssetOrStructData.Get<FAssetType>()->GetVariableDefaults().GetPropertyBagStruct();
+			FInstancedPropertyBag& Variables = VariablesContainer.Get<FInstancedPropertyBag>();
+			
+			const UPropertyBag* PropertyBag = Variables.GetPropertyBagStruct();
+			
 			check(PropertyBag != nullptr);
 			TConstArrayView<FPropertyBagPropertyDesc> Descs = PropertyBag->GetPropertyDescs();
 			const FPropertyBagPropertyDesc& Desc = Descs[InVariableIndex];
@@ -521,7 +559,6 @@ EPropertyBagResult FUAFInstanceVariableContainer::AccessVariable(int32 InVariabl
 			}
 			else
 			{
-				FInstancedPropertyBag& Variables = VariablesContainer.Get<FInstancedPropertyBag>();
 				Memory = Desc.CachedProperty->ContainerPtrToValuePtr<uint8>(Variables.GetMutableValue().GetMemory());
 			}
 			break;
@@ -572,7 +609,9 @@ void FUAFInstanceVariableContainer::AccessVariableUnchecked(int32 InVariableInde
 	{
 	case FVariablesContainerType::IndexOfType<FInstancedPropertyBag>():
 		{
-			const UPropertyBag* PropertyBag = AssetOrStructData.Get<FAssetType>()->GetVariableDefaults().GetPropertyBagStruct();
+			FInstancedPropertyBag& Variables = VariablesContainer.Get<FInstancedPropertyBag>();
+			const UPropertyBag* PropertyBag = Variables.GetPropertyBagStruct();
+			
 			check(PropertyBag != nullptr);
 			TConstArrayView<FPropertyBagPropertyDesc> Descs = PropertyBag->GetPropertyDescs();
 			const FPropertyBagPropertyDesc& Desc = Descs[InVariableIndex];
@@ -584,7 +623,6 @@ void FUAFInstanceVariableContainer::AccessVariableUnchecked(int32 InVariableInde
 			}
 			else
 			{
-				FInstancedPropertyBag& Variables = VariablesContainer.Get<FInstancedPropertyBag>();
 				Memory = Desc.CachedProperty->ContainerPtrToValuePtr<uint8>(Variables.GetMutableValue().GetMemory());
 			}
 			break;
@@ -618,6 +656,33 @@ void FUAFInstanceVariableContainer::AccessVariableUnchecked(int32 InVariableInde
 	OutResult = TArrayView<uint8>(Memory, Size);
 }
 
+FName FUAFInstanceVariableContainer::GetFName() const
+{
+	switch (AssetOrStructData.GetIndex())
+	{
+		case FAssetOrStructType::IndexOfType<FAssetType>():
+			{
+				return AssetOrStructData.Get<FAssetType>().GetFName();
+			}
+		break;
+		case FAssetOrStructType::IndexOfType<FRigVMAssetType>():
+			{
+				return AssetOrStructData.Get<FRigVMAssetType>().GetObject()->GetFName();
+			}
+		break;
+	case FAssetOrStructType::IndexOfType<FStructType>():
+		{
+			return AssetOrStructData.Get<FStructType>().GetFName();
+		}
+		break;
+	default:
+		checkNoEntry();
+		break;
+	}
+
+	return NAME_None;
+}
+
 EPropertyBagResult FUAFInstanceVariableContainer::SetVariable(int32 InVariableIndex, const FAnimNextParamType& InType, TConstArrayView<uint8> InNewValue)
 {
 	using namespace UE::UAF;
@@ -630,7 +695,9 @@ EPropertyBagResult FUAFInstanceVariableContainer::SetVariable(int32 InVariableIn
 	{
 	case FVariablesContainerType::IndexOfType<FInstancedPropertyBag>():
 		{
-			const UPropertyBag* PropertyBag = AssetOrStructData.Get<FAssetType>()->GetVariableDefaults().GetPropertyBagStruct();
+			FInstancedPropertyBag& Variables = VariablesContainer.Get<FInstancedPropertyBag>();
+			const UPropertyBag* PropertyBag = Variables.GetPropertyBagStruct();
+			
 			check(PropertyBag != nullptr);
 			TConstArrayView<FPropertyBagPropertyDesc> Descs = PropertyBag->GetPropertyDescs();
 			const FPropertyBagPropertyDesc& Desc = Descs[InVariableIndex];
@@ -642,7 +709,6 @@ EPropertyBagResult FUAFInstanceVariableContainer::SetVariable(int32 InVariableIn
 			}
 			else
 			{
-				FInstancedPropertyBag& Variables = VariablesContainer.Get<FInstancedPropertyBag>();
 				Memory = Desc.CachedProperty->ContainerPtrToValuePtr<uint8>(Variables.GetMutableValue().GetMemory());
 			}
 			break;
@@ -697,8 +763,22 @@ void FUAFInstanceVariableContainer::GetAllVariablesOfType(const FAnimNextParamTy
 	{
 	case FVariablesContainerType::IndexOfType<FInstancedPropertyBag>():
 		{
-			const FAssetType& Asset = AssetOrStructData.Get<FAssetType>();
-			const UPropertyBag* PropertyBag = AssetOrStructData.Get<FAssetType>()->GetVariableDefaults().GetPropertyBagStruct();
+			const FInstancedPropertyBag& Variables = VariablesContainer.Get<FInstancedPropertyBag>();
+			const UPropertyBag* PropertyBag = Variables.GetPropertyBagStruct();;
+			const UObject* OwnerObject = nullptr;
+
+			if (AssetOrStructData.IsType<FAssetType>())
+			{
+				const FAssetType& Asset = AssetOrStructData.Get<FAssetType>();
+				OwnerObject = Asset;
+			}
+			else if(AssetOrStructData.IsType<FRigVMAssetType>())
+			{
+				FRigVMAssetType AssetInterface = AssetOrStructData.Get<FRigVMAssetType>();
+				check(AssetInterface.GetObject());
+				OwnerObject = AssetInterface.GetObject();
+			}
+			
 			check(PropertyBag != nullptr);
 			for (const FPropertyBagPropertyDesc& Desc : PropertyBag->GetPropertyDescs())
 			{
@@ -706,7 +786,7 @@ void FUAFInstanceVariableContainer::GetAllVariablesOfType(const FAnimNextParamTy
 					Desc.ContainerTypes.GetFirstContainerType() == InType.GetContainerType() &&
 					Desc.ValueTypeObject == InType.GetValueTypeObject())
 				{
-					OutVariables.Add(FAnimNextVariableReference::FromName(Desc.Name, Asset));
+					OutVariables.Add(FAnimNextVariableReference::FromName(Desc.Name, OwnerObject));
 				}
 			}
 			break;
@@ -742,8 +822,8 @@ bool FUAFInstanceVariableContainer::ForEachVariable(TFunctionRef<bool(FName, con
 	{
 	case FVariablesContainerType::IndexOfType<FInstancedPropertyBag>():
 		{
-			const FAssetType& Asset = AssetOrStructData.Get<FAssetType>();
-			const UPropertyBag* PropertyBag = AssetOrStructData.Get<FAssetType>()->GetVariableDefaults().GetPropertyBagStruct();
+			const UPropertyBag* PropertyBag = VariablesContainer.Get<FInstancedPropertyBag>().GetPropertyBagStruct();
+
 			check(PropertyBag != nullptr);
 			for (const FPropertyBagPropertyDesc& Desc : PropertyBag->GetPropertyDescs())
 			{
@@ -784,7 +864,8 @@ FAnimNextParamType FUAFInstanceVariableContainer::GetVariableType(int32 InVariab
 	case FVariablesContainerType::IndexOfType<FInstancedPropertyBag>():
 		{
 			const FInstancedPropertyBag& Variables = VariablesContainer.Get<FInstancedPropertyBag>();
-			const UPropertyBag* PropertyBag = AssetOrStructData.Get<FAssetType>()->GetVariableDefaults().GetPropertyBagStruct();
+			const UPropertyBag* PropertyBag = Variables.GetPropertyBagStruct();
+			
 			check(PropertyBag != nullptr);
 			TConstArrayView<FPropertyBagPropertyDesc> Descs = PropertyBag->GetPropertyDescs();
 			return FAnimNextParamType::FromPropertyBagPropertyDesc(Descs[InVariableIndex]);
@@ -878,6 +959,45 @@ void FUAFInstanceVariableContainer::ResolveOverrides()
 			}
 			break;
 		}
+		case FAssetOrStructType::IndexOfType<FRigVMAssetType>():
+		{
+			const IRigVMRuntimeAssetInterface* RigVMAsset = AssetOrStructData.Get<FRigVMAssetType>().GetInterface();
+			checkf(RigVMAsset == PinnedOverrides->Collection[OverridesIndex].AssetOrStructData.Get<FVariableOverrides::FRigVMAssetType>().GetInterface(), TEXT("Variable overrides RigVMRuntimeAsset does not match"));
+			TConstArrayView<FVariableOverrides::FOverride> OverridesView = PinnedOverrides->Collection[OverridesIndex].Overrides;
+			const int32 NumOverrides = OverridesView.Num();
+			TArray<FRigVMExternalVariable> ExternalVariables = const_cast<IRigVMRuntimeAssetInterface*>(RigVMAsset)->GetExternalVariables();
+			for (int32 OverrideIndex = 0; OverrideIndex < NumOverrides; ++OverrideIndex)
+			{
+				const FVariableOverrides::FOverride& Override = OverridesView[OverrideIndex];
+				int32 VariableIndex = ExternalVariables.IndexOfByPredicate([&Override](const FRigVMExternalVariable& InVariable)
+					{
+						return InVariable.GetName() == Override.Name;
+					});
+
+				if (VariableIndex == INDEX_NONE)
+				{
+					ensureMsgf(false, TEXT("Variable not found"));
+					continue;
+				}
+				
+				const FProperty* Property = ExternalVariables[VariableIndex].GetProperty();
+				if (Property == nullptr)
+				{
+					ensureMsgf(false, TEXT("Property not found"));
+					continue;
+				}
+				
+				FPropertyBagPropertyDesc Desc(Override.Name, Property);
+				if (Override.Type != FAnimNextParamType::FromPropertyBagPropertyDesc(Desc))
+				{
+					ensureMsgf(false, TEXT("Variable type does not match"));
+					continue;
+				}
+				ResolvedOverrides[VariableIndex] = Override.Memory;
+			}
+			break;
+		}
+		break;
 	default:
 		checkNoEntry();
 		break;
@@ -888,20 +1008,33 @@ void FUAFInstanceVariableContainer::ResolveOverrides()
 void FUAFInstanceVariableContainer::Migrate()
 {
 	using namespace UE::UAF;
-	
-	switch (VariablesContainer.GetIndex())
+
+	switch (AssetOrStructData.GetIndex())
 	{
-	case FVariablesContainerType::IndexOfType<FInstancedPropertyBag>():
+	case FAssetOrStructType::IndexOfType<FAssetType>():
 		{
 			FInstancedPropertyBag& PropertyBag = VariablesContainer.Get<FInstancedPropertyBag>();
 			PropertyBag.MigrateToNewBagInstance(AssetOrStructData.Get<FAssetType>()->GetVariableDefaults());
 			NumVariables = PropertyBag.GetNumPropertiesInBag();
 			break;
 		}
-	case FVariablesContainerType::IndexOfType<FInstancedStruct>():
+	case FAssetOrStructType::IndexOfType<FStructType>():
 		{
-			FInstancedStruct& InstancedStruct = VariablesContainer.Get<FInstancedStruct>();
-			NumVariables = FStructDataCache::GetStructInfo(InstancedStruct.GetScriptStruct())->GetProperties().Num();
+			break;
+		}
+	case FAssetOrStructType::IndexOfType<FRigVMAssetType>():
+		{
+			FInstancedPropertyBag& PropertyBag = VariablesContainer.Get<FInstancedPropertyBag>();
+
+			const FRigVMAssetType& RigVMAsset = AssetOrStructData.Get<FRigVMAssetType>();
+			TArray<FRigVMExternalVariable> ExternalVariables = const_cast<IRigVMRuntimeAssetInterface*>(RigVMAsset.GetInterface())->GetExternalVariables();
+			for (const FRigVMExternalVariable& ExternalVariable : ExternalVariables)
+			{
+				PropertyBag.AddProperty(ExternalVariable.GetName(), ExternalVariable.GetProperty());
+			}
+
+			
+			NumVariables = PropertyBag.GetNumPropertiesInBag();
 			break;
 		}
 	default:
@@ -999,6 +1132,11 @@ void FUAFInstanceVariableContainer::AddStructReferencedObjects(FReferenceCollect
 			Collector.AddReferencedObject(AssetOrStructData.Get<FStructType>());
 			break;
 		}
+	case FAssetOrStructType::IndexOfType<FRigVMAssetType>():
+		{
+			Collector.AddReferencedObject(AssetOrStructData.Get<FRigVMAssetType>().GetObjectRef());
+			break;
+		}
 	default:
 		checkNoEntry();	
 	}
diff --git a/Engine/Plugins/Experimental/UAF/UAF/Source/UAF/Private/Variables/UAFInstanceVariableData.cpp b/Engine/Plugins/Experimental/UAF/UAF/Source/UAF/Private/Variables/UAFInstanceVariableData.cpp
--- a/Engine/Plugins/Experimental/UAF/UAF/Source/UAF/Private/Variables/UAFInstanceVariableData.cpp
+++ b/Engine/Plugins/Experimental/UAF/UAF/Source/UAF/Private/Variables/UAFInstanceVariableData.cpp
@@ -2,12 +2,12 @@
 
 #include "Variables/UAFInstanceVariableData.h"
 #include "AnimNextRigVMAsset.h"
+#include "RigVMRuntimeAsset.h"
 #include "UAFAssetInstance.h"
 #include "Variables/StructDataCache.h"
 #include "VariableOverrides.h"
 #include "Logging/StructuredLog.h"
 #include "Variables/AnimNextVariableReference.h"
-#include "Variables/AnimNextSharedVariables.h"
 #include "VariableOverridesCollection.h"
 #include "Variables/UAFInstanceVariableContainer.h"
 
@@ -34,6 +34,12 @@ void FUAFInstanceVariableData::Initialize(const FUAFAssetInstance& InInstance, c
 		AddVariablesContainerForStruct(Struct, InInstance, InOverrides);
 	}
 
+	// Add referenced rig VM asset variables
+	for (TScriptInterface<const IRigVMRuntimeAssetInterface> RigVMAsset : InInstance.GetAsset<UUAFRigVMAsset>()->GetReferencedVariableRigVMAssets())
+	{
+		AddVariablesContainerForRigVMAsset(RigVMAsset, InInstance, InOverrides);
+	}
+
 	// Now any overrides we havent already added
 	if (InOverrides.IsValid())
 	{
@@ -70,6 +76,26 @@ void FUAFInstanceVariableData::Initialize(const FUAFAssetInstance& InInstance, c
 				AllSharedVariableContainersMap.Add(Asset, ReferencedAssetIndex);
 				NumAllVariables += NewVariables->NumVariables;
 			}
+			else if (Override.AssetOrStructData.IsType<FVariableOverrides::FRigVMAssetType>())
+			{
+				FVariableOverrides::FRigVMAssetType RigVMAsset = Override.AssetOrStructData.Get<FVariableOverrides::FRigVMAssetType>();
+				if (AllSharedVariableContainersMap.Contains(RigVMAsset.GetObject()))
+				{
+					continue;
+				}
+
+				if (RigVMAsset.GetInterface() == nullptr)
+				{
+					continue;
+				}
+
+				TSharedPtr<FUAFInstanceVariableContainer> NewVariables = MakeShared<FUAFInstanceVariableContainer>(InInstance, RigVMAsset.GetInterface(), InOverrides, OverrideIndex);
+				OwnedVariableContainers.Add(NewVariables.ToSharedRef());
+				const int32 ReferencedAssetIndex = AllSharedVariableContainers.Emplace(NewVariables);
+				// NOTE: not adding to InternalSharedVariableContainers, as these arent needed for running the graph
+				AllSharedVariableContainersMap.Add(RigVMAsset.GetObject(), ReferencedAssetIndex);
+				NumAllVariables += NewVariables->NumVariables;
+			}
 		}
 	}
 
@@ -182,6 +208,59 @@ void FUAFInstanceVariableData::AddVariablesContainerForStruct(const UScriptStruc
 	NumInternalVariables += NewVariables->NumVariables;
 }
 
+void FUAFInstanceVariableData::AddVariablesContainerForRigVMAsset(const TScriptInterface<const IRigVMRuntimeAssetInterface>& InRigVMAsset,
+	const FUAFAssetInstance& InInstance, const TSharedPtr<const UE::UAF::FVariableOverridesCollection>& InOverrides)
+{
+	using namespace UE::UAF;
+
+	check(InRigVMAsset.GetInterface() != nullptr);
+
+	TSharedPtr<FUAFInstanceVariableContainer> NewVariables;
+
+	// First check to see we have an override
+	if(InOverrides.IsValid())
+	{
+		for (int32 OverrideIndex = 0; OverrideIndex < InOverrides->Collection.Num(); ++OverrideIndex)
+		{
+			const FVariableOverrides& Override = InOverrides->Collection[OverrideIndex];
+			if (Override.IsOverrideFor(InRigVMAsset.GetInterface()))
+			{
+				NewVariables = MakeShared<FUAFInstanceVariableContainer>(InInstance, InRigVMAsset.GetInterface(), InOverrides, OverrideIndex);
+				OwnedVariableContainers.Add(NewVariables.ToSharedRef());
+				break;
+			}
+		}
+	}
+
+	// Check to see if this asset exists in our host chain. If so, we take a reference.
+	if (!NewVariables.IsValid())
+	{
+		const FUAFAssetInstance* CurrentHostInstance = InInstance.GetHost();
+		while (CurrentHostInstance)
+		{
+			if (const int32* SetIndexPtr = CurrentHostInstance->Variables.AllSharedVariableContainersMap.Find(InRigVMAsset.GetObject()))
+			{
+				NewVariables = CurrentHostInstance->Variables.AllSharedVariableContainers[*SetIndexPtr].Pin();
+				break;
+			}
+			CurrentHostInstance = CurrentHostInstance->GetHost();
+		}
+	}
+
+	// Otherwise, create a new instance that we own
+	if (!NewVariables.IsValid())
+	{
+		NewVariables = MakeShared<FUAFInstanceVariableContainer>(InInstance, InRigVMAsset.GetInterface());
+		OwnedVariableContainers.Add(NewVariables.ToSharedRef());
+	}
+
+	const int32 ReferencedAssetIndex = AllSharedVariableContainers.Emplace(NewVariables);
+	InternalSharedVariableContainers.Add(NewVariables);
+	AllSharedVariableContainersMap.Add(InRigVMAsset.GetObject(), ReferencedAssetIndex);
+	NumAllVariables += NewVariables->NumVariables;
+	NumInternalVariables += NewVariables->NumVariables;
+}
+
 void FUAFInstanceVariableData::RebuildNameMaps()
 {
 	using namespace UE::UAF;
@@ -198,11 +277,23 @@ void FUAFInstanceVariableData::RebuildNameMaps()
 		{
 		case FUAFInstanceVariableContainer::FVariablesContainerType::IndexOfType<FInstancedPropertyBag>():
 			{
-				const UUAFRigVMAsset* Asset = VariableSet->AssetOrStructData.Get<FUAFInstanceVariableContainer::FAssetType>().Get();
-				const int32 NumVariablesInSet = VariableSet->NumVariables;
+				const FInstancedPropertyBag& PropertyBag = VariableSet->VariablesContainer.Get<FInstancedPropertyBag>();
+				const int32 NumVariablesInSet = PropertyBag.GetNumPropertiesInBag();
+				const UObject* AssetObject = nullptr;
+				if (VariableSet->AssetOrStructData.IsType<FUAFInstanceVariableContainer::FAssetType>())
+				{
+					AssetObject = VariableSet->AssetOrStructData.Get<FUAFInstanceVariableContainer::FAssetType>();
+				}
+				else if (VariableSet->AssetOrStructData.IsType<FUAFInstanceVariableContainer::FRigVMAssetType>())
+				{
+					FUAFInstanceVariableContainer::FRigVMAssetType RigVMInterface = VariableSet->AssetOrStructData.Get<FUAFInstanceVariableContainer::FRigVMAssetType>();
+					AssetObject = RigVMInterface.GetObject();
+				}
+
+				check(AssetObject != nullptr);
+				
 				if (NumVariablesInSet > 0)
 				{
-					const FInstancedPropertyBag& PropertyBag = VariableSet->VariablesContainer.Get<FInstancedPropertyBag>();
 					TConstArrayView<FPropertyBagPropertyDesc> Descs = PropertyBag.GetPropertyBagStruct()->GetPropertyDescs();
 					check(Descs.Num() == NumVariablesInSet);
 					for (int32 VariableIndex = 0; VariableIndex < NumVariablesInSet; ++VariableIndex)
@@ -210,7 +301,7 @@ void FUAFInstanceVariableData::RebuildNameMaps()
 						if (!IsInternalVariableName(Descs[VariableIndex].Name))
 						{
 							FUAFInstanceVariableData::FVariableMapping Mapping(SharedVariableSetIndex, VariableIndex);
-							VariablesMap.Add(FAnimNextVariableReference::FromName(Descs[VariableIndex].Name, Asset), Mapping);
+							VariablesMap.Add(FAnimNextVariableReference::FromName(Descs[VariableIndex].Name, AssetObject), Mapping);
 							LegacyVariableNameMap.Add(Descs[VariableIndex].Name, Mapping);
 						}
 					}
@@ -266,9 +357,10 @@ void FUAFInstanceVariableData::Migrate(const FUAFAssetInstance& InInstance)
 	const UUAFRigVMAsset* InstanceAsset = InInstance.GetAsset<UUAFRigVMAsset>();
 	TConstArrayView<TObjectPtr<const UUAFRigVMAsset>> ReferencedVariableAssets = InstanceAsset->GetReferencedVariableAssets();
 	TConstArrayView<TObjectPtr<const UScriptStruct>> ReferencedVariableStructs = InstanceAsset->GetReferencedVariableStructs();
+	TConstArrayView<TScriptInterface<const IRigVMRuntimeAssetInterface>> ReferencedVariableRigVMAssets = InstanceAsset->GetReferencedVariableRigVMAssets();
 
 	// Referenced assets + struct + self-container
-	const int32 NumExpectedContainers = ReferencedVariableAssets.Num() + ReferencedVariableStructs.Num() + 1;
+	const int32 NumExpectedContainers = ReferencedVariableAssets.Num() + ReferencedVariableStructs.Num() + ReferencedVariableRigVMAssets.Num() + 1;
 	const int32 NumExistingContainers = AllSharedVariableContainersMap.Num();
 
 		// Verify the number of containers, new ones might have been added or existing ones removed
@@ -312,6 +404,24 @@ void FUAFInstanceVariableData::Migrate(const FUAFAssetInstance& InInstance)
 				AddVariablesContainerForStruct(Struct, InInstance, nullptr);
 			}
 		}
+
+		// Referenced IRigVMRuntimeAsset variables
+		for (TScriptInterface<const IRigVMRuntimeAssetInterface> ReferencedRigVMAsset : ReferencedVariableRigVMAssets)
+		{
+			const bool bHasVariableContainerForAsset = OwnedVariableContainers.ContainsByPredicate([ReferencedRigVMAsset](const TSharedRef<FUAFInstanceVariableContainer>& Container)
+			{
+				if (const FUAFInstanceVariableContainer::FRigVMAssetType* VariantPtr = Container->AssetOrStructData.TryGet<FUAFInstanceVariableContainer::FRigVMAssetType>())
+				{
+					return (*VariantPtr) == ReferencedRigVMAsset;
+				}
+				return false;
+			});
+		
+			if (!bHasVariableContainerForAsset)
+			{
+				AddVariablesContainerForRigVMAsset(ReferencedRigVMAsset, InInstance, nullptr);
+			}
+		}
 	}
 	else if (bOutOfDataContainers)
 	{		
@@ -329,6 +439,9 @@ void FUAFInstanceVariableData::Migrate(const FUAFAssetInstance& InInstance)
 				case FUAFInstanceVariableContainer::FAssetOrStructType::IndexOfType<FUAFInstanceVariableContainer::FStructType>():
 					Object = SharedVariableContainer->AssetOrStructData.Get<FUAFInstanceVariableContainer::FStructType>();
 					break;
+				case FUAFInstanceVariableContainer::FAssetOrStructType::IndexOfType<FUAFInstanceVariableContainer::FRigVMAssetType>():
+					Object = SharedVariableContainer->AssetOrStructData.Get<FUAFInstanceVariableContainer::FRigVMAssetType>().GetObject();
+					break;
 				default:
 					checkNoEntry()
 					break;
@@ -337,9 +450,10 @@ void FUAFInstanceVariableData::Migrate(const FUAFAssetInstance& InInstance)
 				// Referenced shared asset, or object corresponds to the instanced asset
 				const bool bIsAssetEntry = ReferencedVariableAssets.Contains(Object) || InstanceAsset == Object;
 				const bool bIsStructEntry = ReferencedVariableStructs.Contains(Object);
+				const bool bIsRigVMEntry = ReferencedVariableRigVMAssets.Contains(Object);
 
 				// Remove variable containers for asset/struct which are no longer referenced/part of this instance
-				if (!bIsAssetEntry && !bIsStructEntry)
+				if (!bIsAssetEntry && !bIsStructEntry && !bIsRigVMEntry)
 				{
 					AllSharedVariableContainers.RemoveAtSwap(Index);
 
@@ -379,6 +493,9 @@ void FUAFInstanceVariableData::Migrate(const FUAFAssetInstance& InInstance)
 		case FUAFInstanceVariableContainer::FAssetOrStructType::IndexOfType<FUAFInstanceVariableContainer::FStructType>():
 			Object = SharedVariableContainer->AssetOrStructData.Get<FUAFInstanceVariableContainer::FStructType>();
 			break;
+		case FUAFInstanceVariableContainer::FAssetOrStructType::IndexOfType<FUAFInstanceVariableContainer::FRigVMAssetType>():
+			Object = SharedVariableContainer->AssetOrStructData.Get<FUAFInstanceVariableContainer::FRigVMAssetType>().GetObject();
+			break;
 		default:
 			checkNoEntry()
 			break;
@@ -547,6 +664,28 @@ void FUAFInstanceVariableData::AccessVariablePropertyByIndex(int32 InIndex, TFun
 	}
 }
 
+EPropertyBagResult FUAFInstanceVariableData::AccessVariableProperty(const FAnimNextVariableReference& InVariable,
+	TFunctionRef<void(const FProperty*, TArrayView<uint8>)> InFunction)
+{
+	using namespace UE::UAF;
+
+	const FVariableMapping* FoundMapping = FindMapping(InVariable);
+	if (FoundMapping == nullptr)
+	{
+		return EPropertyBagResult::PropertyNotFound;
+	}
+
+	const TWeakPtr<FUAFInstanceVariableContainer>& SharedVariableSet = AllSharedVariableContainers[FoundMapping->AllSharedVariableContainersIndex];
+	TSharedPtr<FUAFInstanceVariableContainer> PinnedVariableSet = SharedVariableSet.Pin();
+	check(PinnedVariableSet.IsValid());
+
+	const FProperty* Property = nullptr;
+	TArrayView<uint8> Data;
+	PinnedVariableSet->AccessVariableUnchecked(FoundMapping->VariableIndex, Property, Data);
+	InFunction(Property, Data);
+	return EPropertyBagResult::Success;
+}
+
 void FUAFInstanceVariableData::AddStructReferencedObjects(FReferenceCollector& Collector)
 {
 	for (const TSharedRef<FUAFInstanceVariableContainer>& SharedVariableSet : OwnedVariableContainers)
diff --git a/Engine/Plugins/Experimental/UAF/UAF/Source/UAF/Private/Variables/VariableOverrides.cpp b/Engine/Plugins/Experimental/UAF/UAF/Source/UAF/Private/Variables/VariableOverrides.cpp
--- a/Engine/Plugins/Experimental/UAF/UAF/Source/UAF/Private/Variables/VariableOverrides.cpp
+++ b/Engine/Plugins/Experimental/UAF/UAF/Source/UAF/Private/Variables/VariableOverrides.cpp
@@ -1,8 +1,9 @@
-﻿// Copyright Epic Games, Inc. All Rights Reserved.
+// Copyright Epic Games, Inc. All Rights Reserved.
 
 #include "VariableOverrides.h"
 
 #include "AnimNextRigVMAsset.h"
+#include "RigVMRuntimeAsset.h"
 
 namespace UE::UAF
 {
@@ -22,4 +23,12 @@ FVariableOverrides::FVariableOverrides(const UScriptStruct* InStruct, TArray<FOv
 
 }
 
+FVariableOverrides::FVariableOverrides(const IRigVMRuntimeAssetInterface* InRigVMAsset, TArray<FOverride>&& InOverrides)
+	:Overrides(MoveTemp(InOverrides))
+{
+	check(InRigVMAsset);
+	TScriptInterface<const IRigVMRuntimeAssetInterface> Interface = Cast<const UObject>(InRigVMAsset);
+	check(Interface != nullptr);
+	AssetOrStructData.Set<FRigVMAssetType>(Interface);
+}
 }
diff --git a/Engine/Plugins/Experimental/UAF/UAF/Source/UAF/Private/Variables/VariableOverrides.h b/Engine/Plugins/Experimental/UAF/UAF/Source/UAF/Private/Variables/VariableOverrides.h
--- a/Engine/Plugins/Experimental/UAF/UAF/Source/UAF/Private/Variables/VariableOverrides.h
+++ b/Engine/Plugins/Experimental/UAF/UAF/Source/UAF/Private/Variables/VariableOverrides.h
@@ -5,6 +5,7 @@
 #include "UObject/NameTypes.h"
 #include "Param/ParamType.h"
 
+class IRigVMRuntimeAssetInterface;
 class UUAFRigVMAsset;
 struct FUAFInstanceVariableContainer;
 struct FUAFInstanceVariableData;
@@ -45,6 +46,9 @@ struct FVariableOverrides
 	// Construct a set of overrides for a struct
 	FVariableOverrides(const UScriptStruct* InStruct, TArray<FOverride>&& InOverrides);
 
+	// Construct a set of overrides for a IRigVMRuntimeAssetInterface 
+	FVariableOverrides(const IRigVMRuntimeAssetInterface* InRigVMAsset, TArray<FOverride>&& InOverrides);
+
 	bool IsOverrideFor(const UUAFRigVMAsset* InAsset) const
 	{
 		const TObjectPtr<const UUAFRigVMAsset>* AssetPtr = AssetOrStructData.TryGet<FAssetType>();
@@ -57,6 +61,12 @@ struct FVariableOverrides
 		return StructDataPtr && *StructDataPtr == InStruct;
 	}
 
+	bool IsOverrideFor(const IRigVMRuntimeAssetInterface* InRigVMAsset) const
+	{
+		const FRigVMAssetType* AssetPtr = AssetOrStructData.TryGet<FRigVMAssetType>();
+		return AssetPtr && *AssetPtr == InRigVMAsset;
+	}
+
 private:
 	friend FUAFInstanceVariableContainer;
 	friend FUAFInstanceVariableData;
@@ -65,7 +75,8 @@ private:
 	// Asset or struct to override variables for
 	using FAssetType = TObjectPtr<const UUAFRigVMAsset>;
 	using FStructType = TObjectPtr<const UScriptStruct>;
-	using FAssetOrStructType = TVariant<FAssetType, FStructType>;
+	using FRigVMAssetType = TScriptInterface<const IRigVMRuntimeAssetInterface>;
+	using FAssetOrStructType = TVariant<FAssetType, FStructType, FRigVMAssetType>;
 	FAssetOrStructType AssetOrStructData;
 
 	// Individual variable overrides
diff --git a/Engine/Plugins/Experimental/UAF/UAF/Source/UAF/Public/AnimNextRigVMAsset.h b/Engine/Plugins/Experimental/UAF/UAF/Source/UAF/Public/AnimNextRigVMAsset.h
--- a/Engine/Plugins/Experimental/UAF/UAF/Source/UAF/Public/AnimNextRigVMAsset.h
+++ b/Engine/Plugins/Experimental/UAF/UAF/Source/UAF/Public/AnimNextRigVMAsset.h
@@ -115,7 +115,10 @@ public:
 	TConstArrayView<TObjectPtr<const UUAFRigVMAsset>> GetReferencedVariableAssets() const { return ReferencedVariableAssets; }
 	
 	// Get all the native structs whose variables we will reference/instance at runtime
-   	TConstArrayView<TObjectPtr<const UScriptStruct>> GetReferencedVariableStructs() const { return ReferencedVariableStructs; }
+	TConstArrayView<TObjectPtr<const UScriptStruct>> GetReferencedVariableStructs() const { return ReferencedVariableStructs; }
+
+	// Get all the IRigVMRuntimeAssetInterfaces whose variables we will reference/instance at runtime
+	TConstArrayView<TScriptInterface<const IRigVMRuntimeAssetInterface>> GetReferencedVariableRigVMAssets() const { return ReferencedVariableRigVMAssets; }
 	
 private:
 	// URigVMHost interface
@@ -149,6 +152,10 @@ protected:
 	UPROPERTY()
 	TArray<TObjectPtr<const UScriptStruct>> ReferencedVariableStructs;
 
+	// All the IRigVMRuntimeAssetInterface whose variables we will reference/instance at runtime
+	UPROPERTY()
+	TArray<TScriptInterface<const IRigVMRuntimeAssetInterface>> ReferencedVariableRigVMAssets;
+
 	// All the natively-callable RigVM functions
 	UPROPERTY()
 	TArray<FAnimNextRigVMFunctionData> FunctionData;
diff --git a/Engine/Plugins/Experimental/UAF/UAF/Source/UAF/Public/UAFAssetInstance.h b/Engine/Plugins/Experimental/UAF/UAF/Source/UAF/Public/UAFAssetInstance.h
--- a/Engine/Plugins/Experimental/UAF/UAF/Source/UAF/Public/UAFAssetInstance.h
+++ b/Engine/Plugins/Experimental/UAF/UAF/Source/UAF/Public/UAFAssetInstance.h
@@ -65,6 +65,9 @@ struct FUAFAssetInstance
 		return Variables.GetVariable(InVariable, FAnimNextParamType::GetType<ValueType>(), TArrayView<uint8>(reinterpret_cast<uint8*>(&OutResult), sizeof(ValueType)));
 	}
 
+	FUAFInstanceVariableData& GetVariables() { return Variables;};
+	const FUAFInstanceVariableData& GetVariables() const { return Variables;};
+
 	// Access a variable's value given its name.
 	// Type must match strictly, no conversions are performed.
 	// @param	InVariable			The variable to get the value of
diff --git a/Engine/Plugins/Experimental/UAF/UAF/Source/UAF/Public/Variables/UAFInstanceVariableData.h b/Engine/Plugins/Experimental/UAF/UAF/Source/UAF/Public/Variables/UAFInstanceVariableData.h
--- a/Engine/Plugins/Experimental/UAF/UAF/Source/UAF/Public/Variables/UAFInstanceVariableData.h
+++ b/Engine/Plugins/Experimental/UAF/UAF/Source/UAF/Public/Variables/UAFInstanceVariableData.h
@@ -20,10 +20,12 @@ struct FRigVMExternalVariableRuntimeData;
 struct FAnimNextRigVMAssetStructData;
 struct FUAFInstanceVariableContainer;
 struct FUAFRigVMComponent;
+class IRigVMRuntimeAssetInterface;
 
 namespace UE::UAF
 {
-	struct FVariableOverrides;
+struct FControlRigTrait;
+struct FVariableOverrides;
 	struct FVariableOverridesCollection;
 	struct FUAFTrace;
 	struct FInstanceTaskContext;
@@ -48,12 +50,14 @@ private:
 	friend UE::UAF::FUAFTrace;
 	friend UE::UAF::FInstanceTaskContext;
 	friend FUAFRigVMComponent;
+	friend UE::UAF::FControlRigTrait;
 
 	// Initialize from a host's asset and its references
 	void Initialize(const FUAFAssetInstance& InHostInstance, const TSharedPtr<const UE::UAF::FVariableOverridesCollection>& InOverrides);
 
 	void AddVariablesContainerForAsset(const UUAFRigVMAsset* InAsset, const FUAFAssetInstance& InInstance, const TSharedPtr<const UE::UAF::FVariableOverridesCollection>& InOverrides);
 	void AddVariablesContainerForStruct(const UScriptStruct* InStruct, const FUAFAssetInstance& InInstance, const TSharedPtr<const UE::UAF::FVariableOverridesCollection>& InOverrides);
+	void AddVariablesContainerForRigVMAsset(const TScriptInterface<const IRigVMRuntimeAssetInterface>& InRigVMAsset, const FUAFAssetInstance& InInstance, const TSharedPtr<const UE::UAF::FVariableOverridesCollection>& InOverrides);
 
 	// Reset to the default state
 	void Reset();
@@ -104,6 +108,9 @@ private:
 	// Access the variable by index directly.
 	UAF_API void AccessVariablePropertyByIndex(int32 InIndex, TFunctionRef<void(const FProperty*, TArrayView<uint8>)> InFunction);
 
+	// Access a variable via a FAnimNextVariableReference
+	UAF_API EPropertyBagResult AccessVariableProperty(const FAnimNextVariableReference& InVariable, TFunctionRef<void(const FProperty*, TArrayView<uint8>)> InFunction);
+
 	// Variable containers owned by the containing instance
 	TArray<TSharedRef<FUAFInstanceVariableContainer>> OwnedVariableContainers;
 
diff --git a/Engine/Plugins/Experimental/UAF/UAF/Source/UAFEditor/Private/Common/GraphEditorSchemaActions.cpp b/Engine/Plugins/Experimental/UAF/UAF/Source/UAFEditor/Private/Common/GraphEditorSchemaActions.cpp
--- a/Engine/Plugins/Experimental/UAF/UAF/Source/UAFEditor/Private/Common/GraphEditorSchemaActions.cpp
+++ b/Engine/Plugins/Experimental/UAF/UAF/Source/UAFEditor/Private/Common/GraphEditorSchemaActions.cpp
@@ -159,7 +159,7 @@ FAnimNextSchemaAction_Variable::FAnimNextSchemaAction_Variable(FName InName, con
 {
 	SetVariableAccessorChoice(InVariableAccessorChoice, true);
 	
-	check(InSourceObject->IsA<UUAFSharedVariables>() || InSourceObject->IsA<UScriptStruct>());
+	check(InSourceObject->IsA<UUAFSharedVariables>() || InSourceObject->IsA<UScriptStruct>() || InSourceObject->Implements<URigVMRuntimeAssetInterface>());
 	SourceObjectPath = InSourceObject->GetPathName();
 
 	if(InType.IsObjectType())
diff --git a/Engine/Plugins/Experimental/UAF/UAF/Source/UAFEditor/Private/Common/Outliner/OutlinerHierarchy.cpp b/Engine/Plugins/Experimental/UAF/UAF/Source/UAFEditor/Private/Common/Outliner/OutlinerHierarchy.cpp
--- a/Engine/Plugins/Experimental/UAF/UAF/Source/UAFEditor/Private/Common/Outliner/OutlinerHierarchy.cpp
+++ b/Engine/Plugins/Experimental/UAF/UAF/Source/UAFEditor/Private/Common/Outliner/OutlinerHierarchy.cpp
@@ -117,6 +117,15 @@ void FOutlinerHierarchy::RetrieveSharedAssetPaths(const TSoftObjectPtr<UUAFRigVM
 				}
 			}
 			break;
+
+		case EAnimNextSharedVariablesType::RigVMAsset:
+			{
+				if (const TScriptInterface<const IRigVMRuntimeAssetInterface>& RigVMAsset = InSharedVariablesEntry->GetRigVMAsset())
+				{
+					OutSharedVariableSourcePaths.Add(RigVMAsset.GetObject());
+				}
+			}
+			break;
 		}
 
 		return true;
diff --git a/Engine/Plugins/Experimental/UAF/UAF/Source/UAFEditor/Private/Common/Outliner/OutlinerHierarchy.h b/Engine/Plugins/Experimental/UAF/UAF/Source/UAFEditor/Private/Common/Outliner/OutlinerHierarchy.h
--- a/Engine/Plugins/Experimental/UAF/UAF/Source/UAFEditor/Private/Common/Outliner/OutlinerHierarchy.h
+++ b/Engine/Plugins/Experimental/UAF/UAF/Source/UAFEditor/Private/Common/Outliner/OutlinerHierarchy.h
@@ -5,6 +5,7 @@
 #include "CoreMinimal.h"
 #include "ISceneOutlinerHierarchy.h"
 
+class URigVMRuntimeAssetInterface;
 class UUAFRigVMAssetEditorData;
 class UUAFRigVMAsset;
 
@@ -30,6 +31,7 @@ protected:
 
 	mutable TSet<TWeakObjectPtr<const UUAFRigVMAsset>> WeakProcessedAssets;
 	mutable TSet<TWeakObjectPtr<const UScriptStruct>> WeakProcessedStructs;
+	mutable TSet<TWeakObjectPtr<const UObject>> WeakProcessedRigVMAssets;
 };
 
 }
diff --git a/Engine/Plugins/Experimental/UAF/UAF/Source/UAFEditor/Private/Variables/Outliner/VariablesOutlinerEntryItem.cpp b/Engine/Plugins/Experimental/UAF/UAF/Source/UAFEditor/Private/Variables/Outliner/VariablesOutlinerEntryItem.cpp
--- a/Engine/Plugins/Experimental/UAF/UAF/Source/UAFEditor/Private/Variables/Outliner/VariablesOutlinerEntryItem.cpp
+++ b/Engine/Plugins/Experimental/UAF/UAF/Source/UAFEditor/Private/Variables/Outliner/VariablesOutlinerEntryItem.cpp
@@ -68,7 +68,12 @@ FString FVariablesOutlinerEntryItem::GetPackageName() const
 	}
 	else if (const FProperty* Property = PropertyPath.Get())
 	{
-		return Property->GetOwner<UScriptStruct>()->GetPackage()->GetName();
+		const UObject* Owner = Property->GetOwner<UObject>();
+		const UPackage* Package = Owner ? Owner->GetPackage() : nullptr;
+		if (Package != nullptr)
+		{
+			return Package->GetName();
+		}
 	}
 	
 	return ISceneOutlinerTreeItem::GetPackageName();
@@ -135,7 +140,7 @@ EAnimNextExportAccessSpecifier FVariablesOutlinerEntryItem::GetAccessSpecifier()
 	IUAFRigVMExportInterface* Export = Cast<IUAFRigVMExportInterface>(WeakEntry.Get());
 	if(Export == nullptr)
 	{
-		return EAnimNextExportAccessSpecifier::Private;
+		return EAnimNextExportAccessSpecifier::Public;
 	}
 
 	return Export->GetExportAccessSpecifier();
diff --git a/Engine/Plugins/Experimental/UAF/UAF/Source/UAFEditor/Private/Variables/Outliner/VariablesOutlinerHierarchy.cpp b/Engine/Plugins/Experimental/UAF/UAF/Source/UAFEditor/Private/Variables/Outliner/VariablesOutlinerHierarchy.cpp
--- a/Engine/Plugins/Experimental/UAF/UAF/Source/UAFEditor/Private/Variables/Outliner/VariablesOutlinerHierarchy.cpp
+++ b/Engine/Plugins/Experimental/UAF/UAF/Source/UAFEditor/Private/Variables/Outliner/VariablesOutlinerHierarchy.cpp
@@ -9,13 +9,16 @@
 #include "VariablesOutlinerStructSharedVariablesItem.h"
 #include "VariablesOutlinerMode.h"
 #include "VariablesOutlinerEntryItem.h"
+#include "VariablesOutlinerRigVMAssetSharedVariablesItem.h"
 #include "Common/Outliner/OutlinerAssetItem.h"
 #include "Common/Outliner/SCommonOutliner.h"
 #include "Entries/AnimNextSharedVariablesEntry.h"
 #include "Entries/AnimNextRigVMAssetEntry.h"
 #include "Entries/AnimNextVariableEntry.h"
+#include "Net/RepLayout.h"
 #include "Variables/SVariablesView.h"
 #include "Variables/AnimNextSharedVariables.h"
+#include "VerseVM/VVMNativeRef.h"
 
 namespace UE::UAF::Editor
 {
@@ -119,8 +122,6 @@ FSceneOutlinerTreeItemPtr FVariablesOutlinerHierarchy::FindOrCreateParentItem(co
 	{
 		TArray<const FSoftObjectPath> SharedVariableSourcePaths;
 		
-		const UScriptStruct* Struct = Cast<UScriptStruct>(OwnerPath.ResolveObject());
-
 		if (SoftAsset.IsValid())
 		{
 			RetrieveSharedAssetPaths(SoftAsset, SharedVariableSourcePaths);
@@ -129,10 +130,14 @@ FSceneOutlinerTreeItemPtr FVariablesOutlinerHierarchy::FindOrCreateParentItem(co
 			UE_CLOG(GOutputHierarchyDebugInformationCvar.GetValueOnAnyThread(), LogAnimation, Warning, TEXT("* FOutlinerAssetItem: %s %i"), *SoftAsset.GetAssetName(), 0);
 			return AssetItem;
 		}
-		else if (Struct)
+		else if (const UScriptStruct* Struct = Cast<UScriptStruct>(OwnerPath.ResolveObject()))
 		{
 			return Mode->CreateItemFor<FVariablesOutlinerStructSharedVariablesItem>(Struct, bCreate);
 		}
+		else if (TScriptInterface<IRigVMRuntimeAssetInterface> RigVMAssetInterface = OwnerPath.ResolveObject())
+		{
+			return Mode->CreateItemFor<FVariablesOutlinerRigVMAssetSharedVariablesItem>(RigVMAssetInterface, bCreate);
+		}
 	}
 
 	return nullptr;
@@ -183,6 +188,11 @@ void FVariablesOutlinerHierarchy::CreateItemsForAsset(const UUAFRigVMAssetEditor
 					for (TFieldIterator<FProperty> It(Struct); It; ++It)
 					{
 						const FProperty* Property = *It;
+						if ((Property->GetPropertyFlags() & CPF_NativeAccessSpecifierPublic) == 0)
+						{
+							continue;
+						}
+						
 						if (FSceneOutlinerTreeItemPtr Item = Mode->CreateItemFor<FVariablesOutlinerEntryItem>(FVariablesOutlinerEntryItem::FItemData{nullptr, Property, Index++}))
 						{
 							FVariablesOutlinerEntryItem* EntryItem = Item->CastTo<FVariablesOutlinerEntryItem>();
@@ -193,6 +203,42 @@ void FVariablesOutlinerHierarchy::CreateItemsForAsset(const UUAFRigVMAssetEditor
 				}
 			}
 			break;
+
+		case EAnimNextSharedVariablesType::RigVMAsset:
+			{
+				if (const IRigVMRuntimeAssetInterface* RigVMAsset = InSharedVariablesEntry->GetRigVMAsset().GetInterface())
+				{
+					const UObject* RigVMAssetObject = InSharedVariablesEntry->GetRigVMAsset().GetObject();
+					check(RigVMAssetObject);
+					if (WeakProcessedRigVMAssets.Contains(RigVMAssetObject))
+					{
+						return true;
+					}
+
+					WeakProcessedRigVMAssets.Add(RigVMAssetObject);
+					
+					if (FSceneOutlinerTreeItemPtr Item = Mode->CreateItemFor<FVariablesOutlinerRigVMAssetSharedVariablesItem>(InSharedVariablesEntry->GetRigVMAsset()))
+					{
+						OutItems.Add(Item);
+					}
+
+					TArray<FRigVMExternalVariable> ExternalVariables = const_cast<IRigVMRuntimeAssetInterface*>(RigVMAsset)->GetExternalVariables();
+					for (int32 i = 0; i < ExternalVariables.Num(); i++)
+					{
+						// We only add the public variables to the list
+						if (ExternalVariables[i].IsPublic())
+						{
+							if (FSceneOutlinerTreeItemPtr Item = Mode->CreateItemFor<FVariablesOutlinerEntryItem>(FVariablesOutlinerEntryItem::FItemData{nullptr, ExternalVariables[i].GetProperty(), i}))
+							{
+								FVariablesOutlinerEntryItem* EntryItem = Item->CastTo<FVariablesOutlinerEntryItem>();
+								EntryItem->WeakSharedVariablesEntry = InSharedVariablesEntry;
+								EntryItem->SetAccessSpecifier(EAnimNextExportAccessSpecifier::Public);
+								OutItems.Add(Item);
+							}
+						}
+					}
+				}
+			}
 		}		
 		
 		return true;
diff --git a/Engine/Plugins/Experimental/UAF/UAF/Source/UAFEditor/Private/Variables/Outliner/VariablesOutlinerMode.cpp b/Engine/Plugins/Experimental/UAF/UAF/Source/UAFEditor/Private/Variables/Outliner/VariablesOutlinerMode.cpp
--- a/Engine/Plugins/Experimental/UAF/UAF/Source/UAFEditor/Private/Variables/Outliner/VariablesOutlinerMode.cpp
+++ b/Engine/Plugins/Experimental/UAF/UAF/Source/UAFEditor/Private/Variables/Outliner/VariablesOutlinerMode.cpp
@@ -198,9 +198,9 @@ FVariablesOutlinerMode::FVariablesOutlinerMode(SCommonOutliner* InOutliner, cons
 		}
 		else if (const FProperty* Property = VariableItem->PropertyPath.Get())
 		{
-			if (const UScriptStruct* Struct = Property->GetOwner<UScriptStruct>())
+			if (const UObject* Owner = Property->GetOwner<UObject>())
 			{
-				TSharedPtr<FAnimNextSchemaAction_Variable> Action = MakeShared<FAnimNextSchemaAction_Variable>(Property->GetFName(), Struct, FAnimNextParamType::FromProperty(Property), FAnimNextSchemaAction_Variable::EVariableAccessorChoice::Deferred);
+				TSharedPtr<FAnimNextSchemaAction_Variable> Action = MakeShared<FAnimNextSchemaAction_Variable>(Property->GetFName(), Owner, FAnimNextParamType::FromProperty(Property), FAnimNextSchemaAction_Variable::EVariableAccessorChoice::Deferred);
 			
 				TSharedPtr<FVariableDragDropOp> Op = FVariableDragDropOp::New(Action, StaticCastSharedRef<FVariablesOutlinerEntryItem>(VariableItem->AsShared()));
 				return Op;
@@ -1817,6 +1817,79 @@ void FVariablesOutlinerMode::PopulateNewVariableToolMenuEntries(UToolMenu* InMen
 				}),
 				false,
 				FSlateIconFinder::FindIconForClass(UUserDefinedStruct::StaticClass())
+			);
+
+						// Note we refer to them as control rig shared variables even though technically any rig vm asset will display here
+			Section.AddSubMenu("AddRigVMAssetSharedVariables",LOCTEXT("RigVMAssetSharedVariableLabel", "Add Control Rig Shared Variables"), LOCTEXT("RigVMAssetSharedVariableTooltip", "Includes added Control Rig Asset SharedVariables"),
+		FNewMenuDelegate::CreateLambda([AssetsAndCategoriesToAddTo](FMenuBuilder& SubMenuBuilder)
+				{
+
+					TArray<FSoftObjectPath> ExistingSharedVariableRigVMPaths;
+					for (auto& Pair : AssetsAndCategoriesToAddTo)
+					{
+						if (UUAFRigVMAssetEditorData* EditorData = Pair.Key.Get())
+						{
+							EditorData->ForEachEntryOfType<UUAFSharedVariablesEntry>([&ExistingSharedVariableRigVMPaths](const UUAFSharedVariablesEntry* SharedVariablesEntry) -> bool
+							{
+								if (SharedVariablesEntry->GetType() == EAnimNextSharedVariablesType::RigVMAsset)
+								{
+									ExistingSharedVariableRigVMPaths.Add(SharedVariablesEntry->GetObjectPath());
+								}
+								return true;
+							});
+						}
+					}
+			
+					FAssetPickerConfig AssetPickerConfig;
+					AssetPickerConfig.Filter.bRecursiveClasses = true;
+					AssetPickerConfig.Filter.ClassPaths.Add(URigVMEditorAssetInterface::StaticClass()->GetClassPathName());
+					AssetPickerConfig.InitialAssetViewType = EAssetViewType::List;
+					AssetPickerConfig.OnAssetSelected = FOnAssetSelected::CreateLambda([AssetsAndCategoriesToAddTo](const FAssetData& InAssetData)
+					{
+						FSlateApplication::Get().DismissAllMenus();
+						if(const IRigVMEditorAssetInterface* RigVMEditorAsset = Cast<IRigVMEditorAssetInterface>(InAssetData.GetAsset()))
+						{
+							if (TScriptInterface<IRigVMRuntimeAssetInterface> RuntimeAsset = RigVMEditorAsset->GetRuntimeAssetInterface())
+							{
+								FScopedTransaction Transaction(LOCTEXT("AddSharedVariable", "Add Shared Variables"));
+								for (auto& Pair : AssetsAndCategoriesToAddTo)
+								{
+									if (UUAFRigVMAssetEditorData* EditorData = Pair.Key.Get())
+									{
+										EditorData->AddSharedVariablesRigVMAsset(RigVMEditorAsset->GetRuntimeAssetInterface().GetInterface());
+									}
+								}
+							}
+						}
+					});
+					AssetPickerConfig.OnShouldFilterAsset = FOnShouldFilterAsset::CreateLambda([ExistingSharedVariableRigVMPaths](const FAssetData& InAssetData)
+					{
+						FSoftObjectPath AssetPath = InAssetData.GetSoftObjectPath();
+						if (InAssetData.IsInstanceOf<UBlueprint>())
+						{
+							// We want to check against generated classes for blueprints
+							AssetPath.SetPath(AssetPath.GetAssetPath().ToString() + TEXT("_C"));
+						}
+							
+						// Filter out already referenced assets
+						return ExistingSharedVariableRigVMPaths.ContainsByPredicate([AssetPath](const FSoftObjectPath& ExistingRigVMAssetPath)
+						{
+							return ExistingRigVMAssetPath == AssetPath;
+						});
+					});
+
+					FContentBrowserModule& ContentBrowserModule = FModuleManager::Get().LoadModuleChecked<FContentBrowserModule>(TEXT("ContentBrowser"));
+					TSharedPtr<SWidget> Widget = SNew(SBox)
+					.WidthOverride(300.0f)
+					.HeightOverride(400.0f)
+					[
+						ContentBrowserModule.Get().CreateAssetPicker(AssetPickerConfig)
+					];
+
+					SubMenuBuilder.AddWidget(Widget.ToSharedRef(), FText::GetEmpty(), true);
+				}),
+				false,
+				FSlateIcon(FUAFStyle::Get().GetStyleSetName(), "ClassIcon.AnimNextSharedVariables")
 			);
 		}
 
diff --git a/Engine/Plugins/Experimental/UAF/UAF/Source/UAFEditor/Private/Variables/Outliner/VariablesOutlinerRigVMAssetSharedVariablesItem.cpp b/Engine/Plugins/Experimental/UAF/UAF/Source/UAFEditor/Private/Variables/Outliner/VariablesOutlinerRigVMAssetSharedVariablesItem.cpp
new file mode 100644
--- /dev/null
+++ b/Engine/Plugins/Experimental/UAF/UAF/Source/UAFEditor/Private/Variables/Outliner/VariablesOutlinerRigVMAssetSharedVariablesItem.cpp
@@ -0,0 +1,151 @@
+// Copyright Epic Games, Inc. All Rights Reserved.
+
+#include "VariablesOutlinerRigVMAssetSharedVariablesItem.h"
+
+#include "ISceneOutliner.h"
+#include "Styling/SlateColor.h"
+#include "Misc/Attribute.h"
+#include "Misc/Optional.h"
+#include "SceneOutlinerStandaloneTypes.h"
+#include "AnimNextEditorModule.h"
+#include "AnimNextRigVMAssetEditorData.h"
+#include "Common/Outliner/OutlinerDragDropOps.h"
+#include "Entries/AnimNextSharedVariablesEntry.h"
+#include "Styling/SlateIconFinder.h"
+#include "Styling/StarshipCoreStyle.h"
+#include "UObject/Package.h"
+#include "Widgets/SCompoundWidget.h"
+#include "Widgets/DeclarativeSyntaxSupport.h"
+#include "Widgets/Images/SImage.h"
+#include "Widgets/Text/SInlineEditableTextBlock.h"
+#include "Widgets/Text/STextBlock.h"
+#include "Widgets/Views/STableRow.h"
+#include "Widgets/Views/SListView.h"
+
+#define LOCTEXT_NAMESPACE "VariablesOutlinerTreeItem"
+
+namespace UE::UAF::Editor
+{
+
+const FSceneOutlinerTreeItemType FVariablesOutlinerRigVMAssetSharedVariablesItem::Type(&FOutlinerItem::Type);
+
+class SVariablesOutlinerRigVMSharedVariablesLabel : FSceneOutlinerCommonLabelData, public SCompoundWidget
+{
+	SLATE_BEGIN_ARGS(SVariablesOutlinerRigVMSharedVariablesLabel) {}
+	SLATE_END_ARGS()
+
+	void Construct(const FArguments& InArgs, FVariablesOutlinerRigVMAssetSharedVariablesItem& InTreeItem, ISceneOutliner& SceneOutliner, const STableRow<FSceneOutlinerTreeItemPtr>& InRow)
+	{
+		WeakSceneOutliner = StaticCastSharedRef<ISceneOutliner>(SceneOutliner.AsShared());
+		TreeItem = StaticCastSharedRef<FVariablesOutlinerRigVMAssetSharedVariablesItem>(InTreeItem.AsShared());
+
+		FText AssetName;
+		if(const IRigVMRuntimeAssetInterface* RigVMAssetInterface = InTreeItem.RigVMAsset.GetInterface())
+		{
+			AssetName = FText::FromString(RigVMAssetInterface->GetEditorOnlyData()->GetPathName());
+		}
+		else
+		{
+			AssetName = LOCTEXT("UnknownAssetName", "Unknown Asset");
+		}
+		
+		ChildSlot
+		[
+			SNew(SHorizontalBox)
+			.ToolTipText(FText::Format(LOCTEXT("ImportedVariablesFormat", "Shared variables from '{0}'"), AssetName))
+			+SHorizontalBox::Slot()
+			.AutoWidth()
+			.VAlign(VAlign_Center)
+			.Padding(2.0f, 2.0f)
+			[
+				SNew(SImage)
+				.Image(FSlateIconFinder::FindIconBrushForClass(UUserDefinedStruct::StaticClass()))
+			]
+			+SHorizontalBox::Slot()
+			.AutoWidth()
+			.VAlign(VAlign_Center)
+			.Padding(2.0f, 2.0f)
+			[
+				SNew(STextBlock)
+				.Font(FStyleFonts::Get().NormalBold)
+				.Text(this, &SVariablesOutlinerRigVMSharedVariablesLabel::GetDisplayText)
+				.HighlightText(SceneOutliner.GetFilterHighlightText())
+				.ColorAndOpacity(this, &SVariablesOutlinerRigVMSharedVariablesLabel::GetForegroundColor)
+			]
+		];
+	}
+	
+	FText GetDisplayText() const
+	{
+		if (const TSharedPtr<FVariablesOutlinerRigVMAssetSharedVariablesItem> Item = TreeItem.Pin())
+		{
+			return FText::FromString(Item->GetDisplayString());
+		}
+		return FText();
+	}
+
+	virtual FSlateColor GetForegroundColor() const override
+	{
+		TOptional<FLinearColor> BaseColor;
+		if (TreeItem.IsValid())
+		{
+			BaseColor = FSceneOutlinerCommonLabelData::GetForegroundColor(*TreeItem.Pin());
+		}
+		return BaseColor.IsSet() ? BaseColor.GetValue() : FSlateColor::UseForeground();
+	}
+
+
+	TWeakPtr<FVariablesOutlinerRigVMAssetSharedVariablesItem> TreeItem;
+	TSharedPtr<SInlineEditableTextBlock> TextBlock;
+};
+
+FVariablesOutlinerRigVMAssetSharedVariablesItem::FVariablesOutlinerRigVMAssetSharedVariablesItem(const TScriptInterface<const IRigVMRuntimeAssetInterface>& InRigVMAsset)
+: FOutlinerItem(FVariablesOutlinerRigVMAssetSharedVariablesItem::Type, {nullptr, 0})
+	, RigVMAsset(InRigVMAsset)
+{
+}
+
+bool FVariablesOutlinerRigVMAssetSharedVariablesItem::IsValid() const
+{
+	return RigVMAsset != nullptr;
+}
+
+FSceneOutlinerTreeItemID FVariablesOutlinerRigVMAssetSharedVariablesItem::GetID() const
+{
+	if (RigVMAsset)
+	{
+		const FSoftObjectPath Path = RigVMAsset.GetObject()->GetPathName();
+		return GetTypeHash(Path);
+	}
+	
+	return GetTypeHash(RigVMAsset);
+}
+
+FString FVariablesOutlinerRigVMAssetSharedVariablesItem::GetDisplayString() const
+{
+	if(RigVMAsset == nullptr)
+	{
+		return FString();
+	}
+
+	return RigVMAsset.GetObject()->GetName();
+}
+
+TSharedRef<SWidget> FVariablesOutlinerRigVMAssetSharedVariablesItem::GenerateLabelWidget(ISceneOutliner& Outliner, const STableRow<FSceneOutlinerTreeItemPtr>& InRow)
+{
+	return SNew(SVariablesOutlinerRigVMSharedVariablesLabel, *this, Outliner, InRow);
+}
+
+FString FVariablesOutlinerRigVMAssetSharedVariablesItem::GetPackageName() const
+{
+	if(RigVMAsset == nullptr)
+	{
+		return ISceneOutlinerTreeItem::GetPackageName();
+	}
+
+	return RigVMAsset.GetObject()->GetPackage()->GetName();
+}
+
+}
+
+#undef LOCTEXT_NAMESPACE // "VariablesOutlinerTreeItem"
diff --git a/Engine/Plugins/Experimental/UAF/UAF/Source/UAFEditor/Private/Variables/Outliner/VariablesOutlinerRigVMAssetSharedVariablesItem.h b/Engine/Plugins/Experimental/UAF/UAF/Source/UAFEditor/Private/Variables/Outliner/VariablesOutlinerRigVMAssetSharedVariablesItem.h
new file mode 100644
--- /dev/null
+++ b/Engine/Plugins/Experimental/UAF/UAF/Source/UAFEditor/Private/Variables/Outliner/VariablesOutlinerRigVMAssetSharedVariablesItem.h
@@ -0,0 +1,34 @@
+// Copyright Epic Games, Inc. All Rights Reserved.
+
+#pragma once
+
+#include "ISceneOutlinerTreeItem.h"
+#include "Common/Outliner/OutlinerEntryItem.h"
+
+class IRigVMRuntimeAssetInterface;
+class UUAFSharedVariablesEntry;
+
+namespace UE::UAF::Editor
+{
+
+class IWorkspaceOutlinerItemDetails;
+
+struct FVariablesOutlinerRigVMAssetSharedVariablesItem : FOutlinerItem
+{
+	static const FSceneOutlinerTreeItemType Type;
+
+	FVariablesOutlinerRigVMAssetSharedVariablesItem(const TScriptInterface<const IRigVMRuntimeAssetInterface>& InRigVMAsset);
+
+	// Begin ISceneOutlinerTreeItem overrides
+	virtual bool IsValid() const override;
+	virtual FSceneOutlinerTreeItemID GetID() const override;
+	virtual FString GetDisplayString() const override;
+	virtual bool CanInteract() const override { return true; }
+	virtual TSharedRef<SWidget> GenerateLabelWidget(ISceneOutliner& Outliner, const STableRow<FSceneOutlinerTreeItemPtr>& InRow) override;
+	virtual FString GetPackageName() const override;
+	// End ISceneOutlinerTreeItem overrides
+
+	TScriptInterface<const IRigVMRuntimeAssetInterface> RigVMAsset;
+};
+
+}
diff --git a/Engine/Plugins/Experimental/UAF/UAF/Source/UAFEditor/Private/Variables/SVariablePicker.cpp b/Engine/Plugins/Experimental/UAF/UAF/Source/UAFEditor/Private/Variables/SVariablePicker.cpp
--- a/Engine/Plugins/Experimental/UAF/UAF/Source/UAFEditor/Private/Variables/SVariablePicker.cpp
+++ b/Engine/Plugins/Experimental/UAF/UAF/Source/UAFEditor/Private/Variables/SVariablePicker.cpp
@@ -89,7 +89,7 @@ void SVariablePicker::RefreshEntries()
 
 	// Add variables exposed on assets
 	TMap<FAssetData, FAnimNextAssetRegistryExports> Exports;
-	if(UncookedOnly::FUtils::GetExportedVariablesFromAssetRegistry(Exports))
+	if(UncookedOnly::FUtils::GetExportsOfTypeFromAssetRegistry<FAnimNextVariableDeclarationData>(Exports))
 	{
 		for(const TPair<FAssetData, FAnimNextAssetRegistryExports>& ExportPair : Exports)
 		{
@@ -141,7 +141,19 @@ bool SVariablePicker::GetFieldInfo(UE::PropertyViewer::SPropertyViewer::FHandle
 		const FContainerInfo& ContainerInfo = CachedContainers[*ContainerIndexPtr];
 		if(FProperty* Property = InField.Get<FProperty>())
 		{
-			OutVariableReference = FAnimNextSoftVariableReference::FromName(Property->GetFName(), ContainerInfo.AssetData.ToSoftObjectPath());
+			FSoftObjectPath ObjectPath = ContainerInfo.AssetData.ToSoftObjectPath();
+			UClass* AssetClass = ContainerInfo.AssetData.GetClass(EResolveClass::Yes);
+			
+			if (AssetClass != nullptr && AssetClass->IsChildOf(UBlueprint::StaticClass()))
+			{
+				// Shortcut to avoid loading asset
+				// We know that the generated class of a blueprint will be appended with _C
+				// We want the path to be to the runtime asset (i.e. the generated class) not the blueprint
+				FString GeneratedClassName = ContainerInfo.AssetData.GetObjectPathString() + TEXT("_C");
+				ObjectPath = GeneratedClassName;
+			}
+
+			OutVariableReference = FAnimNextSoftVariableReference::FromName(Property->GetFName(), ObjectPath);
 			OutType = FAnimNextParamType::FromProperty(Property);
 
 			return true;
diff --git a/Engine/Plugins/Experimental/UAF/UAF/Source/UAFUncookedOnly/Internal/AnimNextRigVMAssetEditorData.h b/Engine/Plugins/Experimental/UAF/UAF/Source/UAFUncookedOnly/Internal/AnimNextRigVMAssetEditorData.h
--- a/Engine/Plugins/Experimental/UAF/UAF/Source/UAFUncookedOnly/Internal/AnimNextRigVMAssetEditorData.h
+++ b/Engine/Plugins/Experimental/UAF/UAF/Source/UAFUncookedOnly/Internal/AnimNextRigVMAssetEditorData.h
@@ -190,6 +190,9 @@ public:
 	/** Adds shared variables struct to this asset */
 	UE_API UUAFSharedVariablesEntry* AddSharedVariablesStruct(const UScriptStruct* InStruct, bool bSetupUndoRedo = true, bool bPrintPythonCommand = true);
 
+	/** Adds shared variables from IRigVMRuntimeAssetInterfaces to this asset */
+	UE_API UUAFSharedVariablesEntry* AddSharedVariablesRigVMAsset(const IRigVMRuntimeAssetInterface* InRigVMAsset, bool bSetupUndoRedo = true, bool bPrintPythonCommand = true);
+
 	/** Adds a function to this asset */
 	UE_API URigVMLibraryNode* AddFunction(FName InFunctionName, bool bInMutable, bool bInSetupUndoRedo = true, bool bPrintPythonCommand = true);
 
diff --git a/Engine/Plugins/Experimental/UAF/UAF/Source/UAFUncookedOnly/Internal/AnimNextSharedVariableNode.h b/Engine/Plugins/Experimental/UAF/UAF/Source/UAFUncookedOnly/Internal/AnimNextSharedVariableNode.h
--- a/Engine/Plugins/Experimental/UAF/UAF/Source/UAFUncookedOnly/Internal/AnimNextSharedVariableNode.h
+++ b/Engine/Plugins/Experimental/UAF/UAF/Source/UAFUncookedOnly/Internal/AnimNextSharedVariableNode.h
@@ -42,6 +42,9 @@ private:
 
 	UPROPERTY()
 	TObjectPtr<const UScriptStruct> Struct;
+
+	UPROPERTY()
+	TScriptInterface<const IRigVMRuntimeAssetInterface> RigVMAsset;
 	
 	UPROPERTY()
 	EAnimNextSharedVariablesType Type;
diff --git a/Engine/Plugins/Experimental/UAF/UAF/Source/UAFUncookedOnly/Internal/Entries/AnimNextSharedVariablesEntry.h b/Engine/Plugins/Experimental/UAF/UAF/Source/UAFUncookedOnly/Internal/Entries/AnimNextSharedVariablesEntry.h
--- a/Engine/Plugins/Experimental/UAF/UAF/Source/UAFUncookedOnly/Internal/Entries/AnimNextSharedVariablesEntry.h
+++ b/Engine/Plugins/Experimental/UAF/UAF/Source/UAFUncookedOnly/Internal/Entries/AnimNextSharedVariablesEntry.h
@@ -7,6 +7,7 @@
 #include "Param/ParamType.h"
 #include "AnimNextSharedVariablesEntry.generated.h"
 
+class IRigVMRuntimeAssetInterface;
 class UUAFSharedVariables;
 class UUAFRigVMAssetEditorData;
 class UAssetDefinition_AnimNextSharedVariablesEntry;
@@ -35,8 +36,9 @@ namespace UE::UAF::UncookedOnly
 UENUM()
 enum class EAnimNextSharedVariablesType : uint8
 {
-	Asset,
-	Struct
+	Asset,		// reference to shared variable defined in UUAFSharedVariables
+	Struct,		// Reference to a property of a native struct
+	RigVMAsset	// Reference to an external variable defined in a RigVMAsset
 };
 
 UCLASS(MinimalAPI, Category = "Shared Variables", DisplayName = "Shared Variables")
@@ -80,6 +82,9 @@ class UUAFSharedVariablesEntry : public UUAFRigVMAssetEntry
 	// Set this entry to use an asset (rather than a struct)
 	UAFUNCOOKEDONLY_API void SetAsset(const UUAFSharedVariables* InAsset, bool bSetupUndoRedo = true);
 
+	// Set this entry to use an asset (rather than a struct)
+	UAFUNCOOKEDONLY_API void SetRigVMAsset(const IRigVMRuntimeAssetInterface* InRigVMAsset, bool bSetupUndoRedo = true);
+
 	// Get the asset that this entry represents. If this entry uses a struct, this call will return nullptr
 	UAFUNCOOKEDONLY_API const UUAFSharedVariables* GetAsset() const;
 
@@ -92,6 +97,8 @@ class UUAFSharedVariablesEntry : public UUAFRigVMAssetEntry
 	// Get the struct that this entry represents. If this entry uses an asset, this call will return nullptr
 	UAFUNCOOKEDONLY_API const UScriptStruct* GetStruct() const;
 
+	UAFUNCOOKEDONLY_API const TScriptInterface<const IRigVMRuntimeAssetInterface>& GetRigVMAsset() const { return RigVMAsset; }
+
 	// Recompiles this asset when the linked asset is modified
 	void HandleAssetModified(UUAFRigVMAssetEditorData* InEditorData, EAnimNextEditorDataNotifType InType, UObject* InSubject);
 
@@ -104,6 +111,10 @@ class UUAFSharedVariablesEntry : public UUAFRigVMAssetEntry
 	UPROPERTY(VisibleAnywhere, Category = "Asset")
 	TObjectPtr<const UUAFSharedVariables> Asset;
 
+	/** The asset whose variables we share */
+	UPROPERTY(VisibleAnywhere, Category = "Asset")
+	TScriptInterface<const IRigVMRuntimeAssetInterface> RigVMAsset;
+
 	/** Soft reference to the asset/struct for error reporting */
 	UPROPERTY()
 	FSoftObjectPath ObjectPath;
diff --git a/Engine/Plugins/Experimental/UAF/UAF/Source/UAFUncookedOnly/Internal/UncookedOnlyUtils.h b/Engine/Plugins/Experimental/UAF/UAF/Source/UAFUncookedOnly/Internal/UncookedOnlyUtils.h
--- a/Engine/Plugins/Experimental/UAF/UAF/Source/UAFUncookedOnly/Internal/UncookedOnlyUtils.h
+++ b/Engine/Plugins/Experimental/UAF/UAF/Source/UAFUncookedOnly/Internal/UncookedOnlyUtils.h
@@ -254,7 +254,7 @@ private:
 	static void GetStructVariableExports(const UScriptStruct* Struct, TSet<FAnimNextExport>& OutExports);
 	static void GetSubGraphVariableExportsRecursive(const URigVMEdGraph* RigVMEdGraph, TSet<FAnimNextExport>& OutExports);
 	static void ProcessPinAssetReferences(const URigVMPin* InPin, FWorkspaceOutlinerItemExports& OutExports, const FWorkspaceOutlinerItemExport& RootExport, int32 ParentExportIndex);
-	static void RetrieveReferencedSharedVariableSources(const UUAFRigVMAsset* InAsset, TArray<const UUAFRigVMAsset*>& ReferencedVariableAssets, TArray<const UScriptStruct*>& ReferencedVariableStructs);
+	static void RetrieveReferencedSharedVariableSources(const UUAFRigVMAsset* InAsset, TSet<const UUAFRigVMAsset*>& ReferencedVariableAssets, TSet<const UScriptStruct*>& ReferencedVariableStructs, TSet<TScriptInterface<const IRigVMRuntimeAssetInterface>>& ReferencedVariableRigVMAssets);
 };
 
 }
diff --git a/Engine/Plugins/Experimental/UAF/UAF/Source/UAFUncookedOnly/Private/AnimNextControllerBase.cpp b/Engine/Plugins/Experimental/UAF/UAF/Source/UAFUncookedOnly/Private/AnimNextControllerBase.cpp
--- a/Engine/Plugins/Experimental/UAF/UAF/Source/UAFUncookedOnly/Private/AnimNextControllerBase.cpp
+++ b/Engine/Plugins/Experimental/UAF/UAF/Source/UAFUncookedOnly/Private/AnimNextControllerBase.cpp
@@ -75,6 +75,11 @@ UUAFSharedVariableNode* UAnimNextControllerBase::AddSharedVariableNode(const FSt
 		
 		VariableNode->UpdateCachedVariableGUID();
 	}
+	else if (const IRigVMRuntimeAssetInterface* RigVMAsset = Cast<IRigVMRuntimeAssetInterface>(SourceObject))
+	{
+		VariableNode->Type = EAnimNextSharedVariablesType::RigVMAsset;
+		VariableNode->RigVMAsset = SourceObject;
+	}
 	else
 	{
 		return nullptr;
diff --git a/Engine/Plugins/Experimental/UAF/UAF/Source/UAFUncookedOnly/Private/AnimNextRigVMAssetEditorData.cpp b/Engine/Plugins/Experimental/UAF/UAF/Source/UAFUncookedOnly/Private/AnimNextRigVMAssetEditorData.cpp
--- a/Engine/Plugins/Experimental/UAF/UAF/Source/UAFUncookedOnly/Private/AnimNextRigVMAssetEditorData.cpp
+++ b/Engine/Plugins/Experimental/UAF/UAF/Source/UAFUncookedOnly/Private/AnimNextRigVMAssetEditorData.cpp
@@ -2556,34 +2556,6 @@ UUAFSharedVariablesEntry* UUAFRigVMAssetEditorData::AddSharedVariables(const UUA
 		return nullptr;
 	}
 
-	auto CheckForPublicMembers = [](UUAFSharedVariables_EditorData* InEditorData, auto& InCheckForPublicMembers)
-	{
-		for(UUAFRigVMAssetEntry* Entry : InEditorData->Entries)
-		{
-			if(UAnimNextVariableEntry* VariableEntry = Cast<UAnimNextVariableEntry>(Entry))
-			{
-				if(VariableEntry->GetExportAccessSpecifier() == EAnimNextExportAccessSpecifier::Public)
-				{
-					return true;
-				}
-			}
-			else if(UUAFSharedVariablesEntry* SharedVariablesEntry = Cast<UUAFSharedVariablesEntry>(Entry))
-			{
-				// SharedVariables could be nullptr as the entry is a Struct based source, or the underlying asset (UUAFSharedVariablesEntry::ObjectPath) has been deleted
-				if (const UUAFSharedVariables* SharedVariables = SharedVariablesEntry->GetAsset())
-				{
-					UUAFSharedVariables_EditorData* EditorData = UE::UAF::UncookedOnly::FUtils::GetEditorData<UUAFSharedVariables_EditorData>(SharedVariables);
-					if(InCheckForPublicMembers(EditorData, InCheckForPublicMembers))
-					{
-						return true;
-					}
-				}
-			}
-		}
-
-		return false;
-	};
-
 	// Check for duplicate entry
 	auto DuplicatePredicate = [InSharedVariables](const UUAFRigVMAssetEntry* InEntry)
 	{
@@ -2698,6 +2670,69 @@ UUAFSharedVariablesEntry* UUAFRigVMAssetEditorData::AddSharedVariablesStruct(con
 	return NewEntry;
 }
 
+UUAFSharedVariablesEntry* UUAFRigVMAssetEditorData::AddSharedVariablesRigVMAsset(const IRigVMRuntimeAssetInterface* InRigVMAsset, bool bSetupUndoRedo,
+	bool bPrintPythonCommand)
+{
+		using namespace UE::UAF::UncookedOnly;
+
+	if(InRigVMAsset == nullptr)
+	{
+		ReportError(TEXT("UUAFRigVMAssetEditorData::AddSharedVariablesRigVMAsset: Invalid asset supplied."));
+		return nullptr;
+	}
+
+	if(!GetEntryClasses().Contains(UUAFSharedVariablesEntry::StaticClass()) || !CanAddNewEntry(UUAFSharedVariablesEntry::StaticClass()))
+	{
+		ReportError(TEXT("UUAFRigVMAssetEditorData::AddSharedVariables: Cannot add a shared variables to this asset - entry is not allowed."));
+		return nullptr;
+	}
+
+	// Check for duplicate entry
+	auto DuplicatePredicate = [InRigVMAsset](const UUAFRigVMAssetEntry* InEntry)
+	{
+		if(const UUAFSharedVariablesEntry* SharedVariablesEntry = Cast<UUAFSharedVariablesEntry>(InEntry))
+		{
+			return SharedVariablesEntry->GetRigVMAsset() == InRigVMAsset;
+		}
+		return false;
+	};
+
+	if(Entries.ContainsByPredicate(DuplicatePredicate))
+	{
+		ReportError(TEXT("UUAFRigVMAssetEditorData::AddSharedVariables: Shared variables already referenced."));
+		return nullptr;
+	}
+
+	// We are about to perform multiple changes, add a scope now to avoid multiple compile broadcasts per change
+	UUAFRigVMAsset* Asset = FUtils::GetAsset<UUAFRigVMAsset>(this);
+	FCompilationScope CompileScope(LOCTEXT("ModifiedAssetsAddSharedVariables", "Add Shared Variables"), Asset);
+
+	UUAFSharedVariablesEntry* NewEntry = CreateNewSubEntry<UUAFSharedVariablesEntry>(this);
+	NewEntry->SetRigVMAsset(InRigVMAsset);
+	NewEntry->Initialize(this);
+
+	if(bSetupUndoRedo)
+	{
+		NewEntry->Modify();
+		Modify();
+	}
+
+	AddEntryInternal(NewEntry);
+
+	CustomizeNewAssetEntry(NewEntry);
+
+	BroadcastModified(EAnimNextEditorDataNotifType::EntryAdded, NewEntry);
+
+	if (bPrintPythonCommand)
+	{
+		RigVMPythonUtils::Print(GetName(),
+				FString::Printf(TEXT("asset.add_shared_variables(unreal.find_object(outer=None, name='%s'))"),
+				 *(Cast<UObject>(InRigVMAsset)->GetPathName())));
+	}
+
+	return NewEntry;
+}
+
 URigVMLibraryNode* UUAFRigVMAssetLibrary::AddFunction(UUAFRigVMAsset* InAsset, FName InFunctionName, bool bInMutable, bool bSetupUndoRedo, bool bPrintPythonCommand)
 {
 	return UE::UAF::UncookedOnly::FUtils::GetEditorData(InAsset)->AddFunction(InFunctionName, bInMutable, bSetupUndoRedo, bPrintPythonCommand);
@@ -2823,6 +2858,34 @@ void UUAFRigVMAssetEditorData::GetAllVariables(TArray<FVariableInfo>& OutVariabl
 					}
 				}
 				break;
+			case EAnimNextSharedVariablesType::RigVMAsset:
+				{
+					if(const IRigVMRuntimeAssetInterface* RigVMAsset = SharedVariablesEntry->GetRigVMAsset().GetInterface())
+					{
+						TArray<FRigVMExternalVariable> RigVMVariables = const_cast<IRigVMRuntimeAssetInterface*>(RigVMAsset)->GetExternalVariables();
+						for (const FRigVMExternalVariable& Variable : RigVMVariables)
+						{
+							const FProperty* Property = Variable.GetProperty();
+							if (InAccess == EVariableAccessFilter::All || Property->HasAnyPropertyFlags(CPF_NativeAccessSpecifierPublic))
+							{
+								FGuid Guid = FGuid::NewDeterministicGuid(Property->GetPathName());
+								OutVariables.Add(
+									{
+										Property->GetFName(),
+										FAnimNextParamType::FromProperty(Property),
+										SharedVariablesEntry->GetRigVMAsset().GetObject(),
+										EAnimNextExportAccessSpecifier::Public,
+										Property,
+										TConstArrayView<uint8>(),
+										Guid
+									});
+							}
+						}
+					}
+				}
+				break;
+			default:
+				checkNoEntry();
 			}
 		}
 	}
@@ -3048,7 +3111,7 @@ FInstancedPropertyBag UUAFRigVMAssetEditorData::GenerateCombinedPropertyBag(cons
 		return bPropertyWithSameNameAlreadyExists;
 	};
 
-	auto AddInstancedPropertyBag = [this, &PropertyDescriptions, &PropertyDefaultValues, &InCompileContext, Asset, &HasOverlappingPropertyName](const FInstancedPropertyBag& InPropertyBag, const UUAFRigVMAsset* SourceAsset, bool bValidateProperties = false)
+	auto AddInstancedPropertyBag = [this, &PropertyDescriptions, &PropertyDefaultValues, &InCompileContext, Asset, &HasOverlappingPropertyName](const FInstancedPropertyBag& InPropertyBag, const UObject* SourceAsset, bool bValidateProperties = false)
 	{
 		if(const UPropertyBag* PropertyBag = InPropertyBag.GetPropertyBagStruct())
 		{
@@ -3091,6 +3154,28 @@ FInstancedPropertyBag UUAFRigVMAssetEditorData::GenerateCombinedPropertyBag(cons
 		AddInstancedPropertyBag(ReferencedVariableAssetBag, ReferencedVariableAsset, true);
 	}
 
+	// add shared variables from Rig VM assets
+	for (TScriptInterface ReferencedRigVMAsset : Asset->ReferencedVariableRigVMAssets)
+	{
+		if (ReferencedRigVMAsset == nullptr)
+		{
+			continue;
+		}
+
+		TArray<FRigVMExternalVariable> ExternalVariables = const_cast<IRigVMRuntimeAssetInterface*>(ReferencedRigVMAsset.GetInterface())->GetExternalVariables();
+		for (FRigVMExternalVariable& ExternalVariable : ExternalVariables)
+		{
+			FPropertyBagPropertyDesc PropertyDesc(ExternalVariable.GetName(), ExternalVariable.GetProperty());
+
+			if (!HasOverlappingPropertyName(ExternalVariable.GetName(), ReferencedRigVMAsset.GetObject()))
+			{
+				PropertyDescriptions.Add(PropertyDesc);
+				TConstArrayView<uint8> DefaultValueMemory(ExternalVariable.GetMemory(), ExternalVariable.GetProperty()->GetSize());
+				PropertyDefaultValues.Add(DefaultValueMemory);	
+			}
+		}
+	}
+
 	// Required to keep default struct property values from being freed
 	TArray<FInstancedStruct> InstancedReferenceStructs;
 	// Next add native structs
diff --git a/Engine/Plugins/Experimental/UAF/UAF/Source/UAFUncookedOnly/Private/AnimNextSharedVariableNode.cpp b/Engine/Plugins/Experimental/UAF/UAF/Source/UAFUncookedOnly/Private/AnimNextSharedVariableNode.cpp
--- a/Engine/Plugins/Experimental/UAF/UAF/Source/UAFUncookedOnly/Private/AnimNextSharedVariableNode.cpp
+++ b/Engine/Plugins/Experimental/UAF/UAF/Source/UAFUncookedOnly/Private/AnimNextSharedVariableNode.cpp
@@ -159,6 +159,54 @@ void UUAFSharedVariableNode::PostLoad()
 					}
 					break;
 				}
+			case EAnimNextSharedVariablesType::RigVMAsset:
+				{
+					if (const IRigVMRuntimeAssetInterface* RigVMAssetInterface = RigVMAsset.GetInterface())
+					{
+						check(RigVMAsset.GetObject() != nullptr);
+						const FAnimNextParamType ExpectedVariableType = FAnimNextParamType::FromRigVMTemplateArgument(FRigVMTemplateArgumentType(FName(*GetCPPType()), GetCPPTypeObject()));	
+						FName ExpectedStructVariableName = ExpectedVariableName;
+						const FProperty* Property = RigVMAssetInterface->FindGeneratedPropertyByName(ExpectedStructVariableName);
+
+						if (Property)
+						{
+							const FAnimNextParamType PropertyParamType = FAnimNextParamType::FromProperty(Property);
+							if (ExpectedVariableType != PropertyParamType || ExpectedStructVariableName != ExpectedVariableName)
+							{
+								if (UAnimNextControllerBase* AnimNextController = Cast<UAnimNextControllerBase>(OuterHost->GetController(OuterGraph)))
+								{
+									// Update with new variable name - and type if required, populating orphan pins of non-compatible data types
+									AnimNextController->RefreshSharedVariableNode(GetFName(), RigVMAsset.GetObject()->GetPathName(), ExpectedStructVariableName, *Property->GetCPPType(), const_cast<UObject*>(PropertyParamType.GetValueTypeObject()), false, true, false );
+								}
+							}
+						}
+						else
+						{
+							// node cannot be resolved
+							UE_LOG(LogAnimation, Warning, TEXT("Unable to find struct FProperty with expected name: %s, type: %s in %s, SharedVariableNode will be invalidated"), *ExpectedVariableName.ToString(), *ExpectedVariableType.ToString(), *RigVMAsset.GetObject()->GetPathName());
+
+							if (UAnimNextControllerBase* AnimNextController = Cast<UAnimNextControllerBase>(OuterHost->GetController(OuterGraph)))
+							{
+								AnimNextController->RefreshSharedVariableNode(GetFName(), RigVMAsset.GetObject()->GetPathName(), ExpectedStructVariableName, *RigVMTypeUtils::GetWildCardCPPTypeName().ToString(), RigVMTypeUtils::GetWildCardCPPTypeObject(), false, true, false );
+							}
+						}
+					}
+					else
+					{
+						UE_LOG(LogAnimation, Warning, TEXT("Unable to find RigVMAsset for %s when trying to resolve SharedVariable node source"), *GetPathName());
+						
+						if (UAnimNextControllerBase* AnimNextController = Cast<UAnimNextControllerBase>(OuterHost->GetController(OuterGraph)))
+						{
+							AnimNextController->RefreshSharedVariableNode(GetFName(), TEXT(""), ExpectedVariableName, *RigVMTypeUtils::GetWildCardCPPTypeName().ToString(), RigVMTypeUtils::GetWildCardCPPTypeObject(), false, true, false );
+						}
+					}
+				}
+				break;
+			default:
+				{
+					checkNoEntry();
+				}
+				break;
 			}
 		}
 	}
@@ -216,6 +264,29 @@ void UUAFSharedVariableNode::UpdateCachedVariableGUID()
 			}
 		}
 		break;
+	case EAnimNextSharedVariablesType::RigVMAsset:
+		{
+			if (RigVMAsset)
+			{
+				if (const FProperty* Property = RigVMAsset->FindGeneratedPropertyByName(GetVariableName()))
+				{
+					CachedGuid = UE::UAF::UncookedOnly::FUtils::GenerateScriptStructPropertyGUID(Property);
+				}
+				else
+				{
+					UE_LOG(LogAnimation, Warning, TEXT("Unable to generate GUID for %s as RigVMAsset %s does not contain variable %s"), *GetPathName(), RigVMAsset.GetObject() ? *RigVMAsset.GetObject()->GetPathName() : TEXT("Invalid Asset Object"), *GetVariableName().ToString());
+				}
+			}
+			else
+			{
+				UE_LOG(LogAnimation, Warning, TEXT("Unable to generate GUID for %s as RigVMAsset is invalid"), *GetPathName());
+			}
+		}
+	break;
+	default:
+		{
+			checkNoEntry();
+		}
 	}
 }
 
diff --git a/Engine/Plugins/Experimental/UAF/UAF/Source/UAFUncookedOnly/Private/AnimNextUncookedOnlyModule.cpp b/Engine/Plugins/Experimental/UAF/UAF/Source/UAFUncookedOnly/Private/AnimNextUncookedOnlyModule.cpp
--- a/Engine/Plugins/Experimental/UAF/UAF/Source/UAFUncookedOnly/Private/AnimNextUncookedOnlyModule.cpp
+++ b/Engine/Plugins/Experimental/UAF/UAF/Source/UAFUncookedOnly/Private/AnimNextUncookedOnlyModule.cpp
@@ -9,6 +9,7 @@
 #include "Modules/ModuleManager.h"
 #include "UObject/AssetRegistryTagsContext.h"
 #include "MessageLogModule.h"
+#include "Engine/BlueprintGeneratedClass.h"
 #include "Variables/UniversalObjectLocatorBindingType.h"
 
 #define LOCTEXT_NAMESPACE "AnimNextUncookedOnlyModule"
@@ -86,12 +87,49 @@ FGuid FModule::GetVariableGuidByName(const FName InName, const UObject* InObject
 				Guid = FUtils::GenerateScriptStructPropertyGUID(FoundProperty);
 			}
 		}
+		else if (const IRigVMRuntimeAssetInterface* RigVMRuntimeAssetInterface = Cast<IRigVMRuntimeAssetInterface>(InObject))
+		{
+			// Generate a guid from the path name of the property
+			if (const FProperty* Property = RigVMRuntimeAssetInterface->FindGeneratedPropertyByName(InName))
+			{
+				Guid = FGuid::NewDeterministicGuid(Property->GetPathName());
+			}
+		}
+		else if (const IRigVMEditorAssetInterface* EditorAssetInterface = Cast<IRigVMEditorAssetInterface>(InObject))
+		{
+			// Generate a guid from the path name of the property
+			if (const FProperty* Property = EditorAssetInterface->GetRuntimeAssetInterface()->FindGeneratedPropertyByName(InName))
+			{
+				Guid = FGuid::NewDeterministicGuid(Property->GetPathName());
+			}
+		}
 	}
 	
 	
 	return Guid;
 }
 
+FName GetVariableNameByGuidRigVMRuntimeAsset(const FGuid InGuid, const IRigVMRuntimeAssetInterface* RigVMRuntimeAssetInterface)
+{
+	check(RigVMRuntimeAssetInterface);
+
+	// Currently no lookup available for IRigVMRuntimeAssetInterface
+	// So we must do a manual search
+	TArray<FRigVMExternalVariable> Variables = const_cast<IRigVMRuntimeAssetInterface*>(RigVMRuntimeAssetInterface)->GetExternalVariables();
+	FRigVMExternalVariable* FoundVariable = Variables.FindByPredicate( [InGuid](const FRigVMExternalVariable& ItSearch)
+		{
+			FGuid CompareGuid = ItSearch.GetProperty() != nullptr ? FGuid::NewDeterministicGuid(ItSearch.GetProperty()->GetPathName()) : FGuid();
+			return CompareGuid == InGuid;
+		});
+
+	if (FoundVariable != nullptr)
+	{
+		return FoundVariable->GetName();
+	}
+
+	return NAME_None;
+}
+
 FName FModule::GetVariableNameByGuid(const FGuid InGuid, const UObject* InObject) const
 {
 	auto GetNameFromEditorData = [InGuid](const UUAFRigVMAssetEditorData* EditorData)-> FName
@@ -136,6 +174,19 @@ FName FModule::GetVariableNameByGuid(const FGuid InGuid, const UObject* InObject
 				}
 			}
 		}
+		else if (const UBlueprintGeneratedClass* BpGenerated = Cast<UBlueprintGeneratedClass>(InObject))
+		{
+			return BpGenerated->FindBlueprintPropertyNameFromGuid(InGuid);
+		}
+		else if (const IRigVMRuntimeAssetInterface* RigVMRuntimeAssetInterface = Cast<IRigVMRuntimeAssetInterface>(InObject))
+		{
+			return GetVariableNameByGuidRigVMRuntimeAsset(InGuid, RigVMRuntimeAssetInterface);
+		}
+		else if (const IRigVMEditorAssetInterface* RigVMEditorAssetInterface = Cast<IRigVMEditorAssetInterface>(InObject))
+		{
+			// We assume that if we have a editor asset interface, we will also have a runtime asset interface
+			return GetVariableNameByGuidRigVMRuntimeAsset(InGuid, RigVMEditorAssetInterface->GetRuntimeAssetInterface().GetInterface());
+		}
 	}
 	
 	return NAME_None;
diff --git a/Engine/Plugins/Experimental/UAF/UAF/Source/UAFUncookedOnly/Private/Entries/AnimNextSharedVariablesEntry.cpp b/Engine/Plugins/Experimental/UAF/UAF/Source/UAFUncookedOnly/Private/Entries/AnimNextSharedVariablesEntry.cpp
--- a/Engine/Plugins/Experimental/UAF/UAF/Source/UAFUncookedOnly/Private/Entries/AnimNextSharedVariablesEntry.cpp
+++ b/Engine/Plugins/Experimental/UAF/UAF/Source/UAFUncookedOnly/Private/Entries/AnimNextSharedVariablesEntry.cpp
@@ -79,6 +79,7 @@ void UUAFSharedVariablesEntry::Initialize(UUAFRigVMAssetEditorData* InEditorData
 		}
 #endif
 	}
+	// todo: handle RigVMAsset modification
 }
 
 FName UUAFSharedVariablesEntry::GetEntryName() const
@@ -89,6 +90,10 @@ FName UUAFSharedVariablesEntry::GetEntryName() const
 		return Asset ? Asset->GetFName() : NAME_None;
 	case EAnimNextSharedVariablesType::Struct:
 		return Struct ? Struct->GetFName() : NAME_None;
+	case EAnimNextSharedVariablesType::RigVMAsset:
+		return RigVMAsset ? RigVMAsset.GetObject()->GetFName() : NAME_None;
+	default:
+		checkNoEntry();
 	}
 	
 	return NAME_None;
@@ -102,6 +107,10 @@ FText UUAFSharedVariablesEntry::GetDisplayName() const
 		return Asset ? FText::FromName(Asset->GetFName()) : LOCTEXT("InvalidSharedVariables", "Invalid Shared Variables");
 	case EAnimNextSharedVariablesType::Struct:
 		return Struct ? Struct->GetDisplayNameText() : LOCTEXT("InvalidSharedVariables", "Invalid Shared Variables");
+	case EAnimNextSharedVariablesType::RigVMAsset:
+		return RigVMAsset ? FText::FromName(RigVMAsset.GetObject()->GetFName()) : LOCTEXT("InvalidSharedVariables", "Invalid Shared Variables");
+	default:
+		checkNoEntry();
 	}
 
 	return FText::GetEmpty();
@@ -115,6 +124,10 @@ FText UUAFSharedVariablesEntry::GetDisplayNameTooltip() const
 		return Asset ? FText::FromName(Asset->GetFName()) : LOCTEXT("InvalidSharedVariablesTooltip", "Invalid or deleted Shared Variables");
 	case EAnimNextSharedVariablesType::Struct:
 		return Struct ? FText::FromString(Struct->GetPathName()) : LOCTEXT("InvalidSharedVariables", "Invalid Shared Variables");
+	case EAnimNextSharedVariablesType::RigVMAsset:
+		return RigVMAsset ? FText::FromString(RigVMAsset.GetObject()->GetPathName()) : LOCTEXT("InvalidSharedVariables", "Invalid Shared Variables");
+	default:
+		checkNoEntry();
 	}
 
 	return FText::GetEmpty();
@@ -133,6 +146,24 @@ void UUAFSharedVariablesEntry::SetAsset(const UUAFSharedVariables* InAsset, bool
 	Asset = InAsset;
 	ObjectPath = FSoftObjectPath(Asset);
 	Struct = nullptr;
+	RigVMAsset = nullptr;
+}
+
+void UUAFSharedVariablesEntry::SetRigVMAsset(const IRigVMRuntimeAssetInterface* InRigVMAsset, bool bSetupUndoRedo)
+{
+	check(InRigVMAsset != nullptr);
+
+	if(bSetupUndoRedo)
+	{
+		Modify();
+	}
+
+	Type = EAnimNextSharedVariablesType::RigVMAsset;
+	const UObject* Object = Cast<UObject>(InRigVMAsset); 
+	RigVMAsset = Object;
+	ObjectPath = FSoftObjectPath(Object);
+	Struct = nullptr;
+	Asset = nullptr;
 }
 
 const UUAFSharedVariables* UUAFSharedVariablesEntry::GetAsset() const
@@ -158,6 +189,7 @@ void UUAFSharedVariablesEntry::SetStruct(const UScriptStruct* InStruct, bool bSe
 	Struct = InStruct;
 	ObjectPath = FSoftObjectPath(Struct);
 	Asset = nullptr;
+	RigVMAsset = nullptr;
 }
 
 const UScriptStruct* UUAFSharedVariablesEntry::GetStruct() const
diff --git a/Engine/Plugins/Experimental/UAF/UAF/Source/UAFUncookedOnly/Private/UncookedOnlyUtils.cpp b/Engine/Plugins/Experimental/UAF/UAF/Source/UAFUncookedOnly/Private/UncookedOnlyUtils.cpp
--- a/Engine/Plugins/Experimental/UAF/UAF/Source/UAFUncookedOnly/Private/UncookedOnlyUtils.cpp
+++ b/Engine/Plugins/Experimental/UAF/UAF/Source/UAFUncookedOnly/Private/UncookedOnlyUtils.cpp
@@ -217,10 +217,12 @@ void FUtils::CompileVariables(const FRigVMCompileSettings& InSettings, UUAFRigVM
 
 	TSet<const UUAFRigVMAsset*> SharedVariableAssets;
 	TSet<const UScriptStruct*> SharedVariableStructs;
+	TSet<TScriptInterface<const IRigVMRuntimeAssetInterface>> SharedVariableRigVMAssets;
 	
 	// Set of SharedVariable Assets and structs, retrieved by recursively following references
-	TArray<const UUAFRigVMAsset*> RecursiveSharedVariableAssets;
-	TArray<const UScriptStruct*> RecursiveSharedVariableStructs;
+	TSet<const UUAFRigVMAsset*> RecursiveSharedVariableAssets;
+	TSet<const UScriptStruct*> RecursiveSharedVariableStructs;
+	TSet<TScriptInterface<const IRigVMRuntimeAssetInterface>> RecursiveSharedVariableRigVMAssets;
 
 	for(UUAFRigVMAssetEntry* Entry : EditorData->Entries)
 	{
@@ -267,9 +269,9 @@ void FUtils::CompileVariables(const FRigVMCompileSettings& InSettings, UUAFRigVM
 						continue;
 					}
 					
-					RetrieveReferencedSharedVariableSources(SharedVariablesAsset, RecursiveSharedVariableAssets, RecursiveSharedVariableStructs);
+					RetrieveReferencedSharedVariableSources(SharedVariablesAsset, RecursiveSharedVariableAssets, RecursiveSharedVariableStructs, RecursiveSharedVariableRigVMAssets);
 					
-					if (RecursiveSharedVariableAssets.Remove(InAsset) != 0)
+					if (RecursiveSharedVariableAssets.Contains(InAsset))
 					{
 						OutCompileContext.GetAssetCompileContext().Warning(InAsset, LOCTEXT("RecursiveCircularReference", "@@ is recursively referenced by its own Shared Variable Asset reference {0} - chain should be broken."), FText::FromString(SharedVariablesEntry->GetObjectPath().ToString()));
 					}
@@ -295,12 +297,31 @@ void FUtils::CompileVariables(const FRigVMCompileSettings& InSettings, UUAFRigVM
 					SharedVariableStructs.Add(SharedVariablesStruct);
 					break;
 				}
+			case EAnimNextSharedVariablesType::RigVMAsset:
+				{
+					TScriptInterface<const IRigVMRuntimeAssetInterface> RigVMRuntimeAssetInterface = SharedVariablesEntry->GetRigVMAsset();
+					if (RigVMRuntimeAssetInterface == nullptr)
+					{
+						SharedAssetEntryError(FText::Format(LOCTEXT("InvalidSharedVariableRigVMAssetFound", "Invalid shared variables rig VM asset found ({0})"), FText::FromString(SharedVariablesEntry->GetObjectPath().ToString())));
+						continue;
+					}
+
+					if (SharedVariableRigVMAssets.Contains(RigVMRuntimeAssetInterface))
+					{
+						SharedAssetEntryError(FText::Format(LOCTEXT("DuplicateSharedVariableRigVMAssetFound", "Duplicate shared variables struct found ({0})"), FText::FromString(SharedVariablesEntry->GetObjectPath().ToString())));
+						continue;
+					}
+
+					SharedVariableRigVMAssets.Add(RigVMRuntimeAssetInterface);
+					break;
+				}
 			}
 		}
 	}
 
 	SharedVariableAssets.Append(RecursiveSharedVariableAssets);
 	SharedVariableStructs.Append(RecursiveSharedVariableStructs);
+	SharedVariableRigVMAssets.Append(RecursiveSharedVariableRigVMAssets);
 
 	// Set shared variable assets/structs before we setup variables as GetExternalVariablesImpl relies on these
 	if (SharedVariableAssets.Num() > 0)
@@ -321,6 +342,15 @@ void FUtils::CompileVariables(const FRigVMCompileSettings& InSettings, UUAFRigVM
 		InAsset->ReferencedVariableStructs.Empty();
 	}
 
+	if (SharedVariableRigVMAssets.Num() > 0)
+	{
+		InAsset->ReferencedVariableRigVMAssets = SharedVariableRigVMAssets.Array();
+	}
+	else
+	{
+		InAsset->ReferencedVariableRigVMAssets.Empty();
+	}
+
 	// Now rebuild combined the combined property bag we used for stable properties 
 	InAsset->CombinedPropertyBag = EditorData->GenerateCombinedPropertyBag(InSettings, OutCompileContext);
 
@@ -337,26 +367,28 @@ void FUtils::CompileVariables(const FRigVMCompileSettings& InSettings, UUAFRigVM
 	EditorData->bHasCompiledVariables = true;
 
 	if (CVarValidatePropertyBagAgainstExternalVariables.GetValueOnAnyThread())
-    {
-	    // Compilation might have failed
-	    if (const UPropertyBag* PropertyBag = InAsset->CombinedPropertyBag.GetPropertyBagStruct())
-	    {
-		    TConstArrayView<FPropertyBagPropertyDesc> PropertyDescs = PropertyBag->GetPropertyDescs();
-		    TArray<FRigVMExternalVariable> ExternalVariables;
-		   	GetExternalVariables(EditorData, ExternalVariables);
-		    
-		    const int32 NumDescriptions = PropertyDescs.Num();
-		    if (ensure(NumDescriptions == ExternalVariables.Num()))
-		    {
-		    	for (int32 Index = 0; Index < NumDescriptions; ++Index)
-		    	{
-		    		const FPropertyBagPropertyDesc& Desc = PropertyDescs[Index];
-		    		const FRigVMExternalVariable& Variable = ExternalVariables[Index];
-		    		ensureMsgf(Desc.Name == Variable.GetName(), TEXT("Mismatch between compiled property bag variable name %s and external variable entry at same index %s"), *Desc.Name.ToString(), *Variable.GetName().ToString());
-		    	}
-		    }
-	    }
-    }
+	{
+		// Compilation might have failed
+		if (const UPropertyBag* PropertyBag = InAsset->CombinedPropertyBag.GetPropertyBagStruct())
+		{
+			TConstArrayView<FPropertyBagPropertyDesc> PropertyDescs = PropertyBag->GetPropertyDescs();
+			TArray<FRigVMExternalVariable> ExternalVariables;
+			GetExternalVariables(EditorData, ExternalVariables);
+
+			const int32 NumDescriptions = PropertyDescs.Num();
+			if (ensure(NumDescriptions == ExternalVariables.Num()))
+			{
+				for (int32 Index = 0; Index < NumDescriptions; ++Index)
+				{
+					const FPropertyBagPropertyDesc& Desc = PropertyDescs[Index];
+					const FRigVMExternalVariable& Variable = ExternalVariables[Index];
+					ensureMsgf(Desc.Name == Variable.GetName(),
+						TEXT("Mismatch between compiled property bag variable name %s and external variable entry at same index %s"),
+						*Desc.Name.ToString(), *Variable.GetName().ToString());
+				}
+			}
+		}
+	}
 
 	// Validate injection site
 	InAsset->DefaultInjectionSite.Reset();
@@ -425,8 +457,12 @@ void FUtils::GetExternalVariables(const UUAFRigVMAssetEditorData* EditorData, TA
 
 	TSet<const UScriptStruct*> ReferencedVariableStructs;
 	TSet<const UUAFRigVMAsset*> ReferencedSharedVariableAssets;
-	TArray<const UUAFRigVMAsset*> RecursiveSharedVariableAssets;
-	TArray<const UScriptStruct*> RecursiveSharedVariableStructs;
+	TSet<TScriptInterface<const IRigVMRuntimeAssetInterface>> ReferencedSharedVariableRigVMAssets;
+
+	TSet<const UUAFRigVMAsset*> RecursiveSharedVariableAssets;
+	TSet<const UScriptStruct*> RecursiveSharedVariableStructs;
+	TSet<TScriptInterface<const IRigVMRuntimeAssetInterface>> RecursiveSharedVariableRigVMAssets;
+
 	for(UUAFRigVMAssetEntry* Entry : EditorData->Entries)
 	{
 		if(UUAFSharedVariablesEntry* SharedVariablesEntry = Cast<UUAFSharedVariablesEntry>(Entry))
@@ -436,7 +472,7 @@ void FUtils::GetExternalVariables(const UUAFRigVMAssetEditorData* EditorData, TA
 			case EAnimNextSharedVariablesType::Asset:
 				{
 					const UUAFSharedVariables* SharedVariablesAsset = SharedVariablesEntry->GetAsset();
-					RetrieveReferencedSharedVariableSources(SharedVariablesAsset, RecursiveSharedVariableAssets, RecursiveSharedVariableStructs);
+					RetrieveReferencedSharedVariableSources(SharedVariablesAsset, RecursiveSharedVariableAssets, RecursiveSharedVariableStructs, RecursiveSharedVariableRigVMAssets);
 					ReferencedSharedVariableAssets.Add(SharedVariablesAsset);
 					break;
 				}
@@ -446,12 +482,20 @@ void FUtils::GetExternalVariables(const UUAFRigVMAssetEditorData* EditorData, TA
 					ReferencedVariableStructs.Add(SharedVariablesStruct);
 					break;
 				}
+			case EAnimNextSharedVariablesType::RigVMAsset:
+				{
+					ReferencedSharedVariableRigVMAssets.Add(SharedVariablesEntry->GetRigVMAsset());
+					break;
+				}
+			default:
+				checkNoEntry();
 			}
 		}
 	}
-	
+
 	ReferencedSharedVariableAssets.Append(RecursiveSharedVariableAssets);
 	ReferencedVariableStructs.Append(RecursiveSharedVariableStructs);
+	ReferencedSharedVariableRigVMAssets.Append(RecursiveSharedVariableRigVMAssets);
 	
 	// Append each _sorted_ set of referenced shared asset variables
 	for (const UUAFRigVMAsset* ReferencedVariableAsset : ReferencedSharedVariableAssets.ArrayView())
@@ -478,14 +522,23 @@ void FUtils::GetExternalVariables(const UUAFRigVMAssetEditorData* EditorData, TA
 			PropertyOffsets.Add(It->GetOffset_ForInternal());
 		}
 
-		SortByValuesArray(MakeArrayView(StructVariables), MakeConstArrayView(PropertyOffsets), 
-		    [](const int32 OffsetA, const int32 OffsetB) -> bool
-		    {
-			    return OffsetA < OffsetB;
-		    });
+		SortByValuesArray(MakeArrayView(StructVariables), MakeConstArrayView(PropertyOffsets),
+			[](const int32 OffsetA, const int32 OffsetB) -> bool
+				{
+					return OffsetA < OffsetB;
+				});
 
 		OutVariables.Append(StructVariables);
 	}
+
+	for (TScriptInterface<const IRigVMRuntimeAssetInterface> ReferencedVariableRigVMAsset : ReferencedSharedVariableRigVMAssets.ArrayView())
+	{
+		if (ReferencedVariableRigVMAsset.GetInterface())
+		{
+			TArray<FRigVMExternalVariable> RigVMExternalVariables = const_cast<IRigVMRuntimeAssetInterface*>(ReferencedVariableRigVMAsset.GetInterface())->GetExternalVariables();
+			OutVariables.Append(RigVMExternalVariables);
+		}
+	}
 }
 
 void FUtils::GetInternalVariables(const UUAFRigVMAssetEditorData* EditorData, TArray<FRigVMExternalVariable>& OutVariables)
@@ -518,39 +571,50 @@ void FUtils::GetInternalVariables(const UUAFRigVMAssetEditorData* EditorData, TA
 			AssetVariables.Add(FRigVMExternalVariable::Make(ProgrammaticVariable.GetVariableName(), ProgrammaticVariable.GetType().ToString(), const_cast<UObject*>(ProgrammaticVariable.GetType().GetValueTypeObject()), false, false));
 		}
 	}
-	
-	SortByValuesArray(MakeArrayView(AssetVariables), MakeConstArrayView(VariableInfo), 
+
+	SortByValuesArray(MakeArrayView(AssetVariables), MakeConstArrayView(VariableInfo),
 		[](const TPair<FName, FAnimNextParamType>& A, const TPair<FName, FAnimNextParamType>& B)-> bool
-	    {
-		    if(A.Value.GetSize() != B.Value.GetSize())
-		    {
-			    return A.Value.GetSize() > B.Value.GetSize();
-		    }
-    
-		    return A.Key.LexicalLess(B.Key);
-	    });
+			{
+				if (A.Value.GetSize() != B.Value.GetSize())
+				{
+					return A.Value.GetSize() > B.Value.GetSize();
+				}
+
+				return A.Key.LexicalLess(B.Key);
+			});
 	OutVariables.Append(AssetVariables);
 }
 
-void FUtils::RetrieveReferencedSharedVariableSources(const UUAFRigVMAsset* InAsset, TArray<const UUAFRigVMAsset*>& ReferencedVariableAssets, TArray<const UScriptStruct*>& ReferencedVariableStructs)
+void FUtils::RetrieveReferencedSharedVariableSources(const UUAFRigVMAsset* InAsset, TSet<const UUAFRigVMAsset*>& ReferencedVariableAssets,
+														TSet<const UScriptStruct*>& ReferencedVariableStructs,
+														TSet<TScriptInterface<const IRigVMRuntimeAssetInterface>>& ReferencedVariableRigVMAssets)
 {
 	const UUAFRigVMAssetEditorData* EditorData = GetEditorData<UUAFRigVMAssetEditorData>(InAsset);
-	EditorData->ForEachEntryOfType<UUAFSharedVariablesEntry>([&ReferencedVariableAssets, &ReferencedVariableStructs](const UUAFSharedVariablesEntry* Entry)
-	{
-		if (Entry->Asset)
-		{
-			ensure(!Entry->HasAllFlags(RF_NeedPostLoad));
-			if (ReferencedVariableAssets.AddUnique(Entry->Asset) != INDEX_NONE)
+	EditorData->ForEachEntryOfType<UUAFSharedVariablesEntry>(
+		[&ReferencedVariableAssets, &ReferencedVariableStructs, &ReferencedVariableRigVMAssets](const UUAFSharedVariablesEntry* Entry)
 			{
-				RetrieveReferencedSharedVariableSources(Entry->Asset.Get(), ReferencedVariableAssets, ReferencedVariableStructs);
-			}
-		 }
-		 else if (Entry->Struct)
-		 {
-			ReferencedVariableStructs.AddUnique(Entry->Struct);
-		 }
-		return true;
-	});
+				if (Entry->Asset)
+				{
+					ensure(!Entry->HasAllFlags(RF_NeedPostLoad));
+					bool IsAlreadyInSet = false;
+					ReferencedVariableAssets.Add(Entry->Asset, &IsAlreadyInSet);
+					if (!IsAlreadyInSet)
+					{
+						// We added an element, recurse
+						RetrieveReferencedSharedVariableSources(Entry->Asset.Get(), ReferencedVariableAssets, ReferencedVariableStructs,
+							ReferencedVariableRigVMAssets);
+					}
+				}
+				else if (Entry->Struct)
+				{
+					ReferencedVariableStructs.Add(Entry->Struct);
+				}
+				else if (Entry->RigVMAsset)
+				{
+					ReferencedVariableRigVMAssets.Add(Entry->RigVMAsset);
+				}
+				return true;
+			});
 }
 
 void FUtils::CompileVariableBindings(const FRigVMCompileSettings& InSettings, UUAFRigVMAsset* InAsset, TArray<URigVMGraph*>& OutGraphs)
@@ -1177,6 +1241,13 @@ FAnimNextParamType FUtils::FindVariableType(const FAnimNextVariableReference& In
 			return Variable->GetType();
 		}
 	}
+	else if (const IRigVMRuntimeAssetInterface* RigVMAsset = Cast<IRigVMRuntimeAssetInterface>(InVariableReference.GetObject()))
+	{
+		if (const FProperty* Property = InVariableReference.ResolveProperty())
+		{
+			return FAnimNextParamType::FromProperty(Property);
+		}
+	}
 
 	return FAnimNextParamType();
 }
diff --git a/Engine/Plugins/Experimental/UAF/UAFControlRig/Source/UAFControlRig/Internal/AnimNextControlRigModule.h b/Engine/Plugins/Experimental/UAF/UAFControlRig/Source/UAFControlRig/Internal/AnimNextControlRigModule.h
--- a/Engine/Plugins/Experimental/UAF/UAFControlRig/Source/UAFControlRig/Internal/AnimNextControlRigModule.h
+++ b/Engine/Plugins/Experimental/UAF/UAFControlRig/Source/UAFControlRig/Internal/AnimNextControlRigModule.h
@@ -10,6 +10,8 @@
 namespace UE::UAF::ControlRig
 {
 
+struct FAnimNextControlRigVariableProvider;
+
 class UAFCONTROLRIG_API FAnimNextControlRigModule : public IModuleInterface
 {
 public:
@@ -31,6 +33,8 @@ public:
 	static FOnObjectsReinstanced OnObjectsReinstanced;
 
 	FDelegateHandle OnObjectsReinstancedHandle;
+
+	TSharedPtr<FAnimNextControlRigVariableProvider> VariableProvider;
 #endif
 
 #if WITH_EDITORONLY_DATA
diff --git a/Engine/Plugins/Experimental/UAF/UAFControlRig/Source/UAFControlRig/Internal/ControlRigTrait.h b/Engine/Plugins/Experimental/UAF/UAFControlRig/Source/UAFControlRig/Internal/ControlRigTrait.h
--- a/Engine/Plugins/Experimental/UAF/UAFControlRig/Source/UAFControlRig/Internal/ControlRigTrait.h
+++ b/Engine/Plugins/Experimental/UAF/UAFControlRig/Source/UAFControlRig/Internal/ControlRigTrait.h
@@ -199,6 +199,7 @@ struct FControlRigTrait : FBaseTrait, IEvaluate, IUpdate, IHierarchy, IGarbageCo
 		// Returns the property, variable name and the memory of the variable (as target) and the memory of the latent property (as source)
 		static int32 GetExposedVariablesData(const UE::UAF::FTraitBinding& Binding, const FSharedData* SharedData, FControlRigVariableMappings::FCustomPropertyMappings& OutMappings);
 		static int32 GetExposedControlsData(const UE::UAF::FTraitBinding& Binding, const FSharedData* SharedData, FControlRigVariableMappings::FCustomPropertyMappings& OutMappings);
+		static void MapUAFVariablesToControlRig(const UE::UAF::FTraitBinding& Binding, const FSharedData* SharedData, FControlRigVariableMappings::FCustomPropertyMappings& OutMappings);
 
 		void HandleOnInitialized_AnyThread(URigVMHost*, const FName&);
 
diff --git a/Engine/Plugins/Experimental/UAF/UAFControlRig/Source/UAFControlRig/Private/AnimNextControlRigModule.cpp b/Engine/Plugins/Experimental/UAF/UAFControlRig/Source/UAFControlRig/Private/AnimNextControlRigModule.cpp
--- a/Engine/Plugins/Experimental/UAF/UAFControlRig/Source/UAFControlRig/Private/AnimNextControlRigModule.cpp
+++ b/Engine/Plugins/Experimental/UAF/UAFControlRig/Source/UAFControlRig/Private/AnimNextControlRigModule.cpp
@@ -14,7 +14,11 @@
 
 #if WITH_EDITORONLY_DATA
 #include "ControlRigBlueprintLegacy.h"
-#endif 
+#endif
+
+#if WITH_EDITOR
+#include "AnimNextControlRigVariableProvider.h"
+#endif
 #include "Native/UAFSimplePrePhysicsGraphComponent.h"
 
 #include "ControlRigBlueprintGeneratedClass.h"
@@ -41,6 +45,8 @@ void FAnimNextControlRigModule::StartupModule()
 				OnObjectsReinstanced.Broadcast(ObjectMap);
 			}
 		});
+
+	VariableProvider = MakeShared<FAnimNextControlRigVariableProvider>();
 #endif
 
 	RegisterGraphFactories();
@@ -53,6 +59,8 @@ void FAnimNextControlRigModule::ShutdownModule()
 	UnregisterAssetData();
 #if WITH_EDITOR
 	FCoreUObjectDelegates::OnObjectsReinstanced.Remove(OnObjectsReinstancedHandle);
+
+	VariableProvider.Reset();
 #endif
 }
 
diff --git a/Engine/Plugins/Experimental/UAF/UAFControlRig/Source/UAFControlRig/Private/AnimNextControlRigVariableProvider.cpp b/Engine/Plugins/Experimental/UAF/UAFControlRig/Source/UAFControlRig/Private/AnimNextControlRigVariableProvider.cpp
new file mode 100644
--- /dev/null
+++ b/Engine/Plugins/Experimental/UAF/UAFControlRig/Source/UAFControlRig/Private/AnimNextControlRigVariableProvider.cpp
@@ -0,0 +1,94 @@
+// Copyright Epic Games, Inc. All Rights Reserved.
+
+#include "AnimNextControlRigVariableProvider.h"
+
+#include "ControlRig.h"
+#include "RigVMRuntimeAsset.h"
+
+FAnimNextVariableReference UE::UAF::ControlRig::GetAnimNextVariableReferenceFromRigVMExternalVariable(const FRigVMExternalVariable& RigVMVariable,
+	const IRigVMRuntimeAssetInterface* Asset)
+{
+	return FAnimNextVariableReference::FromName(RigVMVariable.GetName(), Cast<const UObject>(Asset));
+}
+
+#if WITH_EDITOR
+
+#include "AnimNextExports.h"
+#include "ControlRigBlueprintLegacy.h"
+#include "Engine/SkeletalMesh.h"
+#include "UObject/AssetRegistryTagsContext.h"
+
+namespace UE::UAF::ControlRig
+{
+FAnimNextControlRigVariableProvider::FAnimNextControlRigVariableProvider()
+{
+	OnGetExtraObjectTagsHandle = UObject::FAssetRegistryTag::OnGetExtraObjectTagsWithContext.AddStatic(&FAnimNextControlRigVariableProvider::GetAssetRegistryTags);
+}
+
+FAnimNextControlRigVariableProvider::~FAnimNextControlRigVariableProvider()
+{
+	UObject::FAssetRegistryTag::OnGetExtraObjectTagsWithContext.Remove(OnGetExtraObjectTagsHandle);
+}
+
+void FAnimNextControlRigVariableProvider::GetAssetRegistryTags(FAssetRegistryTagsContext Context)
+{
+	// UAF variable asset pickers work via the asset registry, so we need to export variable declarations for all control rigs
+
+	// 1. Legacy control rigs (blueprints)
+	/*
+	if (const UControlRigBlueprint* ControlRig = Cast<UControlRigBlueprint>(Context.GetObject()))
+	{
+		FAnimNextAssetRegistryExports Exports;
+		
+		for (const FRigVMGraphVariableDescription& VariableDesc : ControlRig->GetAssetVariables())
+		{
+			if(!VariableDesc.bPublic)
+			{
+				continue;
+			}
+			
+			FAnimNextParamType AnimNextType = FAnimNextParamType::FromString(VariableDesc.CPPType);
+			FGuid Guid = ControlRig->FindBlueprintPropertyGuidFromName(VariableDesc.Name);
+			FAnimNextExport ParameterExport = FAnimNextExport::MakeExport<FAnimNextVariableDeclarationData>(VariableDesc.Name, AnimNextType, Guid, EAnimNextExportedVariableFlags::Public | EAnimNextExportedVariableFlags::Declared);
+			
+			Exports.Exports.Add(ParameterExport);
+		}
+
+		
+		FString TagValue;
+		FAnimNextAssetRegistryExports::StaticStruct()->ExportText(TagValue, &Exports, nullptr, nullptr, PPF_None, nullptr);
+		Context.AddTag(UObject::FAssetRegistryTag(UE::UAF::ExportsAnimNextAssetRegistryTag, TagValue, UObject::FAssetRegistryTag::TT_Hidden));
+	}
+	
+	else*/
+	if (const IRigVMEditorAssetInterface* RigVMAsset = Cast<IRigVMEditorAssetInterface>(Context.GetObject()))
+	{
+		FAnimNextAssetRegistryExports Exports;
+
+		for (const FRigVMGraphVariableDescription& VariableDesc : RigVMAsset->GetAssetVariables())
+		{
+			if(!VariableDesc.bPublic)
+			{
+				continue;
+			}
+			
+			FAnimNextParamType AnimNextType = FAnimNextParamType::FromString(VariableDesc.CPPType);
+			if (const FProperty* Property = RigVMAsset->GetRuntimeAssetInterface()->FindGeneratedPropertyByName(VariableDesc.Name))
+			{
+				FGuid Guid = FGuid::NewDeterministicGuid(Property->GetPathName());
+				FAnimNextExport ParameterExport = FAnimNextExport::MakeExport<FAnimNextVariableDeclarationData>(VariableDesc.Name, AnimNextType, Guid, EAnimNextExportedVariableFlags::Public | EAnimNextExportedVariableFlags::Declared);
+				Exports.Exports.Add(ParameterExport);
+			}
+		}
+
+		
+		FString TagValue;
+		FAnimNextAssetRegistryExports::StaticStruct()->ExportText(TagValue, &Exports, nullptr, nullptr, PPF_None, nullptr);
+		Context.AddTag(UObject::FAssetRegistryTag(UE::UAF::ExportsAnimNextAssetRegistryTag, TagValue, UObject::FAssetRegistryTag::TT_Hidden));
+		
+	}
+}
+}
+
+#endif // #if WITH_EDITOR
+
diff --git a/Engine/Plugins/Experimental/UAF/UAFControlRig/Source/UAFControlRig/Private/AnimNextControlRigVariableProvider.h b/Engine/Plugins/Experimental/UAF/UAFControlRig/Source/UAFControlRig/Private/AnimNextControlRigVariableProvider.h
new file mode 100644
--- /dev/null
+++ b/Engine/Plugins/Experimental/UAF/UAFControlRig/Source/UAFControlRig/Private/AnimNextControlRigVariableProvider.h
@@ -0,0 +1,34 @@
+// Copyright Epic Games, Inc. All Rights Reserved.
+
+#pragma once
+
+#include "CoreMinimal.h"
+#include "UObject/AssetRegistryTagsContext.h"
+#include "Variables/AnimNextVariableReference.h"
+
+class IRigVMRuntimeAssetInterface;
+struct FRigVMExternalVariable;
+
+namespace UE::UAF::ControlRig
+{
+UAFCONTROLRIG_API FAnimNextVariableReference GetAnimNextVariableReferenceFromRigVMExternalVariable(const FRigVMExternalVariable& RigVMVariable, const IRigVMRuntimeAssetInterface* Asset);
+
+#if WITH_EDITOR
+struct FAnimNextControlRigVariableProvider
+{
+public:
+	FAnimNextControlRigVariableProvider();
+	~FAnimNextControlRigVariableProvider();
+
+	static FAnimNextVariableReference GetVariableReferenceFromControlRig(const FRigVMExternalVariable& RigVMVariable, const IRigVMRuntimeAssetInterface* Asset);
+	
+protected:
+	static void GetAssetRegistryTags(FAssetRegistryTagsContext Context);
+
+private:
+	FDelegateHandle OnGetExtraObjectTagsHandle;
+};
+#endif // WITH_EDITOR
+}
+
+
diff --git a/Engine/Plugins/Experimental/UAF/UAFControlRig/Source/UAFControlRig/Private/ControlRigTrait.cpp b/Engine/Plugins/Experimental/UAF/UAFControlRig/Source/UAFControlRig/Private/ControlRigTrait.cpp
--- a/Engine/Plugins/Experimental/UAF/UAFControlRig/Source/UAFControlRig/Private/ControlRigTrait.cpp
+++ b/Engine/Plugins/Experimental/UAF/UAFControlRig/Source/UAFControlRig/Private/ControlRigTrait.cpp
@@ -15,7 +15,10 @@
 #include "Components/SkeletalMeshComponent.h"
 #include "Engine/SkeletalMesh.h"
 #include "AnimationDataSource.h"
+#include "AnimNextControlRigVariableProvider.h"
+#include "RigVMRuntimeAsset.h"
 #include "RigVMCore/RigVMMemoryStorageStruct.h"
+#include "TraitCore/NodeInstance.h"
 #if WITH_EDITOR
 #include "RigVMDeveloperTypeUtils.h"
 #include "RigVMModel/RigVMController.h"
@@ -34,6 +37,7 @@ void FControlRigTraitSharedData::ConstructLatentProperties(const UE::UAF::FTrait
 
 		UE::UAF::FControlRigTrait::FInstanceData::GetExposedVariablesData(Binding, SharedData, InstanceData->PropertyMappings);
 		UE::UAF::FControlRigTrait::FInstanceData::GetExposedControlsData(Binding, SharedData, InstanceData->PropertyMappings);
+		UE::UAF::FControlRigTrait::FInstanceData::MapUAFVariablesToControlRig(Binding, SharedData, InstanceData->PropertyMappings);
 
 		const TArray<FControlRigVariableMappings::FCustomPropertyData>& Mappings = InstanceData->PropertyMappings.GetMappings();
 		const int32 NumProperties = Mappings.Num();
@@ -406,6 +410,56 @@ int32 FControlRigTrait::FInstanceData::GetExposedControlsData(const UE::UAF::FTr
 	return NumElementsAdded;
 }
 
+void FControlRigTrait::FInstanceData::MapUAFVariablesToControlRig(const UE::UAF::FTraitBinding& Binding, const FSharedData* SharedData,
+	FControlRigVariableMappings::FCustomPropertyMappings& OutMappings)
+{
+	FUAFInstanceVariableData& Variables = Binding.GetTraitPtr().GetNodeInstance()->GetOwner().GetVariables();
+
+	if (const IRigVMRuntimeAssetInterface* ControlRigAsset = Cast<const IRigVMRuntimeAssetInterface>(SharedData->ControlRigClass.Get()))
+	{
+		// Todo: this is a copy so its slow
+		// Ideally we can get the details of the control rig without a copy
+		TArray<FRigVMExternalVariable> ControlRigVariables = const_cast<IRigVMRuntimeAssetInterface*>(ControlRigAsset)->GetExternalVariables();
+
+		for(FRigVMExternalVariable& ControlRigVar: ControlRigVariables)
+		{
+			// Check that this property isn't already in the map
+			if (OutMappings.GetMappings().ContainsByPredicate([&ControlRigVar](const FControlRigVariableMappings::FCustomPropertyData& Candidate)
+				{
+					 return Candidate.TargetName == ControlRigVar.GetName();
+				}))
+			{
+				// Mapping already exists, skip
+				// This ensures latent properties are not overridden
+				continue;
+			}
+
+			// Can happen if the rig has changed and hasn't been compiled
+			if (ControlRigVar.GetMemory() == nullptr)
+			{
+				continue;
+			}
+			
+			// Look up the variable in the uaf variable list
+			FAnimNextVariableReference AnimNextReference = ControlRig::GetAnimNextVariableReferenceFromRigVMExternalVariable(ControlRigVar, ControlRigAsset);
+
+			Variables.AccessVariableProperty(AnimNextReference, [&OutMappings, &ControlRigVar](const FProperty* Property, TArrayView<uint8> Data)
+				{
+					// Verify types match
+					if (ControlRigVar.GetProperty()->SameType(Property))
+					{
+						OutMappings.AddVariable(ControlRigVar.GetName(), ControlRigVar.GetMemory(), Property, Data.GetData());
+					}
+					else
+					{
+						UE_LOG(LogAnimation, Error, TEXT("Control rig property type %s does not match UAF property type %s for variable %s"), *ControlRigVar.GetProperty()->GetClass()->GetName(), *Property->GetClass()->GetName(), *ControlRigVar.GetName().ToString());
+					}
+				});
+		}
+	}
+	
+}
+
 void FControlRigTrait::FInstanceData::HandleOnInitialized_AnyThread(URigVMHost*, const FName&)
 {
 	check(EvaluationTask);
@@ -586,6 +640,7 @@ void FControlRigTrait::PreUpdate(FUpdateTraversalContext& Context, const TTraitB
 		InstanceData->PropertyMappings.Reset();
 		UE::UAF::FControlRigTrait::FInstanceData::GetExposedVariablesData(Binding, SharedData, InstanceData->PropertyMappings);
 		UE::UAF::FControlRigTrait::FInstanceData::GetExposedControlsData(Binding, SharedData, InstanceData->PropertyMappings);
+		UE::UAF::FControlRigTrait::FInstanceData::MapUAFVariablesToControlRig(Binding, SharedData, InstanceData->PropertyMappings);
 
 		InstanceData->InitializeControlRig(Context, Binding);
 		InstanceData->bReinitializeControlRig = false;
diff --git a/Engine/Plugins/Runtime/RigVM/Source/RigVM/Private/RigVMRuntimeAsset.cpp b/Engine/Plugins/Runtime/RigVM/Source/RigVM/Private/RigVMRuntimeAsset.cpp
--- a/Engine/Plugins/Runtime/RigVM/Source/RigVM/Private/RigVMRuntimeAsset.cpp
+++ b/Engine/Plugins/Runtime/RigVM/Source/RigVM/Private/RigVMRuntimeAsset.cpp
@@ -160,7 +160,7 @@ TArray<FRigVMExternalVariable> URigVMRuntimeAsset::GetExternalVariables()
 	return ExternalVariables;
 }
 
-TArray<FRigVMGraphVariableDescription> URigVMRuntimeAsset::GetAssetVariables()
+TArray<FRigVMGraphVariableDescription> URigVMRuntimeAsset::GetAssetVariables() const
 {
 	TArray<FRigVMGraphVariableDescription> VariableDescs;
 
diff --git a/Engine/Plugins/Runtime/RigVM/Source/RigVM/Public/RigVMRuntimeAsset.h b/Engine/Plugins/Runtime/RigVM/Source/RigVM/Public/RigVMRuntimeAsset.h
--- a/Engine/Plugins/Runtime/RigVM/Source/RigVM/Public/RigVMRuntimeAsset.h
+++ b/Engine/Plugins/Runtime/RigVM/Source/RigVM/Public/RigVMRuntimeAsset.h
@@ -120,7 +120,7 @@ protected:
 	virtual const UStruct* GetVariablesStruct() override { return Variables.GetPropertyBagStruct(); }
 	virtual uint8* GetVariablesMemory() override { return Variables.GetMutableValue().GetMemory(); }
 	UFUNCTION(BlueprintCallable, Category = "VM")
-	UE_API virtual TArray<FRigVMGraphVariableDescription> GetAssetVariables();
+	UE_API virtual TArray<FRigVMGraphVariableDescription> GetAssetVariables() const;
 	UE_API virtual FRigVMGraphVariableDescription FindAssetVariable(const FName& InName) const;
 #if WITH_EDITOR
 	FOnRigVMPreVariablesChanged OnPreVariablesChangedDelegate;
