Add variable container name to variable trace

#Synced-CL 51100498

diff --git a/Engine/Plugins/Experimental/UAF/UAF/Source/UAF/Private/RewindDebugger/UAFTrace.cpp b/Engine/Plugins/Experimental/UAF/UAF/Source/UAF/Private/RewindDebugger/UAFTrace.cpp
--- a/Engine/Plugins/Experimental/UAF/UAF/Source/UAF/Private/RewindDebugger/UAFTrace.cpp
+++ b/Engine/Plugins/Experimental/UAF/UAF/Source/UAF/Private/RewindDebugger/UAFTrace.cpp
@@ -41,6 +41,7 @@ UE_TRACE_EVENT_END()
 UE_TRACE_EVENT_BEGIN(UAF, InstanceVariableDescriptions)
 	UE_TRACE_EVENT_FIELD(uint32, VariableDescriptionHash)
 	UE_TRACE_EVENT_FIELD(uint8[], VariableDescriptionData)
+	UE_TRACE_EVENT_FIELD(UE::Trace::AnsiString, Name)
 UE_TRACE_EVENT_END()
 
 UE_TRACE_CHANNEL_DEFINE(UAFChannel);
@@ -163,13 +164,16 @@ void FUAFTrace::OutputUAFVariables(const FUAFAssetInstance* Instance, const UObj
 	OutputUAFInstance(Instance, OuterObject);
 
 	bool bVariablesOutput = false;
-	if (Instance->Variables.AllSharedVariableContainers.Num() > 0)
+	if (Instance->Variables.OwnedVariableContainers.Num() > 0)
 	{
 		for (const TWeakPtr<FUAFInstanceVariableContainer>& OwnedVariableSet : Instance->Variables.AllSharedVariableContainers)
 		{
-			if (OutputUAFVariableSet(OwnedVariableSet.Pin().ToSharedRef(), InstanceId, OuterObject))
+			if (TSharedPtr<FUAFInstanceVariableContainer> PinnedContainer = OwnedVariableSet.Pin())
 			{
-				bVariablesOutput = true;
+				if (OutputUAFVariableSet(PinnedContainer.ToSharedRef(), InstanceId, OuterObject))
+				{
+					bVariablesOutput = true;
+				}
 			}
 		}
 	}
@@ -226,9 +230,12 @@ bool FUAFTrace::OutputUAFVariableSet(const TSharedRef<FUAFInstanceVariableContai
 						Archive.UsingCustomVersion(FPropertyBagCustomVersion::GUID);
 						Archive << PropertyDescriptions;
 
+						FString NameStr = InVariableSet->GetFName().ToString();
+
 						UE_TRACE_LOG(UAF, InstanceVariableDescriptions, UAFChannel)
 							<< InstanceVariableDescriptions.VariableDescriptionHash(PropertyDescHash)
-							<< InstanceVariableDescriptions.VariableDescriptionData(ArchiveData.GetData(), ArchiveData.Num());
+							<< InstanceVariableDescriptions.VariableDescriptionData(ArchiveData.GetData(), ArchiveData.Num())
+							<< InstanceVariableDescriptions.Name(TCHAR_TO_ANSI(*NameStr));
 					}
 				}
 
diff --git a/Engine/Plugins/Experimental/UAF/UAF/Source/UAFEditor/Private/RewindDebugger/AnimNextAnalyzer.cpp b/Engine/Plugins/Experimental/UAF/UAF/Source/UAFEditor/Private/RewindDebugger/AnimNextAnalyzer.cpp
--- a/Engine/Plugins/Experimental/UAF/UAF/Source/UAFEditor/Private/RewindDebugger/AnimNextAnalyzer.cpp
+++ b/Engine/Plugins/Experimental/UAF/UAF/Source/UAFEditor/Private/RewindDebugger/AnimNextAnalyzer.cpp
@@ -71,8 +71,13 @@ bool FUAFAnalyzer::OnEvent(uint16 RouteId, EStyle Style, const FOnEventContext&
 		{
 			uint32 VariableDescHash = EventData.GetValue<uint32>("VariableDescriptionHash");
 			TArrayView<const uint8> VariableDescData = EventData.GetArrayView<uint8>("VariableDescriptionData");
+			FString Name;
+			if (EventData.GetString("Name", Name) == false)
+			{
+				Name = TEXT("Unknown Name");
+			}
 
-			Provider.AppendVariableDescriptions(VariableDescHash, VariableDescData);
+			Provider.AppendVariableDescriptions(VariableDescHash, VariableDescData, Name);
 			break;
 		}
 	}
diff --git a/Engine/Plugins/Experimental/UAF/UAF/Source/UAFEditor/Private/RewindDebugger/AnimNextModuleTrack.cpp b/Engine/Plugins/Experimental/UAF/UAF/Source/UAFEditor/Private/RewindDebugger/AnimNextModuleTrack.cpp
--- a/Engine/Plugins/Experimental/UAF/UAF/Source/UAFEditor/Private/RewindDebugger/AnimNextModuleTrack.cpp
+++ b/Engine/Plugins/Experimental/UAF/UAF/Source/UAFEditor/Private/RewindDebugger/AnimNextModuleTrack.cpp
@@ -166,19 +166,20 @@ bool FAnimNextModuleTrack::UpdateInternal()
 								}
 
 								// Load the property descriptions
-								FMemoryReader DescriptionReader(DescriptionData->Data);
+								FMemoryReader DescriptionReader(DescriptionData->PropertyBagData);
 								FObjectAsTraceIdProxyArchiveReader DescriptionReaderProxy(DescriptionReader, GameplayProvider);
 								DescriptionReaderProxy.UsingCustomVersion(UE::UAF::FUAFTrace::CustomVersionGUID);
 								DescriptionReaderProxy << PropertyDescriptions;
 								DetailsObject->Properties.SetNum(DetailsObject->Properties.Num() + 1);
-								FInstancedPropertyBag& PropertyBag = DetailsObject->Properties.Last();
+								DetailsObject->Properties.Last().Name = DescriptionData->Name;
+								FInstancedPropertyBag& PropertyBag = DetailsObject->Properties.Last().PropertyBag;
 								PropertyBag.AddProperties(PropertyDescriptions);
 
 								// Load the property values
 								FMemoryReader Reader(VariableListData.ValueData);
 								FObjectAsTraceIdProxyArchiveReader ReaderProxy(Reader, GameplayProvider);
 
-								UPropertyBag* PropertyBagStruct = const_cast<UPropertyBag*>(DetailsObject->Properties.Last().GetPropertyBagStruct());
+								UPropertyBag* PropertyBagStruct = const_cast<UPropertyBag*>(DetailsObject->Properties.Last().PropertyBag.GetPropertyBagStruct());
 								if (PropertyBagStruct != nullptr)
 								{
 									PropertyBagStruct->SerializeItem(ReaderProxy, PropertyBag.GetMutableValue().GetMemory(), nullptr);
diff --git a/Engine/Plugins/Experimental/UAF/UAF/Source/UAFEditor/Private/RewindDebugger/AnimNextModuleTrack.h b/Engine/Plugins/Experimental/UAF/UAF/Source/UAFEditor/Private/RewindDebugger/AnimNextModuleTrack.h
--- a/Engine/Plugins/Experimental/UAF/UAF/Source/UAFEditor/Private/RewindDebugger/AnimNextModuleTrack.h
+++ b/Engine/Plugins/Experimental/UAF/UAF/Source/UAFEditor/Private/RewindDebugger/AnimNextModuleTrack.h
@@ -14,14 +14,26 @@
 
 #include "AnimNextModuleTrack.generated.h"
 
+USTRUCT()
+struct FInstancedPropertyBagDetails
+{
+	GENERATED_BODY()
+
+	UPROPERTY(VisibleAnywhere, Category = Details)
+	FString Name;
+
+	UPROPERTY(EditAnywhere, Category = Details, meta=(ShowOnlyInnerProperties))
+	FInstancedPropertyBag PropertyBag;
+};
+
 UCLASS()
 class UPropertyBagDetailsObject : public UObject
 {
 	GENERATED_BODY()
 
 public:
-	UPROPERTY(EditAnywhere, Category=Details)
-	TArray<FInstancedPropertyBag> Properties;
+	UPROPERTY(EditAnywhere, Category=Details, meta=(TitleProperty = "Name"))
+	TArray<FInstancedPropertyBagDetails> Properties;
 
 	UPROPERTY(EditAnywhere, Category=Details)
 	TArray<FInstancedStruct> NativeProperties;
diff --git a/Engine/Plugins/Experimental/UAF/UAF/Source/UAFEditor/Private/RewindDebugger/AnimNextProvider.cpp b/Engine/Plugins/Experimental/UAF/UAF/Source/UAFEditor/Private/RewindDebugger/AnimNextProvider.cpp
--- a/Engine/Plugins/Experimental/UAF/UAF/Source/UAFEditor/Private/RewindDebugger/AnimNextProvider.cpp
+++ b/Engine/Plugins/Experimental/UAF/UAF/Source/UAFEditor/Private/RewindDebugger/AnimNextProvider.cpp
@@ -52,7 +52,7 @@ void FUAFProvider::AppendVariables(double ProfileTime, double RecordingTime, uin
 	}
 }
 
-void FUAFProvider::AppendVariableDescriptions(uint32 PropertyDescriptionHash, const TArrayView<const uint8>& VariableDescriptionData)
+void FUAFProvider::AppendVariableDescriptions(uint32 PropertyDescriptionHash, const TArrayView<const uint8>& VariableDescriptionData, const FString& Name)
 {
 	Session.WriteAccessCheck();
 	
@@ -60,7 +60,7 @@ void FUAFProvider::AppendVariableDescriptions(uint32 PropertyDescriptionHash, co
 	{
 		if (!VariableDescriptionData.IsEmpty())
 		{
-			HashToPropertyDescriptionDataMap.Add(PropertyDescriptionHash, FPropertyDescriptionData(VariableDescriptionData));
+			HashToPropertyDescriptionDataMap.Add(PropertyDescriptionHash, FPropertyDescriptionData(VariableDescriptionData, Name));
 		}
 	}
 }
diff --git a/Engine/Plugins/Experimental/UAF/UAF/Source/UAFEditor/Private/RewindDebugger/AnimNextProvider.h b/Engine/Plugins/Experimental/UAF/UAF/Source/UAFEditor/Private/RewindDebugger/AnimNextProvider.h
--- a/Engine/Plugins/Experimental/UAF/UAF/Source/UAFEditor/Private/RewindDebugger/AnimNextProvider.h
+++ b/Engine/Plugins/Experimental/UAF/UAF/Source/UAFEditor/Private/RewindDebugger/AnimNextProvider.h
@@ -16,8 +16,10 @@ namespace TraceServices { class IAnalysisSession; }
 
 struct FPropertyDescriptionData
 {
-	FPropertyDescriptionData(TConstArrayView<uint8> DataView) : Data(DataView) {}
-	TArray<uint8> Data;
+	FPropertyDescriptionData(TConstArrayView<uint8> DataView, const FString& InName) : PropertyBagData(DataView), Name(InName) {}
+	TArray<uint8> PropertyBagData;
+
+	FString Name;
 };
 
 enum class EPropertyVariableDataType
@@ -67,7 +69,7 @@ public:
 
 	void AppendInstance(uint64 InstanceId, uint64 HostInstanceId,  uint64 AssetId, uint64 OuterObjectId);
 	void AppendVariables(double ProfileTime, double RecordingTime, uint64 InModuleInstanceId, EPropertyVariableDataType PropertyVariableDataType, uint32 PropertyDescriptionHash, const TArrayView<const uint8>& VariableData);
-	void AppendVariableDescriptions(uint32 PropertyDescriptionHash, const TArrayView<const uint8>& VariableDescriptionData);
+	void AppendVariableDescriptions(uint32 PropertyDescriptionHash, const TArrayView<const uint8>& VariableDescriptionData, const FString& Name);
 
 	bool GetModuleId(uint64 ComponentId, uint64& OutModuleId) const;
 	const FDataInterfaceData* GetDataInterfaceData(uint64 DataInterfaceId) const;
diff --git a/Engine/Plugins/Experimental/UAF/UAF/Source/UAFEditor/Private/UAF/Debug/SVariableWatch.cpp b/Engine/Plugins/Experimental/UAF/UAF/Source/UAFEditor/Private/UAF/Debug/SVariableWatch.cpp
--- a/Engine/Plugins/Experimental/UAF/UAF/Source/UAFEditor/Private/UAF/Debug/SVariableWatch.cpp
+++ b/Engine/Plugins/Experimental/UAF/UAF/Source/UAFEditor/Private/UAF/Debug/SVariableWatch.cpp
@@ -368,7 +368,7 @@ namespace UE::UAF::Editor
 
 				// Load the property descriptions
 				TArray<FPropertyBagPropertyDesc> PropertyDescriptions;
-				FMemoryReader DescriptionReader(DescriptionData->Data);
+				FMemoryReader DescriptionReader(DescriptionData->PropertyBagData);
 				FObjectAsTraceIdProxyArchiveReader DescriptionReaderProxy(DescriptionReader, GameplayProvider);
 				DescriptionReaderProxy.UsingCustomVersion(UE::UAF::FUAFTrace::CustomVersionGUID);
 				DescriptionReaderProxy << PropertyDescriptions;
