- Remove TSet::ArrayView () calls (not supported in sparse TSet)
- Added nullptr checks to asset/struct for SharedVariableEntries
- Reverted back removal of (recursive) circular shared variable references

#Synced-CL 51158668

diff --git a/Engine/Plugins/Experimental/UAF/UAF/Source/UAFUncookedOnly/Private/UncookedOnlyUtils.cpp b/Engine/Plugins/Experimental/UAF/UAF/Source/UAFUncookedOnly/Private/UncookedOnlyUtils.cpp
--- a/Engine/Plugins/Experimental/UAF/UAF/Source/UAFUncookedOnly/Private/UncookedOnlyUtils.cpp
+++ b/Engine/Plugins/Experimental/UAF/UAF/Source/UAFUncookedOnly/Private/UncookedOnlyUtils.cpp
@@ -271,7 +271,7 @@ void FUtils::CompileVariables(const FRigVMCompileSettings& InSettings, UUAFRigVM
 					
 					RetrieveReferencedSharedVariableSources(SharedVariablesAsset, RecursiveSharedVariableAssets, RecursiveSharedVariableStructs, RecursiveSharedVariableRigVMAssets);
 					
-					if (RecursiveSharedVariableAssets.Contains(InAsset))
+					if (RecursiveSharedVariableAssets.Remove(InAsset) != 0)
 					{
 						OutCompileContext.GetAssetCompileContext().Warning(InAsset, LOCTEXT("RecursiveCircularReference", "@@ is recursively referenced by its own Shared Variable Asset reference {0} - chain should be broken."), FText::FromString(SharedVariablesEntry->GetObjectPath().ToString()));
 					}
@@ -471,15 +471,19 @@ void FUtils::GetExternalVariables(const UUAFRigVMAssetEditorData* EditorData, TA
 			{
 			case EAnimNextSharedVariablesType::Asset:
 				{
-					const UUAFSharedVariables* SharedVariablesAsset = SharedVariablesEntry->GetAsset();
-					RetrieveReferencedSharedVariableSources(SharedVariablesAsset, RecursiveSharedVariableAssets, RecursiveSharedVariableStructs, RecursiveSharedVariableRigVMAssets);
-					ReferencedSharedVariableAssets.Add(SharedVariablesAsset);
+					if (const UUAFSharedVariables* SharedVariablesAsset = SharedVariablesEntry->GetAsset())
+					{
+						RetrieveReferencedSharedVariableSources(SharedVariablesAsset, RecursiveSharedVariableAssets, RecursiveSharedVariableStructs, RecursiveSharedVariableRigVMAssets);
+						ReferencedSharedVariableAssets.Add(SharedVariablesAsset);
+					}
 					break;
 				}
 			case EAnimNextSharedVariablesType::Struct:
 				{
-					const UScriptStruct* SharedVariablesStruct = SharedVariablesEntry->GetStruct();
-					ReferencedVariableStructs.Add(SharedVariablesStruct);
+					if (const UScriptStruct* SharedVariablesStruct = SharedVariablesEntry->GetStruct())
+					{
+						ReferencedVariableStructs.Add(SharedVariablesStruct);
+					}
 					break;
 				}
 			case EAnimNextSharedVariablesType::RigVMAsset:
@@ -498,13 +502,13 @@ void FUtils::GetExternalVariables(const UUAFRigVMAssetEditorData* EditorData, TA
 	ReferencedSharedVariableRigVMAssets.Append(RecursiveSharedVariableRigVMAssets);
 	
 	// Append each _sorted_ set of referenced shared asset variables
-	for (const UUAFRigVMAsset* ReferencedVariableAsset : ReferencedSharedVariableAssets.ArrayView())
+	for (const UUAFRigVMAsset* ReferencedVariableAsset : ReferencedSharedVariableAssets)
 	{
 		GetInternalVariables(GetEditorData<const UUAFRigVMAssetEditorData>(ReferencedVariableAsset), OutVariables);
 	}
 	
 	// Append each set of referenced shared struct variables, sorted by their property offset (can differ from its declaration)
-	for (const UScriptStruct* ReferencedVariableStruct : ReferencedVariableStructs.ArrayView())
+	for (const UScriptStruct* ReferencedVariableStruct : ReferencedVariableStructs)
 	{
 		TArray<int32> PropertyOffsets;
 		TArray<FRigVMExternalVariable> StructVariables;
@@ -531,7 +535,7 @@ void FUtils::GetExternalVariables(const UUAFRigVMAssetEditorData* EditorData, TA
 		OutVariables.Append(StructVariables);
 	}
 
-	for (TScriptInterface<const IRigVMRuntimeAssetInterface> ReferencedVariableRigVMAsset : ReferencedSharedVariableRigVMAssets.ArrayView())
+	for (TScriptInterface<const IRigVMRuntimeAssetInterface> ReferencedVariableRigVMAsset : ReferencedSharedVariableRigVMAssets)
 	{
 		if (ReferencedVariableRigVMAsset.GetInterface())
 		{
@@ -585,9 +589,7 @@ void FUtils::GetInternalVariables(const UUAFRigVMAssetEditorData* EditorData, TA
 	OutVariables.Append(AssetVariables);
 }
 
-void FUtils::RetrieveReferencedSharedVariableSources(const UUAFRigVMAsset* InAsset, TSet<const UUAFRigVMAsset*>& ReferencedVariableAssets,
-														TSet<const UScriptStruct*>& ReferencedVariableStructs,
-														TSet<TScriptInterface<const IRigVMRuntimeAssetInterface>>& ReferencedVariableRigVMAssets)
+void FUtils::RetrieveReferencedSharedVariableSources(const UUAFRigVMAsset* InAsset, TSet<const UUAFRigVMAsset*>& ReferencedVariableAssets, TSet<const UScriptStruct*>& ReferencedVariableStructs, TSet<TScriptInterface<const IRigVMRuntimeAssetInterface>>& ReferencedVariableRigVMAssets)
 {
 	const UUAFRigVMAssetEditorData* EditorData = GetEditorData<UUAFRigVMAssetEditorData>(InAsset);
 	EditorData->ForEachEntryOfType<UUAFSharedVariablesEntry>(
