Fix crash when opening abstract editor with label binding and opening the set binding tab
* Also add missing null checks

#Synced-CL 51008764
#Synced-CL 51008764

diff --git a/Engine/Plugins/Experimental/UAF/UAF/Source/UAFEditor/Private/UAF/AbstractSkeleton/Sets/SSetBinding.cpp b/Engine/Plugins/Experimental/UAF/UAF/Source/UAFEditor/Private/UAF/AbstractSkeleton/Sets/SSetBinding.cpp
--- a/Engine/Plugins/Experimental/UAF/UAF/Source/UAFEditor/Private/UAF/AbstractSkeleton/Sets/SSetBinding.cpp
+++ b/Engine/Plugins/Experimental/UAF/UAF/Source/UAFEditor/Private/UAF/AbstractSkeleton/Sets/SSetBinding.cpp
@@ -33,9 +33,35 @@ namespace UE::UAF::Editor
 				SNew(SWidgetSwitcher)
 				.WidgetIndex_Lambda([this]()
 				{
-					return SetBinding->GetSetCollection() == nullptr ? 0 : 1;
+					if (!SetBinding.IsValid())
+					{
+						return 0;
+					}
+
+					if (!SetBinding->GetSetCollection())
+					{
+						return 1;
+					}
+
+					return 2;
 				})
 				+ SWidgetSwitcher::Slot()
+				[
+					SNew(SBorder)
+					.BorderImage(FAppStyle::Get().GetBrush("Brushes.Recessed"))
+					[
+						SNew(SOverlay)
+						+ SOverlay::Slot()
+						.HAlign(HAlign_Center)
+						.VAlign(VAlign_Center)
+						[
+							SNew(STextBlock)
+							.ColorAndOpacity(FSlateColor::UseSubduedForeground())
+							.Text(LOCTEXT("NoSetBindingSelected", "No Set Binding Selected"))
+						]
+					]
+				]
+				+ SWidgetSwitcher::Slot()
 				[
 					SNew(SBorder)
 					.BorderImage(FAppStyle::Get().GetBrush("Brushes.Recessed"))
@@ -161,32 +187,35 @@ namespace UE::UAF::Editor
 
 	void SSetBinding::HandleUnbindSelection()
 	{
-		const FScopedTransaction Transaction(LOCTEXT("UnbindSelection", "Unbind Selection"));
+		if (SetBinding.IsValid())
+		{
+			const FScopedTransaction Transaction(LOCTEXT("UnbindSelection", "Unbind Selection"));
 		
-		TArray<TSharedPtr<ITreeItem>> Selection = TreeView->GetSelectedItems();
+			TArray<TSharedPtr<ITreeItem>> Selection = TreeView->GetSelectedItems();
 
-		// Unbind all bones/attributes first
-		for (const TSharedPtr<ITreeItem>& Selected : Selection)
-		{
-			if (Selected->GetType() == FTreeItem_Bone::StaticGetType())
+			// Unbind all bones/attributes first
+			for (const TSharedPtr<ITreeItem>& Selected : Selection)
 			{
-				TSharedPtr<FTreeItem_Bone> BoneItem = StaticCastSharedPtr<FTreeItem_Bone>(Selected);
-				SetBinding->RemoveBoneFromSet(BoneItem->Binding.BoneName);
-			}
-			else if (Selected->GetType() == FTreeItem_Attribute::StaticGetType())
-			{
-				TSharedPtr<FTreeItem_Attribute> AttributeItem = StaticCastSharedPtr<FTreeItem_Attribute>(Selected);
-				SetBinding->RemoveAttributeFromSet(AttributeItem->Binding.Attribute);
+				if (Selected->GetType() == FTreeItem_Bone::StaticGetType())
+				{
+					TSharedPtr<FTreeItem_Bone> BoneItem = StaticCastSharedPtr<FTreeItem_Bone>(Selected);
+					SetBinding->RemoveBoneFromSet(BoneItem->Binding.BoneName);
+				}
+				else if (Selected->GetType() == FTreeItem_Attribute::StaticGetType())
+				{
+					TSharedPtr<FTreeItem_Attribute> AttributeItem = StaticCastSharedPtr<FTreeItem_Attribute>(Selected);
+					SetBinding->RemoveAttributeFromSet(AttributeItem->Binding.Attribute);
+				}
 			}
-		}
 
-		// In a second pass, unbind all sets
-		for (const TSharedPtr<ITreeItem>& Selected : Selection)
-		{
-			if (Selected->GetType() == FTreeItem_Set::StaticGetType())
+			// In a second pass, unbind all sets
+			for (const TSharedPtr<ITreeItem>& Selected : Selection)
 			{
-				TSharedPtr<FTreeItem_Set> SetItem = StaticCastSharedPtr<FTreeItem_Set>(Selected);
-				SetBinding->RemoveAllFromSet(SetItem->Set.SetName);
+				if (Selected->GetType() == FTreeItem_Set::StaticGetType())
+				{
+					TSharedPtr<FTreeItem_Set> SetItem = StaticCastSharedPtr<FTreeItem_Set>(Selected);
+					SetBinding->RemoveAllFromSet(SetItem->Set.SetName);
+				}
 			}
 		}
 	}
@@ -436,12 +465,15 @@ namespace UE::UAF::Editor
 					{
 						TSharedPtr<FTreeItem_Set> TargetSetTreeItem = StaticCastSharedPtr<FTreeItem_Set>(TargetTreeItem);
 
-						const FScopedTransaction Transaction(LOCTEXT("BindAttributeToSet", "Bind Attribute to Set"));
-
-						for (const TSharedPtr<SSetsSkeletonTree::FTreeItem>& Item : DragDropOp->Items)
+						if (SetBindingWidget.SetBinding.IsValid())
 						{
-							SetBindingWidget.SetBinding->RemoveBoneFromSet(Item->BoneName);
-							SetBindingWidget.SetBinding->AddBoneToSet(Item->BoneName, TargetSetTreeItem->Set.SetName);
+							const FScopedTransaction Transaction(LOCTEXT("BindAttributeToSet", "Bind Attribute to Set"));
+							
+							for (const TSharedPtr<SSetsSkeletonTree::FTreeItem>& Item : DragDropOp->Items)
+							{
+								SetBindingWidget.SetBinding->RemoveBoneFromSet(Item->BoneName);
+								SetBindingWidget.SetBinding->AddBoneToSet(Item->BoneName, TargetSetTreeItem->Set.SetName);
+							}
 						}
 
 						return FReply::Handled();
@@ -450,13 +482,16 @@ namespace UE::UAF::Editor
 					if (const TSharedPtr<FAttributeDragDropOp> DragDropOp = DragDropEvent.GetOperationAs<FAttributeDragDropOp>())
 					{
 						TSharedPtr<FTreeItem_Set> TargetSetTreeItem = StaticCastSharedPtr<FTreeItem_Set>(TargetTreeItem);
-						
-						const FScopedTransaction Transaction(LOCTEXT("BindAttributeToSet", "Bind Attribute to Set"));
 
-						for (const FAnimationAttributeIdentifier& Attribute : DragDropOp->Attributes)
+						if (SetBindingWidget.SetBinding.IsValid())
 						{
-							SetBindingWidget.SetBinding->RemoveAttributeFromSet(Attribute);
-							SetBindingWidget.SetBinding->AddAttributeToSet(Attribute, TargetSetTreeItem->Set.SetName);
+							const FScopedTransaction Transaction(LOCTEXT("BindAttributeToSet", "Bind Attribute to Set"));
+							
+							for (const FAnimationAttributeIdentifier& Attribute : DragDropOp->Attributes)
+							{
+								SetBindingWidget.SetBinding->RemoveAttributeFromSet(Attribute);
+								SetBindingWidget.SetBinding->AddAttributeToSet(Attribute, TargetSetTreeItem->Set.SetName);
+							}
 						}
 
 						return FReply::Handled();
@@ -518,7 +553,7 @@ namespace UE::UAF::Editor
 			.OnAcceptDrop_Lambda([this](const FDragDropEvent& DragDropEvent, EItemDropZone DropZone, TSharedPtr<ITreeItem> TargetTreeItem) -> FReply
 				{
 					const TSharedPtr<FSkeletonTreeDragDropOp> DragDropOp = DragDropEvent.GetOperationAs<FSkeletonTreeDragDropOp>();
-					if (DragDropOp.IsValid())
+					if (DragDropOp.IsValid() && SetBindingWidget.SetBinding.IsValid())
 					{
 						const FScopedTransaction Transaction(LOCTEXT("AddAttributeToSet", "Add Attribute to Set"));
 
@@ -578,7 +613,7 @@ namespace UE::UAF::Editor
 			.OnAcceptDrop_Lambda([this](const FDragDropEvent& DragDropEvent, EItemDropZone DropZone, TSharedPtr<ITreeItem> TargetTreeItem) -> FReply
 				{
 					const TSharedPtr<FSkeletonTreeDragDropOp> DragDropOp = DragDropEvent.GetOperationAs<FSkeletonTreeDragDropOp>();
-					if (DragDropOp.IsValid())
+					if (DragDropOp.IsValid() && SetBindingWidget.SetBinding.IsValid())
 					{
 						const FScopedTransaction Transaction(LOCTEXT("AddAttributeToSet", "Add Attribute to Set"));
 
