UAF: Fix crash when trying to delete a variable

#Synced-CL 51258030

diff --git a/Engine/Plugins/Experimental/UAF/UAF/Source/UAFEditor/Private/Variables/VariableCustomization.cpp b/Engine/Plugins/Experimental/UAF/UAF/Source/UAFEditor/Private/Variables/VariableCustomization.cpp
--- a/Engine/Plugins/Experimental/UAF/UAF/Source/UAFEditor/Private/Variables/VariableCustomization.cpp
+++ b/Engine/Plugins/Experimental/UAF/UAF/Source/UAFEditor/Private/Variables/VariableCustomization.cpp
@@ -218,18 +218,20 @@ FText FVariableCustomization::GetCategory() const
 	for(int32 Index = 0; Index < Variables.Num(); Index++)
 	{
 		TWeakObjectPtr<UAnimNextVariableEntry> WeakVariable = Variables[Index];
-		UAnimNextVariableEntry* PinnedVariable = WeakVariable.Get();
-		check(PinnedVariable);
-
-		if (Index == 0)
-		{
-			CommonCategory = PinnedVariable->GetVariableCategory(); 	
-		}
-		else
+		if (WeakVariable.IsValid())
 		{
-			if (CommonCategory != PinnedVariable->GetVariableCategory())
+			const UAnimNextVariableEntry* PinnedVariable = WeakVariable.Get();
+			check(PinnedVariable);
+			if (Index == 0)
 			{
-				return LOCTEXT("MultipleValuesLabel", "Multiple Values");
+				CommonCategory = PinnedVariable->GetVariableCategory(); 	
+			}
+			else
+			{
+				if (CommonCategory != PinnedVariable->GetVariableCategory())
+				{
+					return LOCTEXT("MultipleValuesLabel", "Multiple Values");
+				}
 			}
 		}
 	}
