UAF: Fix function nodes not opening on double click if the function node was placed in a different asset than the function outer e.g. using a system function in a anim graph graph panel

#Synced-CL 51290764

diff --git a/Engine/Plugins/Experimental/UAF/UAF/Source/UAFEditor/Private/AnimNextEditorModule.cpp b/Engine/Plugins/Experimental/UAF/UAF/Source/UAFEditor/Private/AnimNextEditorModule.cpp
--- a/Engine/Plugins/Experimental/UAF/UAF/Source/UAFEditor/Private/AnimNextEditorModule.cpp
+++ b/Engine/Plugins/Experimental/UAF/UAF/Source/UAFEditor/Private/AnimNextEditorModule.cpp
@@ -1318,8 +1318,7 @@ void FAnimNextEditorModule::RegisterWorkspaceDocumentTypes(Workspace::IWorkspace
 								const FSoftObjectPath GraphSoftPath = EditorObject;
 								if (GraphSoftPath.IsValid() && GraphSoftPath.IsSubobject())
 								{
-									const FWorkspaceOutlinerItemExport ParentExport = InContext.Document.Export;
-									FWorkspaceOutlinerItemExport Export = FWorkspaceOutlinerItemExport(*GraphSoftPath.GetSubPathUtf8String(), ParentExport);
+									FWorkspaceOutlinerItemExport Export = FWorkspaceOutlinerItemExport(*GraphSoftPath.GetSubPathUtf8String(), GraphSoftPath.GetWithoutSubPath());
 									WorkspaceEditor->OpenExports({ Export });
 								}
 							}
