UAF: Fix typo in FAnimNextVariableReference::Identical

We should be guarding against null Other ptrs. Null Object ptrs are valid.

#Synced-CL 51211255

diff --git a/Engine/Plugins/Experimental/UAF/UAF/Source/UAF/Private/Variables/AnimNextVariableReference.cpp b/Engine/Plugins/Experimental/UAF/UAF/Source/UAF/Private/Variables/AnimNextVariableReference.cpp
--- a/Engine/Plugins/Experimental/UAF/UAF/Source/UAF/Private/Variables/AnimNextVariableReference.cpp
+++ b/Engine/Plugins/Experimental/UAF/UAF/Source/UAF/Private/Variables/AnimNextVariableReference.cpp
@@ -189,7 +189,7 @@ void FAnimNextVariableReference::ValidateVariableNameAndGuid(FName& InOutName, F
 
 bool FAnimNextVariableReference::Identical(const FAnimNextVariableReference* Other, uint32 PortFlags) const
 {
-	return Object != nullptr && (Name == Other->Name && Object == Other->Object);
+	return Other != nullptr && (Name == Other->Name && Object == Other->Object);
 }
 
 bool FAnimNextVariableReference::IsValid() const
