UAF: Remove FunctionsGroup segment from the export generation in the function outliner. It is no longer needed since removing functions from the workspace outliner where they were grouped under the FunctionsGroup

#Synced-CL 51334642

diff --git a/Engine/Plugins/Experimental/UAF/UAF/Source/UAFEditor/Private/Functions/Outliner/FunctionsOutlinerMode.cpp b/Engine/Plugins/Experimental/UAF/UAF/Source/UAFEditor/Private/Functions/Outliner/FunctionsOutlinerMode.cpp
--- a/Engine/Plugins/Experimental/UAF/UAF/Source/UAFEditor/Private/Functions/Outliner/FunctionsOutlinerMode.cpp
+++ b/Engine/Plugins/Experimental/UAF/UAF/Source/UAFEditor/Private/Functions/Outliner/FunctionsOutlinerMode.cpp
@@ -558,8 +558,7 @@ bool FFunctionsOutlinerMode::OpenFunctionItem(const FSceneOutlinerTreeItemPtr It
 				const FWorkspaceOutlinerItemExport AssetExport = UncookedOnly::FUtils::MakeAssetReferenceExport(Asset, RelativeAssetExport);
 				if (AssetExport.GetFirstAssetPath().IsValid() && RelativeAssetExport != AssetExport)
 				{
-					const FWorkspaceOutlinerItemExport FunctionsGroupExport("FunctionsGroup", AssetExport);
-					WorkspaceEditor->OpenExports({UncookedOnly::FUtils::MakeFunctionExport(FunctionNode, FunctionsGroupExport)}, OpenDocumentCause);
+					WorkspaceEditor->OpenExports({UncookedOnly::FUtils::MakeFunctionExport(FunctionNode, AssetExport)}, OpenDocumentCause);
 					return true;
 				}
 
