UAF: Add setting to drive function tab opening behaviour with default behaviour opening functions in a separate tab

#Synced-CL 51334642

diff --git a/Engine/Plugins/Experimental/UAF/UAF/Source/UAFEditor/Private/AnimNextEditorModule.cpp b/Engine/Plugins/Experimental/UAF/UAF/Source/UAFEditor/Private/AnimNextEditorModule.cpp
--- a/Engine/Plugins/Experimental/UAF/UAF/Source/UAFEditor/Private/AnimNextEditorModule.cpp
+++ b/Engine/Plugins/Experimental/UAF/UAF/Source/UAFEditor/Private/AnimNextEditorModule.cpp
@@ -1304,8 +1304,17 @@ void FAnimNextEditorModule::RegisterWorkspaceDocumentTypes(Workspace::IWorkspace
 								const FSoftObjectPath GraphSoftPath = EditorObject;
 								if (GraphSoftPath.IsValid() && GraphSoftPath.IsSubobject())
 								{
+									FDocumentTracker::EOpenDocumentCause OpenCause = FDocumentTracker::NavigatingCurrentDocument;
+									if (const UUAFEditorUserSettings* UAFEditorUserSettings = GetMutableDefault<UUAFEditorUserSettings>())
+									{
+										if (UAFEditorUserSettings->bAutoOpenFunctionsInSeparateTabs)
+										{
+											OpenCause = FDocumentTracker::OpenNewDocument;
+										}
+									}
+									
 									FWorkspaceOutlinerItemExport Export = FWorkspaceOutlinerItemExport(*GraphSoftPath.GetSubPathUtf8String(), GraphSoftPath.GetWithoutSubPath());
-									WorkspaceEditor->OpenExports({ Export });
+									WorkspaceEditor->OpenExports({ Export }, OpenCause);
 								}
 							}
 						}
diff --git a/Engine/Plugins/Experimental/UAF/UAF/Source/UAFEditor/Private/Functions/Outliner/FunctionsOutlinerMode.cpp b/Engine/Plugins/Experimental/UAF/UAF/Source/UAFEditor/Private/Functions/Outliner/FunctionsOutlinerMode.cpp
--- a/Engine/Plugins/Experimental/UAF/UAF/Source/UAFEditor/Private/Functions/Outliner/FunctionsOutlinerMode.cpp
+++ b/Engine/Plugins/Experimental/UAF/UAF/Source/UAFEditor/Private/Functions/Outliner/FunctionsOutlinerMode.cpp
@@ -22,6 +22,7 @@
 #include "Common/Outliner/OutlinerItemMenuContexts.h"
 #include "Common/Outliner/OutlinerDragDropOps.h"
 #include "Common/Outliner/SCommonOutliner.h"
+#include "Settings/UAFEditorUserSettings.h"
 
 #define LOCTEXT_NAMESPACE "FFunctionsOutlinerMode"
 
@@ -505,7 +506,16 @@ void FFunctionsOutlinerMode::BindCommands(const TSharedRef<FUICommandList>& OutC
 
 void FFunctionsOutlinerMode::OnItemDoubleClick(FSceneOutlinerTreeItemPtr Item)
 {
-	if (OpenFunctionItem(Item))
+	FDocumentTracker::EOpenDocumentCause OpenCause = FDocumentTracker::NavigatingCurrentDocument;
+	if (const UUAFEditorUserSettings* UAFEditorUserSettings = GetMutableDefault<UUAFEditorUserSettings>())
+	{
+		if (UAFEditorUserSettings->bAutoOpenFunctionsInSeparateTabs)
+		{
+			OpenCause = FDocumentTracker::OpenNewDocument;
+		}
+	}
+	
+	if (OpenFunctionItem(Item, OpenCause))
 	{
 		return;
 	}
diff --git a/Engine/Plugins/Experimental/UAF/UAF/Source/UAFEditor/Private/Settings/UAFEditorUserSettings.h b/Engine/Plugins/Experimental/UAF/UAF/Source/UAFEditor/Private/Settings/UAFEditorUserSettings.h
--- a/Engine/Plugins/Experimental/UAF/UAF/Source/UAFEditor/Private/Settings/UAFEditorUserSettings.h
+++ b/Engine/Plugins/Experimental/UAF/UAF/Source/UAFEditor/Private/Settings/UAFEditorUserSettings.h
@@ -8,7 +8,7 @@
 /**
  * UAF Editor Settings that are unique per user.
  */
-UCLASS(MinimalAPI, Config = EditorPerProjectUserSettings)
+UCLASS(MinimalAPI, Config = EditorPerProjectUserSettings, meta = (DisplayName = "UAF Editor User Settings"))
 class UUAFEditorUserSettings : public UDeveloperSettings
 {
 	GENERATED_BODY()
@@ -20,4 +20,8 @@ public:
 	/** True if the UAF Browser hotkey opens the sidebar (Right), false will open status bar (Bottom) */
 	UPROPERTY(EditAnywhere, Category="UAF", Config)
 	bool bSummonUAFBrowserHotkeyOpensSidebar = true;
+	
+	/** If true, functions in UAF workspaces will open in a separate tab per default */
+	UPROPERTY(EditAnywhere, Category="UAF", Config)
+	bool bAutoOpenFunctionsInSeparateTabs = true;
 };
\ No newline at end of file
