UAF: Add expand and collapse context menu entries to the common outliner

#Synced-CL 51408588

diff --git a/Engine/Plugins/Experimental/UAF/UAF/Source/UAFEditor/Private/Common/Outliner/CommonOutlinerMode.cpp b/Engine/Plugins/Experimental/UAF/UAF/Source/UAFEditor/Private/Common/Outliner/CommonOutlinerMode.cpp
--- a/Engine/Plugins/Experimental/UAF/UAF/Source/UAFEditor/Private/Common/Outliner/CommonOutlinerMode.cpp
+++ b/Engine/Plugins/Experimental/UAF/UAF/Source/UAFEditor/Private/Common/Outliner/CommonOutlinerMode.cpp
@@ -444,6 +444,22 @@ void FCommonOutlinerMode::BindCommands(const TSharedRef<FUICommandList>& OutComm
 		FIsActionChecked(),
 		FIsActionButtonVisible::CreateRaw(this, &FCommonOutlinerMode::CanModifyEntryExport)
 	);
+	
+	CommandList->MapAction(
+		FOutlinerCommands::Get().ExpandEntries,
+		FExecuteAction::CreateRaw(this, &FCommonOutlinerMode::SetExpansionOnSelectedEntries, true),
+		FCanExecuteAction::CreateRaw(this, &FCommonOutlinerMode::CanExpandEntries),
+		FIsActionChecked(),
+		FIsActionButtonVisible::CreateRaw(this, &FCommonOutlinerMode::CanExpandEntries)
+	);
+	
+	CommandList->MapAction(
+		FOutlinerCommands::Get().CollapseEntries,
+		FExecuteAction::CreateRaw(this, &FCommonOutlinerMode::SetExpansionOnSelectedEntries, false),
+		FCanExecuteAction::CreateRaw(this, &FCommonOutlinerMode::CanCollapseEntries),
+		FIsActionChecked(),
+		FIsActionButtonVisible::CreateRaw(this, &FCommonOutlinerMode::CanCollapseEntries)
+	);
 
 	RegisterToolMenus();
 }
@@ -667,7 +683,68 @@ bool FCommonOutlinerMode::CanModifyEntryExport() const
 
 	return false;
 }
-
+	
+	
+void FCommonOutlinerMode::SetExpansionOnSelectedEntries(bool bExpand) const
+{
+	if (SceneOutliner)
+	{
+		const FSceneOutlinerItemSelection& Selection = SceneOutliner->GetSelection();
+		for (const TWeakPtr<ISceneOutlinerTreeItem>& Item :  Selection.SelectedItems)
+		{
+			if (const TSharedPtr<ISceneOutlinerTreeItem> ItemPtr = Item.Pin())
+			{
+				if (ItemPtr->IsA<FOutlinerCategoryItem>() || ItemPtr->IsA<FOutlinerAssetItem>())
+				{
+					SceneOutliner->SetItemExpansion(ItemPtr, bExpand);
+				}
+			}
+		}
+	}
+}
+	
+bool FCommonOutlinerMode::CanExpandEntries() const
+{
+	if (SceneOutliner)
+	{
+		const FSceneOutlinerItemSelection& Selection = SceneOutliner->GetSelection();
+		for (const TWeakPtr<ISceneOutlinerTreeItem>& Item : Selection.SelectedItems)
+		{
+			if (const TSharedPtr<ISceneOutlinerTreeItem> ItemPtr = Item.Pin())
+			{
+				if ((ItemPtr->IsA<FOutlinerCategoryItem>() || ItemPtr->IsA<FOutlinerAssetItem>()) 
+					&& !SceneOutliner->IsItemExpanded(ItemPtr))
+				{
+					return true;
+				}
+			}
+		}
+	}
+	
+	return false;
+}
+	
+bool FCommonOutlinerMode::CanCollapseEntries() const
+{
+	if (SceneOutliner)
+	{
+		const FSceneOutlinerItemSelection& Selection = SceneOutliner->GetSelection();
+		for (const TWeakPtr<ISceneOutlinerTreeItem>& Item : Selection.SelectedItems)
+		{
+			if (const TSharedPtr<ISceneOutlinerTreeItem> ItemPtr = Item.Pin())
+			{
+				if ((ItemPtr->IsA<FOutlinerCategoryItem>() || ItemPtr->IsA<FOutlinerAssetItem>()) 
+					&& SceneOutliner->IsItemExpanded(ItemPtr))
+				{
+					return true;
+				}
+			}
+		}
+	}
+	
+	return false;
+}
+	
 SCommonOutliner* FCommonOutlinerMode::GetOutliner() const
 {
 	return static_cast<SCommonOutliner*>(SceneOutliner);
@@ -770,6 +847,9 @@ void FCommonOutlinerMode::RegisterToolMenus()
 					EntriesSection.AddMenuEntryWithCommandList(FGenericCommands::Get().Copy, MenuContext->WeakCommandList.Pin());
 					EntriesSection.AddMenuEntryWithCommandList(FGenericCommands::Get().Cut, MenuContext->WeakCommandList.Pin());
 
+					EntriesSection.AddMenuEntryWithCommandList(FOutlinerCommands::Get().ExpandEntries, MenuContext->WeakCommandList.Pin());
+					EntriesSection.AddMenuEntryWithCommandList(FOutlinerCommands::Get().CollapseEntries, MenuContext->WeakCommandList.Pin());
+					
 					EntriesSection.AddSubMenu("FindReferencesMenu", LOCTEXT("FindReferencesMenuLabel", "Find References..."), LOCTEXT("FindReferencesMenuToolTip", "Find references to selected variables by different means."),
 						FNewToolMenuDelegate::CreateLambda(
 							[MenuContext](UToolMenu* InSubmenu)
diff --git a/Engine/Plugins/Experimental/UAF/UAF/Source/UAFEditor/Private/Common/Outliner/CommonOutlinerMode.h b/Engine/Plugins/Experimental/UAF/UAF/Source/UAFEditor/Private/Common/Outliner/CommonOutlinerMode.h
--- a/Engine/Plugins/Experimental/UAF/UAF/Source/UAFEditor/Private/Common/Outliner/CommonOutlinerMode.h
+++ b/Engine/Plugins/Experimental/UAF/UAF/Source/UAFEditor/Private/Common/Outliner/CommonOutlinerMode.h
@@ -83,6 +83,11 @@ protected:
 
 	void SetEntryExportAccessSpecifiersOnSelection(EAnimNextExportAccessSpecifier AnimNextExportAccessSpecifier) const;
 	bool CanModifyEntryExport() const;
+	
+	void SetExpansionOnSelectedEntries(bool bExpand) const;
+	bool CanExpandEntries() const;
+	bool CanCollapseEntries() const;
+	
 protected:
 	typedef TFunction<TSharedPtr<FDragDropOperation>(const FSceneOutlinerTreeItemPtr&)> FDragDropOpConstructor;
 	TMap<FSceneOutlinerTreeItemType, FDragDropOpConstructor> ItemTypeToDragDropOpMap;
diff --git a/Engine/Plugins/Experimental/UAF/UAF/Source/UAFEditor/Private/Common/Outliner/OutlinerCommands.cpp b/Engine/Plugins/Experimental/UAF/UAF/Source/UAFEditor/Private/Common/Outliner/OutlinerCommands.cpp
--- a/Engine/Plugins/Experimental/UAF/UAF/Source/UAFEditor/Private/Common/Outliner/OutlinerCommands.cpp
+++ b/Engine/Plugins/Experimental/UAF/UAF/Source/UAFEditor/Private/Common/Outliner/OutlinerCommands.cpp
@@ -18,6 +18,9 @@ void FOutlinerCommands::RegisterCommands()
 	UI_COMMAND(MakeEntryPrivate, "Sets Entry to Private", "Sets Entry to be Private.", EUserInterfaceActionType::Button, FInputChord());
 
 	UI_COMMAND(OpenInNewTab, "Open in New Tab", "Opens the selected item in a new tab.", EUserInterfaceActionType::Button, FInputChord());
+	
+	UI_COMMAND(ExpandEntries, "Expand Entries", "Expands entries in the current selection.", EUserInterfaceActionType::Button, FInputChord());
+	UI_COMMAND(CollapseEntries, "Collapse Entries", "Collapses entries in the current selection.", EUserInterfaceActionType::Button, FInputChord());
 }
 }
 
diff --git a/Engine/Plugins/Experimental/UAF/UAF/Source/UAFEditor/Private/Common/Outliner/OutlinerCommands.h b/Engine/Plugins/Experimental/UAF/UAF/Source/UAFEditor/Private/Common/Outliner/OutlinerCommands.h
--- a/Engine/Plugins/Experimental/UAF/UAF/Source/UAFEditor/Private/Common/Outliner/OutlinerCommands.h
+++ b/Engine/Plugins/Experimental/UAF/UAF/Source/UAFEditor/Private/Common/Outliner/OutlinerCommands.h
@@ -26,6 +26,8 @@ public:
 	TSharedPtr<FUICommandInfo> MakeEntryPublic;
 	TSharedPtr<FUICommandInfo> MakeEntryPrivate;
 	TSharedPtr<FUICommandInfo> OpenInNewTab;
+	TSharedPtr<FUICommandInfo> ExpandEntries;
+	TSharedPtr<FUICommandInfo> CollapseEntries;
 };
 }
 
