Mark the _asset_ referencing a modified (requiring compilation) shared-variable asset as requiring compilation as well

#Synced-CL 51408661

diff --git a/Engine/Plugins/Experimental/UAF/UAF/Source/UAFUncookedOnly/Internal/Entries/AnimNextSharedVariablesEntry.h b/Engine/Plugins/Experimental/UAF/UAF/Source/UAFUncookedOnly/Internal/Entries/AnimNextSharedVariablesEntry.h
--- a/Engine/Plugins/Experimental/UAF/UAF/Source/UAFUncookedOnly/Internal/Entries/AnimNextSharedVariablesEntry.h
+++ b/Engine/Plugins/Experimental/UAF/UAF/Source/UAFUncookedOnly/Internal/Entries/AnimNextSharedVariablesEntry.h
@@ -101,6 +101,9 @@ class UUAFSharedVariablesEntry : public UUAFRigVMAssetEntry
 
 	// Recompiles this asset when the linked asset is modified
 	void HandleAssetModified(UUAFRigVMAssetEditorData* InEditorData, EAnimNextEditorDataNotifType InType, UObject* InSubject);
+	
+	// Marks this asset as requiring compilation when the linked asset is marked as such
+	void HandleAssetRequiresCompilation(TWeakObjectPtr<UUAFRigVMAssetEditorData> WeakEditorData) const;
 
 #if WITH_LIVE_CODING
 	// Recompiles this asset when live coding is run (potentially modifying the struct we reference)
diff --git a/Engine/Plugins/Experimental/UAF/UAF/Source/UAFUncookedOnly/Private/Entries/AnimNextSharedVariablesEntry.cpp b/Engine/Plugins/Experimental/UAF/UAF/Source/UAFUncookedOnly/Private/Entries/AnimNextSharedVariablesEntry.cpp
--- a/Engine/Plugins/Experimental/UAF/UAF/Source/UAFUncookedOnly/Private/Entries/AnimNextSharedVariablesEntry.cpp
+++ b/Engine/Plugins/Experimental/UAF/UAF/Source/UAFUncookedOnly/Private/Entries/AnimNextSharedVariablesEntry.cpp
@@ -69,6 +69,9 @@ void UUAFSharedVariablesEntry::Initialize(UUAFRigVMAssetEditorData* InEditorData
 	{
 		UUAFSharedVariables_EditorData* EditorData = UE::UAF::UncookedOnly::FUtils::GetEditorData<UUAFSharedVariables_EditorData>(Asset.Get());
 		EditorData->ModifiedDelegate.AddUObject(this, &UUAFSharedVariablesEntry::HandleAssetModified);
+		
+		TWeakObjectPtr<UUAFRigVMAssetEditorData> WeakEditorData = EditorData;
+		EditorData->RecompileRequiredChangedEvent.AddUObject(this, &UUAFSharedVariablesEntry::HandleAssetRequiresCompilation, WeakEditorData);
 	}
 	else if (Struct)
 	{
@@ -220,6 +223,20 @@ void UUAFSharedVariablesEntry::HandleAssetModified(UUAFRigVMAssetEditorData* InE
 	}
 }
 
+void UUAFSharedVariablesEntry::HandleAssetRequiresCompilation(TWeakObjectPtr<UUAFRigVMAssetEditorData> WeakEditorData) const
+{
+	if(UUAFRigVMAssetEditorData* EditorData = GetTypedOuter<UUAFRigVMAssetEditorData>())
+	{
+		if (UUAFRigVMAssetEditorData* SharedAssetEditorData = WeakEditorData.Get())
+		{
+			if (SharedAssetEditorData->bVMRecompilationRequired)
+			{
+				EditorData->SetVMRecompilationRequired(true);
+			}
+		}
+	}
+}
+
 #if WITH_LIVE_CODING
 void UUAFSharedVariablesEntry::HandlePatchComplete()
 {
