Add check to ensure any statically referenced UAFAssets have compiled with success during instantiation of AnimNextModule

#Synced-CL 51408661

diff --git a/Engine/Plugins/Experimental/UAF/UAF/Source/UAF/Private/Module/AnimNextModuleInstance.cpp b/Engine/Plugins/Experimental/UAF/UAF/Source/UAF/Private/Module/AnimNextModuleInstance.cpp
--- a/Engine/Plugins/Experimental/UAF/UAF/Source/UAF/Private/Module/AnimNextModuleInstance.cpp
+++ b/Engine/Plugins/Experimental/UAF/UAF/Source/UAF/Private/Module/AnimNextModuleInstance.cpp
@@ -91,6 +91,14 @@ void FAnimNextModuleInstance::Initialize(const FUAFSystemFactoryParams& InFactor
 		{
 			return;
 		}
+
+		for (const UUAFRigVMAsset* StaticallyReferencedAsset : System->ReferencedVariableAssets)
+		{
+			if (!StaticallyReferencedAsset || StaticallyReferencedAsset->CompilationState == EAnimNextRigVMAssetState::CompiledWithErrors)
+			{
+				return;
+			}
+		}
 #endif
 	}
 	else
diff --git a/Engine/Plugins/Experimental/UAF/UAF/Source/UAF/Public/AnimNextRigVMAsset.h b/Engine/Plugins/Experimental/UAF/UAF/Source/UAF/Public/AnimNextRigVMAsset.h
--- a/Engine/Plugins/Experimental/UAF/UAF/Source/UAF/Public/AnimNextRigVMAsset.h
+++ b/Engine/Plugins/Experimental/UAF/UAF/Source/UAF/Public/AnimNextRigVMAsset.h
@@ -79,6 +79,7 @@ protected:
 	friend UE::UAF::Tests::FSystemReference_Creation;
 	friend FUAFAssetInstance;
 	friend class UUAFTestBlueprintLibrary;
+	friend struct FAnimNextModuleInstance;
 
 	UE_API UUAFRigVMAsset(const FObjectInitializer& ObjectInitializer);
 
