Rename FAttributeCollectionPtr to FValueBundlePtr

#Synced-CL 51372242

diff --git a/Engine/Plugins/Experimental/UAF/UAF/Source/UAF/Private/UAF/Attributes/AttributeBindingData.cpp b/Engine/Plugins/Experimental/UAF/UAF/Source/UAF/Private/UAF/Attributes/AttributeBindingData.cpp
--- a/Engine/Plugins/Experimental/UAF/UAF/Source/UAF/Private/UAF/Attributes/AttributeBindingData.cpp
+++ b/Engine/Plugins/Experimental/UAF/UAF/Source/UAF/Private/UAF/Attributes/AttributeBindingData.cpp
@@ -336,7 +336,7 @@ namespace UE::UAF
 				FAttributeNamedSetPtr NamedSet = *NamedSetIt;
 				const FName SetName = NamedSet->GetName();
 
-				FAttributeCollectionPtr DefaultAttributeValues = MakeShared<FValueBundleHeap>(NamedSet);
+				FValueBundlePtr DefaultAttributeValues = MakeShared<FValueBundleHeap>(NamedSet);
 				DefaultAttributeValues->SetValueSpace(FValueSpace(EValueSpaceType::Local));
 
 				for (auto TypedSetIt = NamedSet->CreateTypedSetIterator(); TypedSetIt; ++TypedSetIt)
@@ -422,15 +422,15 @@ namespace UE::UAF
 		return FAttributeBindingIndex();
 	}
 
-	FAttributeCollectionPtr FAttributeBindingData::GetDefaultAttributeValues(FName SetName) const
+	FValueBundlePtr FAttributeBindingData::GetDefaultAttributeValues(FName SetName) const
 	{
-		if (const FAttributeCollectionPtr* AttributeCollection = DefaultAttributeValuesByName.Find(SetName))
+		if (const FValueBundlePtr* AttributeCollection = DefaultAttributeValuesByName.Find(SetName))
 		{
 			return *AttributeCollection;
 		}
 
 		// Set not found
-		return FAttributeCollectionPtr();
+		return FValueBundlePtr();
 	}
 
 	uint32 FAttributeBindingData::Release() const
diff --git a/Engine/Plugins/Experimental/UAF/UAF/Source/UAF/Private/UAF/Attributes/AttributeNamedSet.cpp b/Engine/Plugins/Experimental/UAF/UAF/Source/UAF/Private/UAF/Attributes/AttributeNamedSet.cpp
--- a/Engine/Plugins/Experimental/UAF/UAF/Source/UAF/Private/UAF/Attributes/AttributeNamedSet.cpp
+++ b/Engine/Plugins/Experimental/UAF/UAF/Source/UAF/Private/UAF/Attributes/AttributeNamedSet.cpp
@@ -61,7 +61,7 @@ namespace UE::UAF
 		return FAttributeNamedSetPtr(DesiredLOD, true);
 	}
 
-	FAttributeCollectionPtr FAttributeNamedSet::GetDefaultAttributeValues() const
+	FValueBundlePtr FAttributeNamedSet::GetDefaultAttributeValues() const
 	{
 		return Owner->GetDefaultAttributeValues(Name);
 	}
diff --git a/Engine/Plugins/Experimental/UAF/UAF/Source/UAF/Private/UAF/ValueRuntime/Transformers/AdditiveSpace.cpp b/Engine/Plugins/Experimental/UAF/UAF/Source/UAF/Private/UAF/ValueRuntime/Transformers/AdditiveSpace.cpp
--- a/Engine/Plugins/Experimental/UAF/UAF/Source/UAF/Private/UAF/ValueRuntime/Transformers/AdditiveSpace.cpp
+++ b/Engine/Plugins/Experimental/UAF/UAF/Source/UAF/Private/UAF/ValueRuntime/Transformers/AdditiveSpace.cpp
@@ -485,7 +485,7 @@ namespace UE::UAF::Transformers
 			{
 				// Our base is missing, use the reference pose
 				const FAttributeTypedSetPtr& TypedSet = Input->GetTypedSet();
-				const FAttributeCollectionPtr ReferenceValues = TypedSet->GetBindingData()->GetDefaultAttributeValues(TypedSet->GetNamedSet()->GetName());
+				const FValueBundlePtr ReferenceValues = TypedSet->GetBindingData()->GetDefaultAttributeValues(TypedSet->GetNamedSet()->GetName());
 				Base = ReferenceValues->GetBoundValueMaps().Find(Input->GetMappingKey());
 			}
 
@@ -582,7 +582,7 @@ namespace UE::UAF::Transformers
 			{
 				// Our base is missing, use the user authored default values
 				const FAttributeTypedSetPtr& TypedSet = Input->GetTypedSet();
-				const FAttributeCollectionPtr ReferenceValues = TypedSet->GetBindingData()->GetDefaultAttributeValues(TypedSet->GetNamedSet()->GetName());
+				const FValueBundlePtr ReferenceValues = TypedSet->GetBindingData()->GetDefaultAttributeValues(TypedSet->GetNamedSet()->GetName());
 				Base = ReferenceValues->GetBoundValueMaps().Find(Input->GetMappingKey());
 			}
 
@@ -666,7 +666,7 @@ namespace UE::UAF::Transformers
 			{
 				// Our base is missing, use the reference pose
 				const FAttributeTypedSetPtr& TypedSet = Additive->GetTypedSet();
-				const FAttributeCollectionPtr ReferenceValues = TypedSet->GetBindingData()->GetDefaultAttributeValues(TypedSet->GetNamedSet()->GetName());
+				const FValueBundlePtr ReferenceValues = TypedSet->GetBindingData()->GetDefaultAttributeValues(TypedSet->GetNamedSet()->GetName());
 				Base = ReferenceValues->GetBoundValueMaps().Find(Additive->GetMappingKey());
 			}
 
@@ -824,7 +824,7 @@ namespace UE::UAF::Transformers
 			{
 				// Our base is missing, use the user authored default values
 				const FAttributeTypedSetPtr& TypedSet = Additive->GetTypedSet();
-				const FAttributeCollectionPtr ReferenceValues = TypedSet->GetBindingData()->GetDefaultAttributeValues(TypedSet->GetNamedSet()->GetName());
+				const FValueBundlePtr ReferenceValues = TypedSet->GetBindingData()->GetDefaultAttributeValues(TypedSet->GetNamedSet()->GetName());
 				Base = ReferenceValues->GetBoundValueMaps().Find(Additive->GetMappingKey());
 			}
 
diff --git a/Engine/Plugins/Experimental/UAF/UAF/Source/UAF/Public/UAF/Attributes/AttributeBindingData.h b/Engine/Plugins/Experimental/UAF/UAF/Source/UAF/Public/UAF/Attributes/AttributeBindingData.h
--- a/Engine/Plugins/Experimental/UAF/UAF/Source/UAF/Public/UAF/Attributes/AttributeBindingData.h
+++ b/Engine/Plugins/Experimental/UAF/UAF/Source/UAF/Public/UAF/Attributes/AttributeBindingData.h
@@ -14,7 +14,7 @@ class UAbstractSkeletonSetBinding;
 
 namespace UE::UAF
 {
-	using FAttributeCollectionPtr = TSharedPtr<class FValueBundle, ESPMode::ThreadSafe>;
+	using FValueBundlePtr = TSharedPtr<class FValueBundle, ESPMode::ThreadSafe>;
 	using FAttributeBindingDataPtr = TRefCountPtr<const class FAttributeBindingData>;
 
 	// Our baked abstract skeleton binding data
@@ -92,7 +92,7 @@ namespace UE::UAF
 		[[nodiscard]] UE_API FAttributeBindingIndex FindBindingIndex(FName AttributeName, UScriptStruct* AttributeType) const;
 
 		// Returns a collection of the default attribute values for the specified named set
-		[[nodiscard]] UE_API FAttributeCollectionPtr GetDefaultAttributeValues(FName SetName) const;
+		[[nodiscard]] UE_API FValueBundlePtr GetDefaultAttributeValues(FName SetName) const;
 
 		// UE::Private::FQueryableRefCountedObject API
 		UE_API uint32 Release() const;
@@ -123,7 +123,7 @@ namespace UE::UAF
 		TMap<UScriptStruct*, TArray<FName>> AttributeNamesByType;
 
 		// Default attribute values for each named set
-		TMap<FName, FAttributeCollectionPtr> DefaultAttributeValuesByName;
+		TMap<FName, FValueBundlePtr> DefaultAttributeValuesByName;
 
 		// The source binding asset used to create this binding data
 		TWeakObjectPtr<const UAbstractSkeletonSetBinding> BindingAsset;
diff --git a/Engine/Plugins/Experimental/UAF/UAF/Source/UAF/Public/UAF/Attributes/AttributeNamedSet.h b/Engine/Plugins/Experimental/UAF/UAF/Source/UAF/Public/UAF/Attributes/AttributeNamedSet.h
--- a/Engine/Plugins/Experimental/UAF/UAF/Source/UAF/Public/UAF/Attributes/AttributeNamedSet.h
+++ b/Engine/Plugins/Experimental/UAF/UAF/Source/UAF/Public/UAF/Attributes/AttributeNamedSet.h
@@ -13,7 +13,7 @@ namespace UE::UAF
 	class FAttributeBindingData;
 
 	using FAttributeBindingDataPtr = TRefCountPtr<const class FAttributeBindingData>;
-	using FAttributeCollectionPtr = TSharedPtr<class FValueBundle, ESPMode::ThreadSafe>;
+	using FValueBundlePtr = TSharedPtr<class FValueBundle, ESPMode::ThreadSafe>;
 	using FAttributeNamedSetPtr = TRefCountPtr<const class FAttributeNamedSet>;
 
 	class FAttributeNamedSet
@@ -97,7 +97,7 @@ namespace UE::UAF
 		[[nodiscard]] UE_FORCEINLINE_HINT FTypedSetIterator CreateTypedSetIterator() const { return FTypedSetIterator(TypedSetMap); }
 
 		// Returns a collection of the default attribute values for this named set
-		[[nodiscard]] UE_API FAttributeCollectionPtr GetDefaultAttributeValues() const;
+		[[nodiscard]] UE_API FValueBundlePtr GetDefaultAttributeValues() const;
 
 		// Reference counting API
 		UE_API void AddRef() const;
diff --git a/Engine/Plugins/Experimental/UAF/UAF/Source/UAF/Public/UAF/ValueRuntime/ValueBundle.h b/Engine/Plugins/Experimental/UAF/UAF/Source/UAF/Public/UAF/ValueRuntime/ValueBundle.h
--- a/Engine/Plugins/Experimental/UAF/UAF/Source/UAF/Public/UAF/ValueRuntime/ValueBundle.h
+++ b/Engine/Plugins/Experimental/UAF/UAF/Source/UAF/Public/UAF/ValueRuntime/ValueBundle.h
@@ -123,7 +123,7 @@ namespace UE::UAF
 		using FValueBundle::operator=;
 	};
 
-	using FAttributeCollectionPtr = TSharedPtr<FValueBundle, ESPMode::ThreadSafe>;
+	using FValueBundlePtr = TSharedPtr<FValueBundle, ESPMode::ThreadSafe>;
 }
 
 /*
diff --git a/Engine/Plugins/Experimental/UAF/UAFAnimGraph/Source/UAFAnimGraph/Private/EvaluationVM/EvaluationTaskContext.cpp b/Engine/Plugins/Experimental/UAF/UAFAnimGraph/Source/UAFAnimGraph/Private/EvaluationVM/EvaluationTaskContext.cpp
--- a/Engine/Plugins/Experimental/UAF/UAFAnimGraph/Source/UAFAnimGraph/Private/EvaluationVM/EvaluationTaskContext.cpp
+++ b/Engine/Plugins/Experimental/UAF/UAFAnimGraph/Source/UAFAnimGraph/Private/EvaluationVM/EvaluationTaskContext.cpp
@@ -54,7 +54,7 @@ namespace UE::UAF
 		return false;
 	}
 
-	const FAttributeCollectionPtr& FEvaluationTaskContext::GetDefaultAttributeValues() const
+	const FValueBundlePtr& FEvaluationTaskContext::GetDefaultAttributeValues() const
 	{
 		return DefaultAttributeValues;
 	}
diff --git a/Engine/Plugins/Experimental/UAF/UAFAnimGraph/Source/UAFAnimGraph/Private/EvaluationVM/Tasks/DecompressionTools.cpp b/Engine/Plugins/Experimental/UAF/UAFAnimGraph/Source/UAFAnimGraph/Private/EvaluationVM/Tasks/DecompressionTools.cpp
--- a/Engine/Plugins/Experimental/UAF/UAFAnimGraph/Source/UAFAnimGraph/Private/EvaluationVM/Tasks/DecompressionTools.cpp
+++ b/Engine/Plugins/Experimental/UAF/UAFAnimGraph/Source/UAFAnimGraph/Private/EvaluationVM/Tasks/DecompressionTools.cpp
@@ -293,7 +293,7 @@ namespace Private
 				// It also doesn't make sense if the animation is additive, the reference root transform isn't in additive space
 				// and we have nothing to subtract it from, we should probably use the additive identity
 				const FAttributeNamedSetPtr& NamedSet = OutBoneTransforms.GetTypedSet()->GetNamedSet();
-				const FAttributeCollectionPtr ReferenceValues = NamedSet->GetDefaultAttributeValues();
+				const FValueBundlePtr& ReferenceValues = NamedSet->GetDefaultAttributeValues();
 				const TBoundValueMap<FBoneTransformAnimationAttribute>* ReferenceBoneTransforms = ReferenceValues->GetBoundValueMaps().Find<FBoneTransformAnimationAttribute>(OutBoneTransforms.GetMappingKey());
 				check(ReferenceBoneTransforms != nullptr && !ReferenceBoneTransforms->IsEmpty());
 
diff --git a/Engine/Plugins/Experimental/UAF/UAFAnimGraph/Source/UAFAnimGraph/Public/EvaluationVM/EvaluationTaskContext.h b/Engine/Plugins/Experimental/UAF/UAFAnimGraph/Source/UAFAnimGraph/Public/EvaluationVM/EvaluationTaskContext.h
--- a/Engine/Plugins/Experimental/UAF/UAFAnimGraph/Source/UAFAnimGraph/Public/EvaluationVM/EvaluationTaskContext.h
+++ b/Engine/Plugins/Experimental/UAF/UAFAnimGraph/Source/UAFAnimGraph/Public/EvaluationVM/EvaluationTaskContext.h
@@ -41,7 +41,7 @@ namespace UE::UAF
 		[[nodiscard]] UE_API bool SetNamedSet(FName SetName);
 
 		// Returns the default attribute values for the current named set (or nullptr if we aren't valid)
-		[[nodiscard]] UE_API const FAttributeCollectionPtr& GetDefaultAttributeValues() const;
+		[[nodiscard]] UE_API const FValueBundlePtr& GetDefaultAttributeValues() const;
 
 	private:
 		// The set binding we are evaluating with
@@ -57,7 +57,7 @@ namespace UE::UAF
 		FAttributeNamedSetPtr NamedSet;
 
 		// Default values as specified by the current set binding and named set
-		FAttributeCollectionPtr DefaultAttributeValues;
+		FValueBundlePtr DefaultAttributeValues;
 	};
 }
 
