UAF: Inline graphs as virtual value bundles

Renames some classes/files that got missed in review of the last set of changes
Supplies an 'inline graph' virtual value bundle (no usage yet, that is to come)

#Synced-CL 51338579

diff --git a/Engine/Plugins/Experimental/UAF/UAF/Source/UAF/Internal/Graph/UAFSystemOutputComponent.h b/Engine/Plugins/Experimental/UAF/UAF/Source/UAF/Internal/Graph/UAFSystemOutputComponent.h
--- a/Engine/Plugins/Experimental/UAF/UAF/Source/UAF/Internal/Graph/UAFSystemOutputComponent.h
+++ b/Engine/Plugins/Experimental/UAF/UAF/Source/UAF/Internal/Graph/UAFSystemOutputComponent.h
@@ -19,7 +19,7 @@ struct FUAFValueBundle;
 
 namespace UE::UAF
 {
-struct FValueBundleImpl_SystemOutput;
+struct FVirtualValueBundle_SystemOutput;
 class ISystemOutputAdapter;
 }
 
@@ -103,7 +103,7 @@ private:
 	friend FRigUnit_UAFWriteSystemOutput;
 	friend UMeshEntityComponentImpl_Skinned;
 	friend FEntitySkeletalAnimationModule;
-	friend UE::UAF::FValueBundleImpl_SystemOutput;
+	friend UE::UAF::FVirtualValueBundle_SystemOutput;
 	friend UUAFComponent;
 	friend UE::UAF::Tests::FUAFSystemTests;
 
diff --git a/Engine/Plugins/Experimental/UAF/UAF/Source/UAF/Internal/UAF/ValueRuntime/IVirtualValueBundle.h b/Engine/Plugins/Experimental/UAF/UAF/Source/UAF/Internal/UAF/ValueRuntime/IVirtualValueBundle.h
--- a/Engine/Plugins/Experimental/UAF/UAF/Source/UAF/Internal/UAF/ValueRuntime/IVirtualValueBundle.h
+++ b/Engine/Plugins/Experimental/UAF/UAF/Source/UAF/Internal/UAF/ValueRuntime/IVirtualValueBundle.h
@@ -3,32 +3,15 @@
 #pragma once
 
 #include "CoreTypes.h"
-#include "Containers/Array.h"
-#include "Misc/MemStack.h"
 
 #define UE_API UAF_API
 
-struct FAnimNextGraphInstance;
-struct FRigUnit_ReadInputPose;
 struct FAnimNextGraphLODPose;
 
 namespace UE::UAF
 {
 class FValueBundleHeap;
-struct FInputValueTrait;
-struct FEvaluateTraversalContext;
-struct FUpdateTraversalContext;
-template<class TraitInterfaceType>
-struct TTraitBinding;
-struct IEvaluate;
-struct IUpdate;
-struct IUpdateTraversal;
-struct IHierarchy;
-struct FWeakTraitPtr;
-using FChildrenArray = TArray<FWeakTraitPtr, TInlineAllocator<8, TMemStackAllocator<>>>;;
-struct FTraitUpdateState;
-struct FUpdateTraversalQueue;
-struct FExecutionContext;
+struct FTraitPtr;
 struct FWeakAnimGraphReference;
 }
 
@@ -46,13 +29,16 @@ public:
 	virtual ~IVirtualValueBundle() = default;
 
 	// Get the LOD pose this bundle represents
-	virtual const FAnimNextGraphLODPose* GetLODPose() const = 0;
+	virtual const FAnimNextGraphLODPose* GetLODPose() const { return nullptr; }
 
 	// Get a value bundle - valid if this bundle is a concrete value
-	virtual const FValueBundleHeap* GetValueBundle() const = 0;
+	virtual const FValueBundleHeap* GetValueBundle() const { return nullptr; }
 
 	// Get an anim graph reference that represents a 'virtual' value bundle - i.e. the result of evaluating the graph in-line
-	virtual const FWeakAnimGraphReference* GetAnimGraphReference() const = 0;
+	virtual const FWeakAnimGraphReference* GetAnimGraphReference() const { return nullptr; }
+	
+	// Get an inline child ptr that represents a 'virtual' value bundle - i.e. the result of evaluating the graph fragment in-line
+	virtual const FTraitPtr* GetInlineGraph() const { return nullptr; }
 };
 
 }
diff --git a/Engine/Plugins/Experimental/UAF/UAF/Source/UAF/Private/Component/AnimNextComponent.cpp b/Engine/Plugins/Experimental/UAF/UAF/Source/UAF/Private/Component/AnimNextComponent.cpp
--- a/Engine/Plugins/Experimental/UAF/UAF/Source/UAF/Private/Component/AnimNextComponent.cpp
+++ b/Engine/Plugins/Experimental/UAF/UAF/Source/UAF/Private/Component/AnimNextComponent.cpp
@@ -15,7 +15,7 @@
 #include "Module/UAFWeakSystemReference.h"
 #include "UAF/UAFAssetFactory.h"
 #include "UAF/ValueRuntime/ValueBundle.h"
-#include "UAF/ValueRuntime/ValueBundleImpl_SystemOutput.h"
+#include "UAF/ValueRuntime/VirtualValueBundle_SystemOutput.h"
 
 #include UE_INLINE_GENERATED_CPP_BY_NAME(AnimNextComponent)
 
@@ -630,7 +630,7 @@ void UUAFComponent::SetupInputOutput()
 				{
 					if (!InContext.AccessVariable<FUAFValueBundle>(InputBinding.Variable, [Component](FUAFValueBundle& InValueBundle)
 					{
-						InValueBundle.SetAs<UE::UAF::FValueBundleImpl_SystemOutput>(FUAFWeakSystemReference(Component->SystemReference));
+						InValueBundle.SetAs<UE::UAF::FVirtualValueBundle_SystemOutput>(FUAFWeakSystemReference(Component->SystemReference));
 					}))
 					{
 						UE_LOGFMT(LogAnimation, Warning, "UUAFComponent::SetupInputOutput: Cannot bind input pose variable: '{0}'", InputBinding.Variable.GetName());
diff --git a/Engine/Plugins/Experimental/UAF/UAF/Source/UAF/Private/UAF/ValueRuntime/ValueBundleImpl_SystemOutput.cpp b/Engine/Plugins/Experimental/UAF/UAF/Source/UAF/Private/UAF/ValueRuntime/ValueBundleImpl_SystemOutput.cpp
deleted file mode 100644
diff --git a/Engine/Plugins/Experimental/UAF/UAF/Source/UAF/Private/UAF/ValueRuntime/ValueBundleImpl_SystemOutput.h b/Engine/Plugins/Experimental/UAF/UAF/Source/UAF/Private/UAF/ValueRuntime/ValueBundleImpl_SystemOutput.h
deleted file mode 100644
diff --git a/Engine/Plugins/Experimental/UAF/UAF/Source/UAF/Private/UAF/ValueRuntime/VirtualValueBundle_SystemOutput.cpp b/Engine/Plugins/Experimental/UAF/UAF/Source/UAF/Private/UAF/ValueRuntime/VirtualValueBundle_SystemOutput.cpp
new file mode 100644
--- /dev/null
+++ b/Engine/Plugins/Experimental/UAF/UAF/Source/UAF/Private/UAF/ValueRuntime/VirtualValueBundle_SystemOutput.cpp
@@ -0,0 +1,41 @@
+﻿// Copyright Epic Games, Inc. All Rights Reserved.
+
+#include "VirtualValueBundle_SystemOutput.h"
+
+#include "Graph/UAFSystemOutputComponent.h"
+
+namespace UE::UAF
+{
+
+const FAnimNextGraphLODPose* FVirtualValueBundle_SystemOutput::GetLODPose() const
+{
+	ReadSystemOutput();
+
+	return &CachedPose;
+}
+
+const FValueBundleHeap* FVirtualValueBundle_SystemOutput::GetValueBundle() const
+{
+	ReadSystemOutput();
+
+	// TODO: Implement this for value runtime
+	return &CachedValueBundle;
+}
+
+void FVirtualValueBundle_SystemOutput::ReadSystemOutput() const
+{
+	// TODO: Implement this for value runtime
+	auto PoseReader = [this](uint32 InSerialNumber, const FAnimNextGraphLODPose& InPose, const FAnimNextGraphReferencePose& InRefPose)
+	{
+		SerialNumber = InSerialNumber;
+		CachedPose = InPose;
+	};
+
+	SystemReference.ReadComponent<FUAFSystemOutputComponent>([this, &PoseReader](TConstStructView<FUAFSystemOutputComponent> InStructView)
+	{
+		const FUAFSystemOutputComponent& SystemOutputComponent = InStructView.Get<FUAFSystemOutputComponent>();
+		SystemOutputComponent.ReadOutput(SerialNumber, PoseReader);
+	});
+}
+
+}
diff --git a/Engine/Plugins/Experimental/UAF/UAF/Source/UAF/Private/UAF/ValueRuntime/VirtualValueBundle_SystemOutput.h b/Engine/Plugins/Experimental/UAF/UAF/Source/UAF/Private/UAF/ValueRuntime/VirtualValueBundle_SystemOutput.h
new file mode 100644
--- /dev/null
+++ b/Engine/Plugins/Experimental/UAF/UAF/Source/UAF/Private/UAF/ValueRuntime/VirtualValueBundle_SystemOutput.h
@@ -0,0 +1,37 @@
+﻿// Copyright Epic Games, Inc. All Rights Reserved.
+
+#pragma once
+
+#include "Graph/AnimNext_LODPose.h"
+#include "Module/UAFWeakSystemReference.h"
+#include "UAF/ValueRuntime/IVirtualValueBundle.h"
+#include "UAF/ValueRuntime/ValueBundle.h"
+
+#define UE_API UAF_API
+
+namespace UE::UAF
+{
+
+struct FVirtualValueBundle_SystemOutput : public IVirtualValueBundle
+{
+	explicit FVirtualValueBundle_SystemOutput(const FUAFWeakSystemReference& InSystemReference)
+		: SystemReference(InSystemReference)
+	{}
+
+private:
+	// IVirtualValueBundle interface
+	virtual const FAnimNextGraphLODPose* GetLODPose() const override;
+	virtual const FValueBundleHeap* GetValueBundle() const override;
+
+	// Read the output value bundle of the system
+	void ReadSystemOutput() const;
+
+	mutable FAnimNextGraphLODPose CachedPose;
+	mutable FValueBundleHeap CachedValueBundle;
+	FUAFWeakSystemReference SystemReference;
+	mutable uint32 SerialNumber = 0;
+};
+
+}
+
+#undef UE_API
\ No newline at end of file
diff --git a/Engine/Plugins/Experimental/UAF/UAFAnimGraph/Source/UAFAnimGraph/Internal/Graph/InputPoseImpl_ExternalGraph.h b/Engine/Plugins/Experimental/UAF/UAFAnimGraph/Source/UAFAnimGraph/Internal/Graph/InputPoseImpl_ExternalGraph.h
deleted file mode 100644
diff --git a/Engine/Plugins/Experimental/UAF/UAFAnimGraph/Source/UAFAnimGraph/Internal/Graph/VirtualValueBundle_ExternalGraph.h b/Engine/Plugins/Experimental/UAF/UAFAnimGraph/Source/UAFAnimGraph/Internal/Graph/VirtualValueBundle_ExternalGraph.h
new file mode 100644
--- /dev/null
+++ b/Engine/Plugins/Experimental/UAF/UAFAnimGraph/Source/UAFAnimGraph/Internal/Graph/VirtualValueBundle_ExternalGraph.h
@@ -0,0 +1,32 @@
+﻿// Copyright Epic Games, Inc. All Rights Reserved.
+
+#pragma once
+
+#include "Graph/WeakAnimGraphReference.h"
+#include "UAF/ValueRuntime/IVirtualValueBundle.h"
+
+#define UE_API UAFANIMGRAPH_API
+
+namespace UE::UAF
+{
+
+struct FVirtualValueBundle_ExternalGraph : public IVirtualValueBundle
+{
+	explicit FVirtualValueBundle_ExternalGraph(const FWeakAnimGraphReference& InGraphReference)
+		: WeakGraphReference(InGraphReference)
+	{}
+
+	explicit FVirtualValueBundle_ExternalGraph(const FAnimGraphReference& InGraphReference)
+		: WeakGraphReference(InGraphReference)
+	{}
+
+private:
+	// IVirtualValueBundle interface
+	UE_API virtual const FWeakAnimGraphReference* GetAnimGraphReference() const override;
+
+	FWeakAnimGraphReference WeakGraphReference;
+};
+
+}
+
+#undef UE_API
\ No newline at end of file
diff --git a/Engine/Plugins/Experimental/UAF/UAFAnimGraph/Source/UAFAnimGraph/Internal/Graph/VirtualValueBundle_InlineGraph.h b/Engine/Plugins/Experimental/UAF/UAFAnimGraph/Source/UAFAnimGraph/Internal/Graph/VirtualValueBundle_InlineGraph.h
new file mode 100644
--- /dev/null
+++ b/Engine/Plugins/Experimental/UAF/UAFAnimGraph/Source/UAFAnimGraph/Internal/Graph/VirtualValueBundle_InlineGraph.h
@@ -0,0 +1,29 @@
+﻿// Copyright Epic Games, Inc. All Rights Reserved.
+
+#pragma once
+
+#include "Graph/WeakAnimGraphReference.h"
+#include "TraitCore/TraitPtr.h"
+#include "UAF/ValueRuntime/IVirtualValueBundle.h"
+
+#define UE_API UAFANIMGRAPH_API
+
+namespace UE::UAF
+{
+
+struct FVirtualValueBundle_InlineGraph : public IVirtualValueBundle
+{
+	explicit FVirtualValueBundle_InlineGraph(const FTraitPtr& InTraitPtr)
+		: TraitPtr(InTraitPtr)
+	{}
+
+private:
+	// IVirtualValueBundle interface
+	UE_API virtual const FTraitPtr* GetInlineGraph() const override;
+
+	FTraitPtr TraitPtr;
+};
+
+}
+
+#undef UE_API
\ No newline at end of file
diff --git a/Engine/Plugins/Experimental/UAF/UAFAnimGraph/Source/UAFAnimGraph/Private/Graph/InputPoseImpl_ExternalGraph.cpp b/Engine/Plugins/Experimental/UAF/UAFAnimGraph/Source/UAFAnimGraph/Private/Graph/InputPoseImpl_ExternalGraph.cpp
deleted file mode 100644
diff --git a/Engine/Plugins/Experimental/UAF/UAFAnimGraph/Source/UAFAnimGraph/Private/Graph/InputPoseImpl_ValueBundle.cpp b/Engine/Plugins/Experimental/UAF/UAFAnimGraph/Source/UAFAnimGraph/Private/Graph/InputPoseImpl_ValueBundle.cpp
deleted file mode 100644
diff --git a/Engine/Plugins/Experimental/UAF/UAFAnimGraph/Source/UAFAnimGraph/Private/Graph/InputPoseImpl_ValueBundle.h b/Engine/Plugins/Experimental/UAF/UAFAnimGraph/Source/UAFAnimGraph/Private/Graph/InputPoseImpl_ValueBundle.h
deleted file mode 100644
diff --git a/Engine/Plugins/Experimental/UAF/UAFAnimGraph/Source/UAFAnimGraph/Private/Graph/RigUnit_UAFRunAsset.cpp b/Engine/Plugins/Experimental/UAF/UAFAnimGraph/Source/UAFAnimGraph/Private/Graph/RigUnit_UAFRunAsset.cpp
--- a/Engine/Plugins/Experimental/UAF/UAFAnimGraph/Source/UAFAnimGraph/Private/Graph/RigUnit_UAFRunAsset.cpp
+++ b/Engine/Plugins/Experimental/UAF/UAFAnimGraph/Source/UAFAnimGraph/Private/Graph/RigUnit_UAFRunAsset.cpp
@@ -2,7 +2,7 @@
 
 #include "Graph/RigUnit_UAFRunAsset.h"
 
-#include "InputPoseImpl_ValueBundle.h"
+#include "VirtualValueBundle_ValueBundle.h"
 #include "TraitInterfaces/IEvaluate.h"
 #include "TraitInterfaces/IUpdate.h"
 #include "Module/AnimNextModuleInstance.h"
@@ -161,5 +161,5 @@ FRigUnit_UAFRunAsset_Execute()
 	}
 
 	// Move our output pose into the result
-	Result.SetAs<FValueBundleImpl_ValueBundle>(MoveTemp(GraphPose));
+	Result.SetAs<FVirtualValueBundle_ValueBundle>(MoveTemp(GraphPose));
 }
diff --git a/Engine/Plugins/Experimental/UAF/UAFAnimGraph/Source/UAFAnimGraph/Private/Graph/VirtualValueBundle_ExternalGraph.cpp b/Engine/Plugins/Experimental/UAF/UAFAnimGraph/Source/UAFAnimGraph/Private/Graph/VirtualValueBundle_ExternalGraph.cpp
new file mode 100644
--- /dev/null
+++ b/Engine/Plugins/Experimental/UAF/UAFAnimGraph/Source/UAFAnimGraph/Private/Graph/VirtualValueBundle_ExternalGraph.cpp
@@ -0,0 +1,13 @@
+﻿// Copyright Epic Games, Inc. All Rights Reserved.
+
+#include "Graph/VirtualValueBundle_ExternalGraph.h"
+
+namespace UE::UAF
+{
+
+const FWeakAnimGraphReference* FVirtualValueBundle_ExternalGraph::GetAnimGraphReference() const
+{
+	return &WeakGraphReference;
+}
+
+}
diff --git a/Engine/Plugins/Experimental/UAF/UAFAnimGraph/Source/UAFAnimGraph/Private/Graph/VirtualValueBundle_InlineGraph.cpp b/Engine/Plugins/Experimental/UAF/UAFAnimGraph/Source/UAFAnimGraph/Private/Graph/VirtualValueBundle_InlineGraph.cpp
new file mode 100644
--- /dev/null
+++ b/Engine/Plugins/Experimental/UAF/UAFAnimGraph/Source/UAFAnimGraph/Private/Graph/VirtualValueBundle_InlineGraph.cpp
@@ -0,0 +1,13 @@
+﻿// Copyright Epic Games, Inc. All Rights Reserved.
+
+#include "Graph/VirtualValueBundle_InlineGraph.h"
+
+namespace UE::UAF
+{
+
+const FTraitPtr* FVirtualValueBundle_InlineGraph::GetInlineGraph() const
+{
+	return &TraitPtr;
+}
+
+}
diff --git a/Engine/Plugins/Experimental/UAF/UAFAnimGraph/Source/UAFAnimGraph/Private/Graph/VirtualValueBundle_ValueBundle.cpp b/Engine/Plugins/Experimental/UAF/UAFAnimGraph/Source/UAFAnimGraph/Private/Graph/VirtualValueBundle_ValueBundle.cpp
new file mode 100644
--- /dev/null
+++ b/Engine/Plugins/Experimental/UAF/UAFAnimGraph/Source/UAFAnimGraph/Private/Graph/VirtualValueBundle_ValueBundle.cpp
@@ -0,0 +1,18 @@
+﻿// Copyright Epic Games, Inc. All Rights Reserved.
+
+#include "Graph/VirtualValueBundle_ValueBundle.h"
+
+namespace UE::UAF
+{
+
+const FAnimNextGraphLODPose* FVirtualValueBundle_ValueBundle::GetLODPose() const
+{
+	return &Pose;
+}
+
+const FValueBundleHeap* FVirtualValueBundle_ValueBundle::GetValueBundle() const
+{
+	return &ValueBundle;
+}
+
+}
diff --git a/Engine/Plugins/Experimental/UAF/UAFAnimGraph/Source/UAFAnimGraph/Private/Graph/VirtualValueBundle_ValueBundle.h b/Engine/Plugins/Experimental/UAF/UAFAnimGraph/Source/UAFAnimGraph/Private/Graph/VirtualValueBundle_ValueBundle.h
new file mode 100644
--- /dev/null
+++ b/Engine/Plugins/Experimental/UAF/UAFAnimGraph/Source/UAFAnimGraph/Private/Graph/VirtualValueBundle_ValueBundle.h
@@ -0,0 +1,44 @@
+﻿// Copyright Epic Games, Inc. All Rights Reserved.
+
+#pragma once
+
+#include "Graph/AnimNext_LODPose.h"
+#include "Graph/WeakAnimGraphReference.h"
+#include "UAF/ValueRuntime/IVirtualValueBundle.h"
+#include "UAF/ValueRuntime/ValueBundle.h"
+
+#define UE_API UAFANIMGRAPH_API
+
+namespace UE::UAF
+{
+
+struct FVirtualValueBundle_ValueBundle : public IVirtualValueBundle
+{
+	explicit FVirtualValueBundle_ValueBundle(const FAnimNextGraphLODPose& InPose)
+		: Pose(InPose)
+	{}
+
+	explicit FVirtualValueBundle_ValueBundle(FAnimNextGraphLODPose&& InPose)
+		: Pose(MoveTemp(InPose))
+	{}
+
+	explicit FVirtualValueBundle_ValueBundle(const FValueBundleHeap& InValueBundle)
+		: ValueBundle(InValueBundle)
+	{}
+	
+		explicit FVirtualValueBundle_ValueBundle(FValueBundleHeap&& InValueBundle)
+		: ValueBundle(MoveTemp(InValueBundle))
+	{}
+	
+	// IVirtualValueBundle interface
+	UE_API virtual const FAnimNextGraphLODPose* GetLODPose() const override;
+	UE_API virtual const FValueBundleHeap* GetValueBundle() const override;
+
+private:
+	FAnimNextGraphLODPose Pose;
+	FValueBundleHeap ValueBundle;
+};
+
+}
+
+#undef UE_API
\ No newline at end of file
diff --git a/Engine/Plugins/Experimental/UAF/UAFAnimGraph/Source/UAFAnimGraph/Private/Traits/InputValueTrait.cpp b/Engine/Plugins/Experimental/UAF/UAFAnimGraph/Source/UAFAnimGraph/Private/Traits/InputValueTrait.cpp
--- a/Engine/Plugins/Experimental/UAF/UAFAnimGraph/Source/UAFAnimGraph/Private/Traits/InputValueTrait.cpp
+++ b/Engine/Plugins/Experimental/UAF/UAFAnimGraph/Source/UAFAnimGraph/Private/Traits/InputValueTrait.cpp
@@ -101,6 +101,13 @@ namespace UE::UAF
 						TraversalQueue.Push(GraphInstance->GetGraphRootPtr(), TraitState);
 					}
 				}
+				else if (const FTraitPtr* TraitPtr = Impl->GetInlineGraph())
+				{
+					if (TraitPtr->IsValid())
+					{
+						TraversalQueue.Push(*TraitPtr, TraitState);
+					}
+				}
 			}
 		});
 	}
@@ -121,6 +128,13 @@ namespace UE::UAF
 						NumChildren = 1;
 					}
 				}
+				else if (const FTraitPtr* TraitPtr = Impl->GetInlineGraph())
+				{
+					if (TraitPtr->IsValid())
+					{
+						NumChildren = 1;
+					}
+				}
 			}
 		});
 
@@ -142,6 +156,13 @@ namespace UE::UAF
 						Children.Add(GraphInstance->GetGraphRootPtr());
 					}
 				}
+				else if (const FTraitPtr* TraitPtr = Impl->GetInlineGraph())
+				{
+					if (TraitPtr->IsValid())
+					{
+						Children.Add(*TraitPtr);
+					}
+				}
 			}
 		});
 	}
