Strip UProperty tag from compilation related variables (otherwise would have been marked as NonTransactional, but that kind of ends up being the same as plain members)
#Synced-CL 51412342

diff --git a/Engine/Plugins/Experimental/UAF/UAF/Source/UAFUncookedOnly/Internal/AnimNextRigVMAssetEditorData.h b/Engine/Plugins/Experimental/UAF/UAF/Source/UAFUncookedOnly/Internal/AnimNextRigVMAssetEditorData.h
--- a/Engine/Plugins/Experimental/UAF/UAF/Source/UAFUncookedOnly/Internal/AnimNextRigVMAssetEditorData.h
+++ b/Engine/Plugins/Experimental/UAF/UAF/Source/UAFUncookedOnly/Internal/AnimNextRigVMAssetEditorData.h
@@ -714,13 +714,8 @@ protected:
 	UPROPERTY(VisibleAnywhere, Category=Components, AdvancedDisplay, meta = (DisplayName = "Required Components"))
 	TArray<TInstancedStruct<FUAFAssetInstanceComponent>> Components;
 
-	UPROPERTY(transient, DuplicateTransient)
 	int32 VMRecompilationBracket = 0;
-
-	UPROPERTY(transient, DuplicateTransient)
 	bool bVMRecompilationRequired = false;
-
-	UPROPERTY(transient, DuplicateTransient)
 	bool bIsCompiling = false;
 
 	UE_DEPRECATED(5.7, "Storing asset entries as external packages has been deprecated")
@@ -767,7 +762,7 @@ protected:
 
 	UPROPERTY()
 	bool bAutoRecompileVM = true;
-
+	
 	bool bVMRecompilationRequested = false;
 	mutable bool bErrorsDuringCompilation = false;
 	mutable bool bWarningsDuringCompilation = false;
