Guard against UAFCompilationScope being constructed with just a scope description (no asset)

#Synced-CL 51412342

diff --git a/Engine/Plugins/Experimental/UAF/UAF/Source/UAFUncookedOnly/Private/UAFCompilationScope.cpp b/Engine/Plugins/Experimental/UAF/UAF/Source/UAFUncookedOnly/Private/UAFCompilationScope.cpp
--- a/Engine/Plugins/Experimental/UAF/UAF/Source/UAFUncookedOnly/Private/UAFCompilationScope.cpp
+++ b/Engine/Plugins/Experimental/UAF/UAF/Source/UAFUncookedOnly/Private/UAFCompilationScope.cpp
@@ -50,7 +50,7 @@ void Compilation::RequestAssetCompilation(TConstArrayView<UUAFRigVMAsset*> InAss
 }
 
 FCompilationScope::FCompilationScope(const FText& InJobName)
-	: FCompilationScope(InJobName, {})
+	: FCompilationScope(InJobName, TConstArrayView<UUAFRigVMAsset*>())
 {
 }
 
@@ -78,6 +78,7 @@ FCompilationScope::FCompilationScope(TConstArrayView<UUAFRigVMAsset*> InAssets)
 FCompilationScope::FCompilationScope(const FText& InJobName, UUAFRigVMAsset* InAsset)
 	: FCompilationScope(InJobName, {&InAsset, 1})
 {
+	check(InAsset);
 }
 
 FCompilationScope::FCompilationScope(const FText& InJobName, TConstArrayView<UUAFRigVMAsset*> InAssets)
