Crash - Ensure when undoing changes to the graph in existing UAF System assets
#fix make sure we do not transact/setup-undo-redo while building UOL fragments
#Synced-CL 51412342

diff --git a/Engine/Plugins/Experimental/UAF/UAF/Source/UAFUncookedOnly/Private/Variables/UniversalObjectLocatorBindingType.cpp b/Engine/Plugins/Experimental/UAF/UAF/Source/UAFUncookedOnly/Private/Variables/UniversalObjectLocatorBindingType.cpp
--- a/Engine/Plugins/Experimental/UAF/UAF/Source/UAFUncookedOnly/Private/Variables/UniversalObjectLocatorBindingType.cpp
+++ b/Engine/Plugins/Experimental/UAF/UAF/Source/UAFUncookedOnly/Private/Variables/UniversalObjectLocatorBindingType.cpp
@@ -775,7 +775,7 @@ void FUniversalObjectLocatorBindingType::BuildBindingGraphFragment(const FRigVMC
 					continue;
 				}
 
-				if (!InArgs.Controller->ResolveWildCardPin(ValuePin->GetPinPath(), FRigVMRegistry::Get().GetTypeIndex(*Input->CPPType, Input->CPPTypeObject)))
+				if (!InArgs.Controller->ResolveWildCardPin(ValuePin->GetPinPath(), FRigVMRegistry::Get().GetTypeIndex(*Input->CPPType, Input->CPPTypeObject), false))
 				{
 					InSettings.ReportError(TEXT("Could not resolve wildcard pin to function's return type"));
 					continue;
@@ -833,7 +833,7 @@ void FUniversalObjectLocatorBindingType::BuildBindingGraphFragment(const FRigVMC
 					continue;
 				}
 
-				if (!InArgs.Controller->ResolveWildCardPin(ValuePin->GetPinPath(), FRigVMRegistry::Get().GetTypeIndex(*Input->CPPType, Input->CPPTypeObject)))
+				if (!InArgs.Controller->ResolveWildCardPin(ValuePin->GetPinPath(), FRigVMRegistry::Get().GetTypeIndex(*Input->CPPType, Input->CPPTypeObject), false))
 				{
 					InSettings.ReportError(TEXT("Could not resolve wildcard pin to function's return type"));
 					continue;
@@ -891,7 +891,7 @@ void FUniversalObjectLocatorBindingType::BuildBindingGraphFragment(const FRigVMC
 					continue;
 				}
 
-				if (!InArgs.Controller->ResolveWildCardPin(ValuePin->GetPinPath(), FRigVMRegistry::Get().GetTypeIndex(*Input->CPPType, Input->CPPTypeObject)))
+				if (!InArgs.Controller->ResolveWildCardPin(ValuePin->GetPinPath(), FRigVMRegistry::Get().GetTypeIndex(*Input->CPPType, Input->CPPTypeObject), false))
 				{
 					InSettings.ReportError(TEXT("Could not resolve wildcard pin to function's return type"));
 					continue;
@@ -1004,4 +1004,4 @@ const UClass* FUniversalObjectLocatorBindingType::GetClass(TConstStructView<FAni
 
 }
 
-#undef LOCTEXT_NAMESPACE
\ No newline at end of file
+#undef LOCTEXT_NAMESPACE
