Crash when undoing creation of Shared Variables from other UAF asset
#fix guard against attempting to Preload EditorData for undo-ed created SharedVariables object
#Synced-CL 51412342

diff --git a/Engine/Plugins/Experimental/UAF/UAF/Source/UAF/Private/AnimNextRigVMAsset.cpp b/Engine/Plugins/Experimental/UAF/UAF/Source/UAF/Private/AnimNextRigVMAsset.cpp
--- a/Engine/Plugins/Experimental/UAF/UAF/Source/UAF/Private/AnimNextRigVMAsset.cpp
+++ b/Engine/Plugins/Experimental/UAF/UAF/Source/UAF/Private/AnimNextRigVMAsset.cpp
@@ -48,8 +48,11 @@ void UUAFRigVMAsset::Serialize(FArchive& Ar)
 	if (Ar.IsLoading())
 	{
 #if WITH_EDITORONLY_DATA
-		// External editor-side APIs assume that the editordata is also loaded with the asset
-		EditorData->ConditionalPreload();
+		if (EditorData)
+		{
+			// External editor-side APIs assume that the editordata is also loaded with the asset
+			EditorData->ConditionalPreload();
+		}
 #endif
 
 		ExtendedExecuteContext.InvalidateCachedMemory();
