UAF: Don't show Open In New Tab context menu entry outside the functions outliner

#Synced-CL 51454554

diff --git a/Engine/Plugins/Experimental/UAF/UAF/Source/UAFEditor/Private/Functions/Outliner/FunctionsOutlinerMode.cpp b/Engine/Plugins/Experimental/UAF/UAF/Source/UAFEditor/Private/Functions/Outliner/FunctionsOutlinerMode.cpp
--- a/Engine/Plugins/Experimental/UAF/UAF/Source/UAFEditor/Private/Functions/Outliner/FunctionsOutlinerMode.cpp
+++ b/Engine/Plugins/Experimental/UAF/UAF/Source/UAFEditor/Private/Functions/Outliner/FunctionsOutlinerMode.cpp
@@ -425,10 +425,13 @@ void FFunctionsOutlinerMode::RegisterToolMenus()
 		{
 			PopulateNewFunctionToolMenuEntries(InMenu, false);
 
-			if (const UCommonOutlinerItemMenuContext* OutlinerContext = InMenu->FindContext<UCommonOutlinerItemMenuContext>())
+			if (const UFunctionsOutlinerItemMenuContext* FunctionOutlinerContext = InMenu->FindContext<UFunctionsOutlinerItemMenuContext>())
 			{
-				FToolMenuSection& Section = InMenu->AddSection("Navigation", LOCTEXT("FunctionNavigationLabel", "Navigation"), FToolMenuInsert(FCommonOutlinerMode::ContextToolEntrySectionName, EToolMenuInsertType::Before));
-				Section.AddMenuEntryWithCommandList(FOutlinerCommands::Get().OpenInNewTab, OutlinerContext->WeakCommandList.Pin());
+				if (const UCommonOutlinerItemMenuContext* OutlinerContext = InMenu->FindContext<UCommonOutlinerItemMenuContext>())
+				{
+					FToolMenuSection& Section = InMenu->AddSection("Navigation", LOCTEXT("FunctionNavigationLabel", "Navigation"), FToolMenuInsert(FCommonOutlinerMode::ContextToolEntrySectionName, EToolMenuInsertType::Before));
+					Section.AddMenuEntryWithCommandList(FOutlinerCommands::Get().OpenInNewTab, OutlinerContext->WeakCommandList.Pin());
+				}
 			}
 		}));
 	}
