RigVM: Fix fallback from Variable Guids, renaming a variable in UAF is not refreshing variable nodes
#rnx

#Synced-CL 51458724

diff --git a/Engine/Plugins/Experimental/UAF/UAF/Source/UAFUncookedOnly/Private/UncookedOnlyUtils.cpp b/Engine/Plugins/Experimental/UAF/UAF/Source/UAFUncookedOnly/Private/UncookedOnlyUtils.cpp
--- a/Engine/Plugins/Experimental/UAF/UAF/Source/UAFUncookedOnly/Private/UncookedOnlyUtils.cpp
+++ b/Engine/Plugins/Experimental/UAF/UAF/Source/UAFUncookedOnly/Private/UncookedOnlyUtils.cpp
@@ -2087,10 +2087,12 @@ void FUtils::ReplaceVariableReferences(UUAFRigVMAssetEditorData* InEditorData, c
 	const UUAFRigVMAsset* FindOuterAsset = Cast<UUAFRigVMAsset>(ReferenceToFind.GetObject());
 	const FAnimNextVariableReference ReferenceToReplaceWith(SoftReferenceToReplaceWith);
 	const UUAFRigVMAsset* ReplaceOuterAsset = Cast<UUAFRigVMAsset>(ReferenceToReplaceWith.GetObject());
-	const FGuid ReferenceGuid = UE::UAF::UncookedOnly::IAnimNextUncookedOnlyModule::Get().GetVariableGuidByName(ReferenceToFind.GetName(), FindOuterAsset);
 
 	const bool bReplacingWithNone = ReferenceToReplaceWith.IsNone();
-
+	const FGuid ReferenceGuid = bReplacingWithNone  
+		? UE::UAF::UncookedOnly::IAnimNextUncookedOnlyModule::Get().GetVariableGuidByName(ReferenceToFind.GetName(), FindOuterAsset)
+		: UE::UAF::UncookedOnly::IAnimNextUncookedOnlyModule::Get().GetVariableGuidByName(ReferenceToReplaceWith.GetName(), FindOuterAsset);
+	
 	// Replacing variable ref local to this asset, with another local variable or none
 	const bool bInternalReplace = (ThisOuterAsset == FindOuterAsset) && (FindOuterAsset == ReplaceOuterAsset || bReplacingWithNone);
 	
