UAF Component GetVariable

#Synced-CL 51464644

diff --git a/Engine/Plugins/Experimental/UAF/UAF/Source/UAF/Private/Component/AnimNextComponent.cpp b/Engine/Plugins/Experimental/UAF/UAF/Source/UAF/Private/Component/AnimNextComponent.cpp
--- a/Engine/Plugins/Experimental/UAF/UAF/Source/UAF/Private/Component/AnimNextComponent.cpp
+++ b/Engine/Plugins/Experimental/UAF/UAF/Source/UAF/Private/Component/AnimNextComponent.cpp
@@ -232,6 +232,11 @@ void UUAFComponent::BlueprintSetVariableReference(const FAnimNextVariableReferen
 	checkNoEntry();
 }
 
+void UUAFComponent::BlueprintGetVariableReference(const FAnimNextVariableReference& Variable, int32& Value)
+{
+	checkNoEntry();
+}
+
 DEFINE_FUNCTION(UUAFComponent::execBlueprintSetVariableReference)
 {
 	using namespace UE::UAF;
@@ -328,16 +333,154 @@ DEFINE_FUNCTION(UUAFComponent::execBlueprintSetVariableReference)
 	}
 }
 
+DEFINE_FUNCTION(UUAFComponent::execBlueprintGetVariableReference)
+{
+	using namespace UE::UAF;
+
+	P_GET_STRUCT_REF(FAnimNextVariableReference, Variable);
+
+	const FProperty* VariableProperty = Variable.ResolveProperty();
+	if (VariableProperty == nullptr)
+	{
+		P_FINISH;
+
+		FBlueprintExceptionInfo ExceptionInfo(
+			EBlueprintExceptionType::AbortExecution,
+			FText::Format(NSLOCTEXT("UAFComponent", "UAFComponent_GetVariablePropertyError", "Failed to resolve the property {0} of variable reference"), FText::FromName(Variable.GetName()))
+		);
+
+		FBlueprintCoreDelegates::ThrowScriptException(P_THIS, Stack, ExceptionInfo);
+		return;
+	}
+	
+	if (Variable.GetName() == NAME_None || Variable.GetObject() == nullptr)
+	{
+		P_FINISH;
+
+		FBlueprintExceptionInfo ExceptionInfo(
+			EBlueprintExceptionType::NonFatalError,
+			NSLOCTEXT("UAFComponent", "UAFComponent_GetVariableInvalidWarning", "Invalid variable supplied to Get Variable")
+		);
+
+		FBlueprintCoreDelegates::ThrowScriptException(P_THIS, Stack, ExceptionInfo);
+		return;
+	}
+
+	// Read wildcard Value input.
+	Stack.MostRecentPropertyAddress = nullptr;
+	Stack.MostRecentPropertyContainer = nullptr;
+	Stack.StepCompiledIn<FProperty>(nullptr);
+
+	const FProperty* TargetProperty = Stack.MostRecentProperty;
+	void* TargetPtr = Stack.MostRecentPropertyAddress;
+
+	P_FINISH;
+	
+	if (!TargetProperty || !TargetPtr)
+	{
+		FBlueprintExceptionInfo ExceptionInfo(
+			EBlueprintExceptionType::AbortExecution,
+			NSLOCTEXT("UAFComponent", "UAFComponent_GetVariableInvalidValueWarning", "Failed to resolve the value for Get Variable")
+		);
+		
+		FBlueprintCoreDelegates::ThrowScriptException(P_THIS, Stack, ExceptionInfo);
+		return;
+	}
+
+	if (!TargetProperty->SameType(VariableProperty))
+	{
+		FBlueprintExceptionInfo ExceptionInfo(
+			EBlueprintExceptionType::AbortExecution,
+			NSLOCTEXT("UAFComponent", "UAFComponent_GetVariablePropertyTypeMismatchWarning", "Pin property type does not match variable property type")
+		);
+		
+		FBlueprintCoreDelegates::ThrowScriptException(P_THIS, Stack, ExceptionInfo);
+		return;
+	}
+
+	EPropertyBagResult Result = EPropertyBagResult::Success;
+	
+	P_NATIVE_BEGIN;
+	{
+		FAnimNextParamType Type = FAnimNextParamType::FromProperty(VariableProperty);
+
+		Result = P_THIS->ReadVariableInternal(Variable, Type, [TargetProperty, TargetPtr](TConstArrayView<uint8> VariableValue)
+			{
+				check(TargetProperty->GetElementSize() == VariableValue.Num());
+				TargetProperty->CopyCompleteValue(TargetPtr, VariableValue.GetData());
+			});
+	}
+	P_NATIVE_END;
+
+	switch (Result)
+	{
+	case EPropertyBagResult::Success:
+		{
+			break;
+		}
+	case EPropertyBagResult::TypeMismatch:
+		{
+			FBlueprintExceptionInfo ExceptionInfo(
+				EBlueprintExceptionType::NonFatalError,
+				FText::Format(NSLOCTEXT("UAFComponent", "UAFComponent_GetVariableTypeMismatch", "Incompatible type supplied for variable '{0}'"), FText::FromName(Variable.GetName()))
+			);
+
+			FBlueprintCoreDelegates::ThrowScriptException(P_THIS, Stack, ExceptionInfo);
+			break;
+		}
+	case EPropertyBagResult::PropertyNotFound:
+		{
+			FBlueprintExceptionInfo ExceptionInfo(
+				EBlueprintExceptionType::NonFatalError,
+				FText::Format(NSLOCTEXT("UAFComponent", "UAFComponent_GetVariableNotFoundWarning", "Unknown variable name '{0}' supplied to Get Variable"), FText::FromName(Variable.GetName()))
+			);
+
+			FBlueprintCoreDelegates::ThrowScriptException(P_THIS, Stack, ExceptionInfo);
+			break;
+		}
+	case EPropertyBagResult::OutOfBounds:
+	case EPropertyBagResult::DuplicatedValue:
+		{
+			FBlueprintExceptionInfo ExceptionInfo(
+				EBlueprintExceptionType::NonFatalError,
+				FText::Format(NSLOCTEXT("UAFComponent", "UAFComponent_GetVariableUnknownWarning", "Unexpected internal error when calling Get Variable for '{0}'"), FText::FromName(Variable.GetName()))
+			);
+
+			FBlueprintCoreDelegates::ThrowScriptException(P_THIS, Stack, ExceptionInfo);
+			break;
+		}
+	}
+}
+
 EPropertyBagResult UUAFComponent::SetVariableInternal(const FAnimNextVariableReference& InVariable, const FAnimNextParamType& InType, TConstArrayView<uint8> InNewValue)
 {
 	return SystemReference.SetVariable(InVariable, InType, InNewValue);
 }
 
+EPropertyBagResult UUAFComponent::GetVariableInternal(const FAnimNextVariableReference& InVariable, const FAnimNextParamType& InType, TArrayView<uint8> OutNewValue)
+{
+	const FProperty* Property = InVariable.ResolveProperty();
+	if (Property)
+	{
+		return SystemReference.ReadVariable(InVariable, InType, [&](TConstArrayView<uint8> VariableValue)
+		{
+			check(OutNewValue.Num() == VariableValue.Num());
+			Property->CopyCompleteValue(OutNewValue.GetData(), VariableValue.GetData());
+		});
+	}
+	return EPropertyBagResult::PropertyNotFound;
+}
+
 EPropertyBagResult UUAFComponent::WriteVariableInternal(const FAnimNextVariableReference& InVariable, const FAnimNextParamType& InType, TFunctionRef<void(TArrayView<uint8>)> InFunction)
 {
 	return SystemReference.WriteVariable(InVariable, InType, InFunction);
 }
 
+EPropertyBagResult UUAFComponent::ReadVariableInternal(const FAnimNextVariableReference& InVariable, const FAnimNextParamType& InType, TFunctionRef<void(TConstArrayView<uint8>)> InFunction)
+{
+	return SystemReference.ReadVariable(InVariable, InType, InFunction);
+}
+
 void UUAFComponent::Activate(bool bReset)
 {
 	if (bReset || ShouldActivate())
diff --git a/Engine/Plugins/Experimental/UAF/UAF/Source/UAF/Private/Module/SystemReference.cpp b/Engine/Plugins/Experimental/UAF/UAF/Source/UAF/Private/Module/SystemReference.cpp
--- a/Engine/Plugins/Experimental/UAF/UAF/Source/UAF/Private/Module/SystemReference.cpp
+++ b/Engine/Plugins/Experimental/UAF/UAF/Source/UAF/Private/Module/SystemReference.cpp
@@ -458,7 +458,7 @@ EPropertyBagResult FSystemReference::SetVariable(const FAnimNextVariableReferenc
 	}
 	return EPropertyBagResult::PropertyNotFound;
 }
-
+		
 EPropertyBagResult FSystemReference::WriteVariable(const FAnimNextVariableReference& InVariable, const FAnimNextParamType& InType, TFunctionRef<void(TArrayView<uint8>)> InFunction) const
 {
 	using namespace UE::UAF;
@@ -469,6 +469,15 @@ EPropertyBagResult FSystemReference::WriteVariable(const FAnimNextVariableRefere
 	return EPropertyBagResult::PropertyNotFound;
 }
 
+EPropertyBagResult FSystemReference::ReadVariable(const FAnimNextVariableReference& InVariable, const FAnimNextParamType& InType, TFunctionRef<void(TConstArrayView<uint8>)> InFunction) const
+{
+	if (IsValid())
+	{
+		return Ptr->ReadVariableFromExternal(InVariable, InType, InFunction);
+	}
+	return EPropertyBagResult::PropertyNotFound;
+}
+	
 void FSystemReference::RunEvent(FName InEventName, float InDeltaTime) const
 {
 	using namespace UE::UAF;
diff --git a/Engine/Plugins/Experimental/UAF/UAF/Source/UAF/Public/Component/AnimNextComponent.h b/Engine/Plugins/Experimental/UAF/UAF/Source/UAF/Public/Component/AnimNextComponent.h
--- a/Engine/Plugins/Experimental/UAF/UAF/Source/UAF/Public/Component/AnimNextComponent.h
+++ b/Engine/Plugins/Experimental/UAF/UAF/Source/UAF/Public/Component/AnimNextComponent.h
@@ -94,20 +94,30 @@ class UUAFComponent : public UActorComponent, public UE::UAF::ISystemOutputAdapt
 	virtual void SignalOutputWritten(const UE::UAF::FModuleTaskContext& InContext) override;
 	
 public:
-	// Sets a module variable's value.
+	// Sets a system variable's value.
+	// Note that variables are double-buffered so writing and reading the same variable won't return the new value unless UAF is updated in between.
 	// @param    Name     The name of the variable to set
 	// @param    Value    The value to set the variable to
 	UE_DEPRECATED(5.6, "This function is no longer used, please use SetVariableReference (for Blueprint) or SetVariable that takes a FAnimNextVariableReference")
 	UFUNCTION(BlueprintCallable, Category = "UAF", DisplayName = "Set Variable", CustomThunk, meta = (CustomStructureParam = Value, UnsafeDuringActorConstruction, DeprecatedFunction, DeprecationMessage = "This function has been deprecated, please use Set Variable that takes a variable reference instead"))
 	UAF_API void BlueprintSetVariable(FName Name, int32 Value);
 
-	// Sets a module variable's value.
-	// @param    Name     The variable to set
+	// Sets a system variable's value.
+	// Note that variables are double-buffered so writing and reading the same variable won't return the new value unless UAF is updated in between.
+	// @param    Variable The variable to set
 	// @param    Value    The value to set the variable to
 	UFUNCTION(BlueprintCallable, BlueprintInternalUseOnly, Category = "UAF", DisplayName = "Set Variable", CustomThunk, meta = (CustomStructureParam = Value, AutoCreateRefTerm = "Variable, Value", UnsafeDuringActorConstruction))
 	UAF_API void BlueprintSetVariableReference(const FAnimNextVariableReference& Variable, const int32& Value);
 
+	// Gets a system variable's value.
+	// Note that variables are double-buffered so writing and reading the same variable won't return the new value unless UAF is updated in between.
+	// @param    Variable The variable to get
+	// @param    Value    The return value of the variable
+	UFUNCTION(BlueprintCallable, BlueprintInternalUseOnly, Category = "UAF", DisplayName = "Get Variable", CustomThunk, meta = (CustomStructureParam = Value, AutoCreateRefTerm = "Variable, Value", UnsafeDuringActorConstruction))
+	UAF_API void BlueprintGetVariableReference(const FAnimNextVariableReference& Variable, int32& Value);
+	
 	// Set a variable.
+	// Note that variables are double-buffered so writing and reading the same variable won't return the new value unless UAF is updated in between.
 	// @param    InVariable         The variable to set
 	// @param    InValue            The value to set the variable to
 	// @return true if the variable was set correctly
@@ -117,6 +127,17 @@ public:
 		return SetVariableInternal(InVariable, FAnimNextParamType::GetType<ValueType>(), TConstArrayView<uint8>(reinterpret_cast<const uint8*>(&InValue), sizeof(ValueType))) == EPropertyBagResult::Success;
 	}
 
+	// Get a variable.
+	// Note that variables are double-buffered so writing and reading the same variable won't return the new value unless UAF is updated in between.
+	// @param    InVariable         The variable to get
+	// @param    OutValue           The return value of the variable
+	// @return true if the variable was gotten correctly
+	template <typename ValueType>
+	bool GetVariable(const FAnimNextVariableReference& InVariable, ValueType& OutValue)
+	{
+		return GetVariableInternal(InVariable, FAnimNextParamType::GetType<ValueType>(), TArrayView<uint8>(reinterpret_cast<uint8*>(&OutValue), sizeof(ValueType))) == EPropertyBagResult::Success;
+	}
+
 	// Accesses a variable for writing.
 	// General read access is not provided via this API due to the current double-buffering strategy used to communicate variable writes to worker threads.
 	// This is intended to allow for copy-free writing of larger data structures & arrays, rather than read access. 
@@ -246,13 +267,20 @@ private:
 
 	DECLARE_FUNCTION(execBlueprintSetVariable);
 	DECLARE_FUNCTION(execBlueprintSetVariableReference);
+	DECLARE_FUNCTION(execBlueprintGetVariableReference);
 
 	// Sets the value of the specified variable
 	UAF_API EPropertyBagResult SetVariableInternal(const FAnimNextVariableReference& InVariable, const FAnimNextParamType& InType, TConstArrayView<uint8> InNewValue);
 
+	// Gets the value of the specified variable, expects OutNewValue to be correctly sized
+	UAF_API EPropertyBagResult GetVariableInternal(const FAnimNextVariableReference& InVariable, const FAnimNextParamType& InType, TArrayView<uint8> OutNewValue);
+	
 	// Accesses value of the specified variable for writing
 	UAF_API EPropertyBagResult WriteVariableInternal(const FAnimNextVariableReference& InVariable, const FAnimNextParamType& InType, TFunctionRef<void(TArrayView<uint8>)> InFunction);
 
+	// Accesses value of the specified variable for reading
+	UAF_API EPropertyBagResult ReadVariableInternal(const FAnimNextVariableReference& InVariable, const FAnimNextParamType& InType, TFunctionRef<void(TConstArrayView<uint8>)> InFunction);
+
 	// Determine what skeletal mesh component to write to
 	USkeletalMeshComponent* DetermineSkeletalMeshComponent() const;
 
diff --git a/Engine/Plugins/Experimental/UAF/UAF/Source/UAF/Public/Module/SystemReference.h b/Engine/Plugins/Experimental/UAF/UAF/Source/UAF/Public/Module/SystemReference.h
--- a/Engine/Plugins/Experimental/UAF/UAF/Source/UAF/Public/Module/SystemReference.h
+++ b/Engine/Plugins/Experimental/UAF/UAF/Source/UAF/Public/Module/SystemReference.h
@@ -199,6 +199,17 @@ private:
 	// @return see EPropertyBagResult
 	UE_API EPropertyBagResult WriteVariable(const FAnimNextVariableReference& InVariable, const FAnimNextParamType& InType, TFunctionRef<void(TArrayView<uint8>)> InFunction) const;
 
+	// Access a variable's value for read only.
+	// This performs non-recursive locks on the system instance to ensure that concurrent access is respected, so care must be taken calling this from
+	// animation worker tasks as deadlocks can occur.
+	// Type must match strictly, no conversions are performed.
+	// Can be called on any thread.
+	// @param	InVariable			The variable to get the value of
+	// @param	InType				The type of the variable
+	// @param	InFunction			Function that will be called if no errors occur
+	// @return see EPropertyBagResult
+	UE_API EPropertyBagResult ReadVariable(const FAnimNextVariableReference& InVariable, const FAnimNextParamType& InType, TFunctionRef<void(TConstArrayView<uint8>)> InFunction) const;
+	
 	// Access the supplied system component for reading.
 	// Can be called on any thread.
 	// @param	InFunction			Function that will be called to access the struct
diff --git a/Engine/Plugins/Experimental/UAF/UAF/Source/UAFUncookedOnly/Internal/Variables/K2Node_UAFComponentGetVariable.h b/Engine/Plugins/Experimental/UAF/UAF/Source/UAFUncookedOnly/Internal/Variables/K2Node_UAFComponentGetVariable.h
new file mode 100644
--- /dev/null
+++ b/Engine/Plugins/Experimental/UAF/UAF/Source/UAFUncookedOnly/Internal/Variables/K2Node_UAFComponentGetVariable.h
@@ -0,0 +1,35 @@
+// Copyright Epic Games, Inc. All Rights Reserved.
+
+#pragma once
+
+#include "K2Node_CallFunction.h"
+
+#include "K2Node_UAFComponentGetVariable.generated.h"
+
+UCLASS(MinimalAPI)
+class UK2Node_UAFComponentGetVariable : public UK2Node_CallFunction
+{
+	GENERATED_BODY()
+
+public:
+	UAFUNCOOKEDONLY_API UK2Node_UAFComponentGetVariable();
+
+private:
+	// //~ Begin UEdGraphNode Interface.
+	UAFUNCOOKEDONLY_API virtual void ReconstructNode() override;
+	UAFUNCOOKEDONLY_API virtual void PinDefaultValueChanged(UEdGraphPin* Pin) override;
+	// //~ End UEdGraphNode Interface.
+
+	//~ Begin K2Node Interface
+	UAFUNCOOKEDONLY_API virtual void GetMenuActions(FBlueprintActionDatabaseRegistrar& InActionRegistrar) const override;
+	UAFUNCOOKEDONLY_API virtual void ValidateNodeDuringCompilation(FCompilerResultsLog& MessageLog) const override;
+	UAFUNCOOKEDONLY_API virtual void PreloadRequiredAssets() override;
+	//~ End K2Node Interface
+
+	void CachePinType();
+	void PropagateCachedTypeToPin();
+
+	// Cached type of the pin used for regenerating after wildcard is reconstructed
+	UPROPERTY()
+	FEdGraphPinType CachedPinType;
+};
\ No newline at end of file
diff --git a/Engine/Plugins/Experimental/UAF/UAF/Source/UAFUncookedOnly/Private/Variables/K2Node_UAFComponentGetVariable.cpp b/Engine/Plugins/Experimental/UAF/UAF/Source/UAFUncookedOnly/Private/Variables/K2Node_UAFComponentGetVariable.cpp
new file mode 100644
--- /dev/null
+++ b/Engine/Plugins/Experimental/UAF/UAF/Source/UAFUncookedOnly/Private/Variables/K2Node_UAFComponentGetVariable.cpp
@@ -0,0 +1,146 @@
+// Copyright Epic Games, Inc. All Rights Reserved.
+
+#include "Variables/K2Node_UAFComponentGetVariable.h"
+
+#include "AnimNextRigVMAsset.h"
+#include "BlueprintActionDatabaseRegistrar.h"
+#include "BlueprintNodeSpawner.h"
+#include "K2Node_Literal.h"
+#include "KismetCompiler.h"
+#include "UncookedOnlyUtils.h"
+#include "Kismet2/BlueprintEditorUtils.h"
+#include "Component/AnimNextComponent.h"
+#include "Param/ParamCompatibility.h"
+#include "Param/ParamUtils.h"
+#include "Variables/AnimNextVariableReference.h"
+
+#include UE_INLINE_GENERATED_CPP_BY_NAME(K2Node_UAFComponentGetVariable)
+
+#define LOCTEXT_NAMESPACE "K2Node_UAFComponentGetVariable"
+
+UK2Node_UAFComponentGetVariable::UK2Node_UAFComponentGetVariable()
+{
+	FunctionReference.SetExternalMember(GET_FUNCTION_NAME_CHECKED(UUAFComponent, BlueprintGetVariableReference), UUAFComponent::StaticClass());
+}
+
+void UK2Node_UAFComponentGetVariable::ReconstructNode()
+{
+	Super::ReconstructNode();
+
+	PropagateCachedTypeToPin();
+}
+
+void UK2Node_UAFComponentGetVariable::PinDefaultValueChanged(UEdGraphPin* Pin)
+{
+	Super::PinDefaultValueChanged(Pin);
+
+	if (Pin == FindPinChecked(TEXT("Variable")))
+	{
+		CachePinType();
+		PropagateCachedTypeToPin();
+	}
+}
+
+void UK2Node_UAFComponentGetVariable::CachePinType()
+{
+	// Cache the pin type if it has been set as a literal type via default value
+	UEdGraphPin* VariablePin = FindPinChecked(TEXT("Variable"));
+	if (VariablePin->LinkedTo.Num() == 0)
+	{
+		FAnimNextVariableReference VariableReference;
+		FAnimNextVariableReference::StaticStruct()->ImportText(*VariablePin->DefaultValue, &VariableReference, nullptr, PPF_None, nullptr, FAnimNextVariableReference::StaticStruct()->GetName());
+		if (VariableReference.IsValid())
+		{
+			FAnimNextParamType Type = UE::UAF::UncookedOnly::FUtils::FindVariableType(VariableReference);
+			if (Type.IsValid())
+			{
+				CachedPinType = UE::UAF::UncookedOnly::FUtils::GetPinTypeFromParamType(Type);
+			}
+		}
+	}
+}
+
+void UK2Node_UAFComponentGetVariable::PropagateCachedTypeToPin()
+{
+	// Find the value of the variable pin
+	UEdGraphPin* VariablePin = FindPinChecked(TEXT("Variable"));
+	UEdGraphPin* ValuePin = FindPinChecked(TEXT("Value"));
+	if (!CachedPinType.PinCategory.IsNone() && VariablePin->LinkedTo.Num() == 0 && ValuePin->LinkedTo.Num() == 0)
+	{
+		ValuePin->PinType = CachedPinType;
+		if (!ValuePin->GetSchema()->IsCurrentPinDefaultValid(ValuePin).IsEmpty())
+		{
+			ValuePin->GetSchema()->ResetPinToAutogeneratedDefaultValue(ValuePin);
+		}
+		GetGraph()->NotifyNodeChanged(this);
+
+		UBlueprint* Blueprint = GetBlueprint();
+		if (!Blueprint->bBeingCompiled)
+		{
+			FBlueprintEditorUtils::MarkBlueprintAsModified(Blueprint);
+		}
+	}
+}
+
+void UK2Node_UAFComponentGetVariable::ValidateNodeDuringCompilation(FCompilerResultsLog& MessageLog) const
+{
+	Super::ValidateNodeDuringCompilation(MessageLog);
+
+	const UK2Node* SourceNode = Cast<UK2Node>(MessageLog.FindSourceObject(this));
+	if (SourceNode == nullptr)
+	{
+		return;
+	}
+	
+	UEdGraphPin* VariablePin = SourceNode->FindPinChecked(TEXT("Variable"));
+	UEdGraphPin* ValuePin = SourceNode->FindPinChecked(TEXT("Value"));
+
+	FAnimNextVariableReference VariableReference;
+	FAnimNextVariableReference::StaticStruct()->ImportText(*VariablePin->DefaultValue, &VariableReference, nullptr, PPF_None, nullptr, FAnimNextVariableReference::StaticStruct()->GetName());
+	if (VariableReference.IsValid())
+	{
+		FAnimNextParamType Type = UE::UAF::UncookedOnly::FUtils::FindVariableType(VariableReference);				// The variable from the UAF Module
+		FAnimNextParamType PinType = UE::UAF::UncookedOnly::FUtils::GetParamTypeFromPinType(ValuePin->PinType);	// The variable from BP pin
+		UE::UAF::FParamCompatibility Compatibility = UE::UAF::FParamUtils::GetCompatibility(Type, PinType);
+		if (Compatibility == UE::UAF::EParamCompatibility::Incompatible || Compatibility == UE::UAF::EParamCompatibility::Incompatible_DataLoss)
+		{
+			MessageLog.Error(*FString::Printf(TEXT("@@ mismatched type, setting '%s' type, but variable is '%s' type"), *PinType.ToString(), *Type.ToString()), this);
+		}
+	}
+	else if (VariableReference.IsNone())
+	{
+		MessageLog.Error(TEXT("@@ variable reference not set"), this);
+	}
+	else
+	{
+		MessageLog.Error(*FString::Printf(TEXT("@@ variable reference '%s' was not found in '%s'"), *VariableReference.GetName().ToString(), VariableReference.GetObject() ? TEXT("None") : *VariableReference.GetObject()->GetPathName()), this);
+	}
+}
+
+void UK2Node_UAFComponentGetVariable::GetMenuActions(FBlueprintActionDatabaseRegistrar& InActionRegistrar) const
+{
+	const UClass* ActionKey = GetClass();
+	if (InActionRegistrar.IsOpenForRegistration(ActionKey))
+	{
+		UBlueprintNodeSpawner* NodeSpawner = UBlueprintNodeSpawner::Create(GetClass());
+		check(NodeSpawner != nullptr);
+
+		InActionRegistrar.AddBlueprintAction(ActionKey, NodeSpawner);
+	}
+}
+
+void UK2Node_UAFComponentGetVariable::PreloadRequiredAssets()
+{
+	Super::PreloadRequiredAssets();
+
+	UEdGraphPin* VariablePin = FindPinChecked(TEXT("Variable"));
+	FAnimNextVariableReference VariableReference;
+	FAnimNextVariableReference::StaticStruct()->ImportText(*VariablePin->DefaultValue, &VariableReference, nullptr, PPF_None, nullptr, FAnimNextVariableReference::StaticStruct()->GetName());
+	if (!VariableReference.IsNone() && VariableReference.GetObject())
+	{
+		UObject* Object = const_cast<UObject*>(VariableReference.GetObject().Get());
+		PreloadObject(Object);
+	}
+}
+
+#undef LOCTEXT_NAMESPACE
