Fix for crash when disabling UAF using cvar

#Synced-CL 51455675

diff --git a/Engine/Plugins/Experimental/UAF/UAF/Source/UAF/Private/RemapPoseData.cpp b/Engine/Plugins/Experimental/UAF/UAF/Source/UAF/Private/RemapPoseData.cpp
--- a/Engine/Plugins/Experimental/UAF/UAF/Source/UAF/Private/RemapPoseData.cpp
+++ b/Engine/Plugins/Experimental/UAF/UAF/Source/UAF/Private/RemapPoseData.cpp
@@ -12,9 +12,9 @@ DEFINE_STAT(STAT_AnimNext_RemapPose_Mesh2Mesh);
 bool FRemapPoseData::ShouldReinit(const UE::UAF::FReferencePose& InSourceRefPose,
 	const UE::UAF::FReferencePose& InTargetRefPose) const
 {
-	if (!SourceRefPose || !TargetRefPose)
+	if (&InSourceRefPose != SourceRefPose || &InTargetRefPose != TargetRefPose)
 	{
-		// First time use.
+		// First time use or dangling pointers.
 		return true;
 	}
 
