- Added TOptional<FAnimNextParamType> to soft variable reference. (allows for referencing a variable by its previous type)
- Added FAnimNextSoftVariableReference::FromNameAndType
- Added GetGuid to ProgrammaticVariable struct (as its now used in two places)
- Use this new optional in the variable replace behaviour, currently a variable type change was not correctly captured (as the types matched)
- Updated new FExternalVariable constructions to use correct (programmatic)variable guids

#Synced-CL 51454995

diff --git a/Engine/Plugins/Experimental/UAF/UAF/Source/UAF/Public/Variables/AnimNextSoftVariableReference.h b/Engine/Plugins/Experimental/UAF/UAF/Source/UAF/Public/Variables/AnimNextSoftVariableReference.h
--- a/Engine/Plugins/Experimental/UAF/UAF/Source/UAF/Public/Variables/AnimNextSoftVariableReference.h
+++ b/Engine/Plugins/Experimental/UAF/UAF/Source/UAF/Public/Variables/AnimNextSoftVariableReference.h
@@ -4,6 +4,7 @@
 
 #include "AnimNextVariableReference.h"
 #include "UObject/SoftObjectPath.h"
+#include "Param/ParamType.h"
 
 class UUAFRigVMAsset;
 struct FAnimNextVariableReference;
@@ -16,9 +17,16 @@ struct FAnimNextSoftVariableReference
 	// Construct a soft reference from a FName
 	static FAnimNextSoftVariableReference FromName(FName InName, const FSoftObjectPath& InSoftObjectPath)
 	{
-		return FAnimNextSoftVariableReference(InName, InSoftObjectPath);
+		return FAnimNextSoftVariableReference(InName, InSoftObjectPath, TOptional<FAnimNextParamType>());
 	}
 
+#if WITH_EDITOR
+	static FAnimNextSoftVariableReference FromNameAndType(FName InName, const FSoftObjectPath& InSoftObjectPath, FAnimNextParamType InType)
+	{
+		return FAnimNextSoftVariableReference(InName, InSoftObjectPath, InType);
+	}
+#endif // WITH_EDITOR
+
 	explicit FAnimNextSoftVariableReference(const FAnimNextVariableReference& VariableReference) : Name(VariableReference.GetName()), SoftObjectPath(VariableReference.GetObject()) 
 	{
 		ensureAlwaysMsgf(!UE::UAF::IsInternalVariableName(Name), TEXT("Attempting to create soft variable reference using a variable name indicative of a programmatic variable %s"), *Name.ToString());
@@ -47,11 +55,24 @@ struct FAnimNextSoftVariableReference
 	{
 		Name = NAME_None; 
 		SoftObjectPath.Reset();
+		OptionalType.Reset();
+	}
+	
+	// Whether parameter type was set
+	bool HasType() const
+	{
+		return OptionalType.IsSet();
+	}
+	
+	// Return reference parameter type, only when set otherwise will assert (use HasType for checking validity)
+	FAnimNextParamType GetType() const
+	{
+		return OptionalType.GetValue();
 	}
 
 	bool operator==(const FAnimNextSoftVariableReference& InOther) const
 	{
-		return Name == InOther.Name && SoftObjectPath == InOther.SoftObjectPath;
+		return Name == InOther.Name && SoftObjectPath == InOther.SoftObjectPath && OptionalType == InOther.OptionalType;
 	}
 
 	bool operator!=(const FAnimNextSoftVariableReference& InOther) const
@@ -61,12 +82,13 @@ struct FAnimNextSoftVariableReference
 
 	friend uint32 GetTypeHash(const FAnimNextSoftVariableReference& InKey)
 	{
-		return HashCombineFast(GetTypeHash(InKey.Name), GetTypeHash(InKey.SoftObjectPath));
+		return HashCombineFast(GetTypeHash(InKey.Name), GetTypeHash(InKey.SoftObjectPath), GetTypeHash(InKey.OptionalType));
 	}
 private:
-	FAnimNextSoftVariableReference(FName InName, const FSoftObjectPath& InSoftObjectPath)
+	FAnimNextSoftVariableReference(FName InName, const FSoftObjectPath& InSoftObjectPath, const TOptional<FAnimNextParamType> InType = TOptional<FAnimNextParamType>())
 		   : Name(InName)
 		   , SoftObjectPath(InSoftObjectPath)
+		   , OptionalType(InType)
 	{}
 
 private:
@@ -75,5 +97,8 @@ private:
 
 	// The asset or struct that the variable reference is contained in
 	FSoftObjectPath SoftObjectPath;
+	
+	// (Expected) type when used for searching references
+	TOptional<FAnimNextParamType> OptionalType;
 };
 
diff --git a/Engine/Plugins/Experimental/UAF/UAF/Source/UAFUncookedOnly/Private/UncookedOnlyUtils.cpp b/Engine/Plugins/Experimental/UAF/UAF/Source/UAFUncookedOnly/Private/UncookedOnlyUtils.cpp
--- a/Engine/Plugins/Experimental/UAF/UAF/Source/UAFUncookedOnly/Private/UncookedOnlyUtils.cpp
+++ b/Engine/Plugins/Experimental/UAF/UAF/Source/UAFUncookedOnly/Private/UncookedOnlyUtils.cpp
@@ -159,8 +159,7 @@ FInstancedPropertyBag FUtils::MakePropertyBagForEditorData(const UUAFRigVMAssetE
 			EAnimNextExportAccessSpecifier::Private,
 			TConstArrayView<uint8>(Variable.GetValuePtr(), Variable.Type.GetSize()),
 			CPF_AdvancedDisplay | CPF_NativeAccessSpecifierPrivate,
-			// [TODO] this assumes Variable.Name being unique yet consistent across compilations
-			FGuid::NewDeterministicGuid(Variable.Name.ToString(), FCrc::StrCrc32<TCHAR>(TEXT("ProgrammaticVariableGUID")))
+			Variable.GetGuid()
 		);
 	}
 
@@ -552,7 +551,7 @@ void FUtils::GetInternalVariables(const UUAFRigVMAssetEditorData* EditorData, TA
 	TArray<FRigVMExternalVariable> AssetVariables;
 	EditorData->ForEachEntryOfType<UAnimNextVariableEntry>([AssetName = EditorData->GetAssetName(), &AssetVariables, &VariableInfo](const UAnimNextVariableEntry* Entry)
 	{
-		const FRigVMExternalVariable Variable = FRigVMExternalVariable::Make(FGuid(), Entry->GetVariableName(), Entry->GetType().ToString(), const_cast<UObject*>(Entry->GetType().GetValueTypeObject()), Entry->Access == EAnimNextExportAccessSpecifier::Public, false);
+		const FRigVMExternalVariable Variable = FRigVMExternalVariable::Make(Entry->GetGuid(), Entry->GetVariableName(), Entry->GetType().ToString(), const_cast<UObject*>(Entry->GetType().GetValueTypeObject()), Entry->Access == EAnimNextExportAccessSpecifier::Public, false);
 		if (!Variable.IsValid(true))
 		{
 			UE_LOG(LogRigVM, Error, TEXT("%s: Property '%s' of type '%s' is not valid."), *AssetName, *Entry->GetName(), *Variable.GetExtendedCPPType().ToString());
@@ -573,7 +572,7 @@ void FUtils::GetInternalVariables(const UUAFRigVMAssetEditorData* EditorData, TA
 		for (const FAnimNextProgrammaticVariable& ProgrammaticVariable : GetVariableCompileContext.GetProgrammaticVariables())
 		{
 			VariableInfo.Add({ProgrammaticVariable.GetVariableName(), ProgrammaticVariable.GetType()});
-			AssetVariables.Add(FRigVMExternalVariable::Make(FGuid(), ProgrammaticVariable.GetVariableName(), ProgrammaticVariable.GetType().ToString(), const_cast<UObject*>(ProgrammaticVariable.GetType().GetValueTypeObject()), false, false));
+			AssetVariables.Add(FRigVMExternalVariable::Make(ProgrammaticVariable.GetGuid(), ProgrammaticVariable.GetVariableName(), ProgrammaticVariable.GetType().ToString(), const_cast<UObject*>(ProgrammaticVariable.GetType().GetValueTypeObject()), false, false));
 		}
 	}
 
@@ -2049,7 +2048,8 @@ void FUtils::SetVariableType(UAnimNextVariableEntry* VariableEntry, const FAnimN
 
 		FCompilationScope CompilerResults(LOCTEXT("ModifiedAssets_SetVariableType", "Modified Assets Set Variable Type"), OuterRigVMAsset);
 		
-		const FAnimNextSoftVariableReference FindReference = FAnimNextSoftVariableReference::FromName(VariableEntry->GetEntryName(), OuterRigVMAsset);
+		const FAnimNextSoftVariableReference OldReference = FAnimNextSoftVariableReference::FromNameAndType(VariableEntry->GetEntryName(), OuterRigVMAsset, VariableEntry->GetType());
+		const FAnimNextSoftVariableReference NewReference = FAnimNextSoftVariableReference::FromNameAndType(VariableEntry->GetEntryName(), OuterRigVMAsset, NewType);
 		
 		// Change variable type first
 		VariableEntry->SetType(NewType);
@@ -2059,12 +2059,12 @@ void FUtils::SetVariableType(UAnimNextVariableEntry* VariableEntry, const FAnimN
 		if (bIsAPublicVariable)
 		{
 			constexpr bool bForceLoadReferencingAssets = false;
-			ReplaceVariableReferencesAcrossProject(FindReference, FindReference, bForceLoadReferencingAssets, bSetupUndoRedo, bPrintPythonCommands);
+			ReplaceVariableReferencesAcrossProject(OldReference, NewReference, bForceLoadReferencingAssets, bSetupUndoRedo, bPrintPythonCommands);
 		}
 		else
 		{
 			// Otherwise replace, potential, local references only 
-			ReplaceVariableReferences(GetEditorData(OuterRigVMAsset), FindReference, FindReference, bSetupUndoRedo, bPrintPythonCommands);
+			ReplaceVariableReferences(GetEditorData(OuterRigVMAsset), OldReference, NewReference, bSetupUndoRedo, bPrintPythonCommands);
 		}
 	}	
 }
@@ -2099,8 +2099,8 @@ void FUtils::ReplaceVariableReferences(UUAFRigVMAssetEditorData* InEditorData, c
 	// Replacing with variable from a different object vs existing reference
 	const bool bDifferentVariableSource = FindOuterAsset != ReplaceOuterAsset;
 
-	FAnimNextParamType FindReferenceType = FindVariableType(ReferenceToFind);
-	const FAnimNextParamType ReplaceReferenceType = FindVariableType(ReferenceToReplaceWith);
+	FAnimNextParamType FindReferenceType = SoftReferenceToFind.HasType() ? SoftReferenceToFind.GetType() : FindVariableType(ReferenceToFind);
+	const FAnimNextParamType ReplaceReferenceType = SoftReferenceToReplaceWith.HasType() ? SoftReferenceToReplaceWith.GetType() : FindVariableType(ReferenceToReplaceWith);
 
 	// In case ReferenceToFind has already been removed/changed match it up with the replacing type (if not None) 
 	if (!FindReferenceType.IsValid() && !ReferenceToReplaceWith.IsNone())
diff --git a/Engine/Plugins/Experimental/UAF/UAF/Source/UAFUncookedOnly/Private/Variables/AnimNextProgrammaticVariable.cpp b/Engine/Plugins/Experimental/UAF/UAF/Source/UAFUncookedOnly/Private/Variables/AnimNextProgrammaticVariable.cpp
--- a/Engine/Plugins/Experimental/UAF/UAF/Source/UAFUncookedOnly/Private/Variables/AnimNextProgrammaticVariable.cpp
+++ b/Engine/Plugins/Experimental/UAF/UAF/Source/UAFUncookedOnly/Private/Variables/AnimNextProgrammaticVariable.cpp
@@ -100,6 +100,12 @@ const uint8* FAnimNextProgrammaticVariable::GetValuePtr() const
 	return PropertyBagStruct && Memory ? PropertyBagStruct->GetPropertyDescs()[0].CachedProperty->ContainerPtrToValuePtr<uint8>(Memory) : nullptr;
 }
 
+FGuid FAnimNextProgrammaticVariable::GetGuid() const
+{
+	// [TODO] this assumes Variable.Name being unique yet consistent across compilations
+	return FGuid::NewDeterministicGuid(Name.ToString(), FCrc::StrCrc32<TCHAR>(TEXT("ProgrammaticVariableGUID")));
+}
+
 FAnimNextProgrammaticVariable FAnimNextProgrammaticVariable::FromRigVMGraphFunctionArgument(const FRigVMGraphFunctionArgument& RigVMGraphFunctionArgument)
 {
 	FAnimNextProgrammaticVariable Result = {};
diff --git a/Engine/Plugins/Experimental/UAF/UAF/Source/UAFUncookedOnly/Public/Variables/AnimNextProgrammaticVariable.h b/Engine/Plugins/Experimental/UAF/UAF/Source/UAFUncookedOnly/Public/Variables/AnimNextProgrammaticVariable.h
--- a/Engine/Plugins/Experimental/UAF/UAF/Source/UAFUncookedOnly/Public/Variables/AnimNextProgrammaticVariable.h
+++ b/Engine/Plugins/Experimental/UAF/UAF/Source/UAFUncookedOnly/Public/Variables/AnimNextProgrammaticVariable.h
@@ -46,7 +46,9 @@ struct FAnimNextProgrammaticVariable
 
 	/** Access the memory for the internal value */
 	UE_API const uint8* GetValuePtr() const;
-
+	
+	/** Returns deterministic GUID */
+	UE_API FGuid GetGuid() const;
 public:
 
 	/** Name of the variable */
