UAF: Fix for variable nodes not resolving correctly in the details panel

#Synced-CL 51494526

diff --git a/Engine/Plugins/Experimental/UAF/UAF/Source/UAFEditor/Private/AnimNextEditorModule.cpp b/Engine/Plugins/Experimental/UAF/UAF/Source/UAFEditor/Private/AnimNextEditorModule.cpp
--- a/Engine/Plugins/Experimental/UAF/UAF/Source/UAFEditor/Private/AnimNextEditorModule.cpp
+++ b/Engine/Plugins/Experimental/UAF/UAF/Source/UAFEditor/Private/AnimNextEditorModule.cpp
@@ -1258,6 +1258,39 @@ void FAnimNextEditorModule::RegisterWorkspaceDocumentTypes(Workspace::IWorkspace
 						}
 					}
 				}
+				else if (const URigVMVariableNode* VariableNode = Cast<URigVMVariableNode>(RigVMEdGraphNode->GetModelNode()))
+				{
+					UUAFRigVMAssetEditorData* EditorData = nullptr;
+					if (const UUAFSharedVariableNode* SharedVarNode = Cast<UUAFSharedVariableNode>(VariableNode))
+					{
+						EditorData = SharedVarNode->Asset ? UncookedOnly::FUtils::GetEditorData<UUAFRigVMAssetEditorData>(SharedVarNode->Asset.Get()) : nullptr;
+					}
+					else
+					{
+						EditorData = VariableNode->GetTypedOuter<UUAFRigVMAssetEditorData>();
+					}
+					
+					if (EditorData)
+					{
+						UAnimNextVariableEntry* Entry = nullptr;
+						EditorData->ForEachEntryOfType<UAnimNextVariableEntry>([InGuid = VariableNode->GetVariableGuid(), &Entry](UAnimNextVariableEntry* VariableEntry)
+						{
+							if (VariableEntry->GetGuid() == InGuid)
+							{
+								Entry = VariableEntry;
+								return false;
+							}
+						
+							return true;
+						});
+
+						if (Entry)
+						{
+							ObjectsToSelect.Add(Entry);
+							continue;
+						}
+					}
+				}
 			}
 			else if (const UEdGraphNode* Node = Cast<UEdGraphNode>(Object))
 			{
