UAF: Allow renaming of a variable via the details panel and move parameter name validation into util function since it is used in multiple places now

#Synced-CL 51494981

diff --git a/Engine/Plugins/Experimental/UAF/UAF/Source/UAFEditor/Internal/EditorUtils.h b/Engine/Plugins/Experimental/UAF/UAF/Source/UAFEditor/Internal/EditorUtils.h
--- a/Engine/Plugins/Experimental/UAF/UAF/Source/UAFEditor/Internal/EditorUtils.h
+++ b/Engine/Plugins/Experimental/UAF/UAF/Source/UAFEditor/Internal/EditorUtils.h
@@ -18,6 +18,7 @@ struct FRigVMFunction;
 struct FGraphContextMenuBuilder;
 
 struct FAnimNextAssetRegistryExports;
+class UAnimNextVariableEntry;
 
 namespace UE::UAF::Editor
 {
@@ -35,6 +36,8 @@ struct FUtils
 	static UAFEDITOR_API bool IsValidParameterNameString(FStringView InStringView, FText& OutErrorText);
 
 	static UAFEDITOR_API bool IsValidParameterName(const FName InName, FText& OutErrorText);
+	
+	static UAFEDITOR_API bool IsValidParameterNameForEntry(const UAnimNextVariableEntry* Entry, const FText& InNewName, FText& OutErrorText);
 
 	static UAFEDITOR_API bool AddSchemaRigUnitAction(const TSubclassOf<URigVMUnitNode>& UnitNodeClass, UScriptStruct* Struct, const FRigVMFunction& Function, FGraphContextMenuBuilder& InContexMenuBuilder);
 
diff --git a/Engine/Plugins/Experimental/UAF/UAF/Source/UAFEditor/Private/EditorUtils.cpp b/Engine/Plugins/Experimental/UAF/UAF/Source/UAFEditor/Private/EditorUtils.cpp
--- a/Engine/Plugins/Experimental/UAF/UAF/Source/UAFEditor/Private/EditorUtils.cpp
+++ b/Engine/Plugins/Experimental/UAF/UAF/Source/UAFEditor/Private/EditorUtils.cpp
@@ -360,6 +360,40 @@ bool FUtils::IsValidParameterName(const FName InName, FText& OutErrorText)
 	return true;
 }
 
+bool FUtils::IsValidParameterNameForEntry(const UAnimNextVariableEntry* Entry, const FText& InNewName, FText& OutErrorText)
+{
+	if(Entry == nullptr)
+	{
+		OutErrorText = LOCTEXT("InvalidVariableError", "Variable is invalid");
+		return false;
+	}
+
+	const UUAFRigVMAssetEditorData* EditorData = Entry->GetTypedOuter<UUAFRigVMAssetEditorData>();
+	if(EditorData == nullptr)
+	{
+		OutErrorText = LOCTEXT("InvalidEditorDataError", "Variable has invalid editor data");
+		return false;
+	}
+
+	const FString NewString = InNewName.ToString();
+	if (!IsValidParameterNameString(NewString, OutErrorText))
+	{
+		return false;
+	}
+
+	const FName Name(*NewString);
+		
+	// Cannot have two entries with matching names
+	UUAFRigVMAssetEntry* ExistingEntry = EditorData->FindEntry(Name);
+	if(ExistingEntry && ExistingEntry != Entry)
+	{
+		OutErrorText = LOCTEXT("NameExistsError", "Variable name already exists in this asset");
+		return false;
+	}
+		
+	return true;
+}
+
 bool FUtils::AddSchemaRigUnitAction(const TSubclassOf<URigVMUnitNode>& UnitNodeClass, UScriptStruct* Struct, const FRigVMFunction& Function, FGraphContextMenuBuilder& InContexMenuBuilder)
 {
 	FString CategoryMetadata, DisplayNameMetadata, MenuDescSuffixMetadata;
diff --git a/Engine/Plugins/Experimental/UAF/UAF/Source/UAFEditor/Private/Variables/Outliner/VariablesOutlinerEntryItem.cpp b/Engine/Plugins/Experimental/UAF/UAF/Source/UAFEditor/Private/Variables/Outliner/VariablesOutlinerEntryItem.cpp
--- a/Engine/Plugins/Experimental/UAF/UAF/Source/UAFEditor/Private/Variables/Outliner/VariablesOutlinerEntryItem.cpp
+++ b/Engine/Plugins/Experimental/UAF/UAF/Source/UAFEditor/Private/Variables/Outliner/VariablesOutlinerEntryItem.cpp
@@ -91,36 +91,7 @@ void FVariablesOutlinerEntryItem::Rename(const FText& InNewName) const
 
 bool FVariablesOutlinerEntryItem::ValidateName(const FText& InNewName, FText& OutErrorMessage) const
 {
-	UAnimNextVariableEntry* Entry = WeakEntry.Get();
-	if(Entry == nullptr)
-	{
-		OutErrorMessage = LOCTEXT("InvalidVariableError", "Variable is invalid");
-		return false;
-	}
-
-	UUAFRigVMAssetEditorData* EditorData = Entry->GetTypedOuter<UUAFRigVMAssetEditorData>();
-	if(EditorData == nullptr)
-	{
-		OutErrorMessage = LOCTEXT("InvalidEditorDataError", "Variable has invalid editor data");
-		return false;
-	}
-
-	const FString NewString = InNewName.ToString();
-	if (!Editor::FUtils::IsValidParameterNameString(NewString, OutErrorMessage))
-	{
-		return false;
-	}
-
-	FName Name(*NewString);
-	// Cannot have two entries with matching names
-	UUAFRigVMAssetEntry* ExistingEntry = EditorData->FindEntry(Name);
-	if(ExistingEntry && ExistingEntry != Entry)
-	{
-		OutErrorMessage = LOCTEXT("NameExistsError", "Variable name already exists in this asset");
-		return false;
-	}
-
-	return true;
+	return FUtils::IsValidParameterNameForEntry(WeakEntry.Get(), InNewName, OutErrorMessage);
 }
 
 void FVariablesOutlinerEntryItem::SetAccessSpecifier(const EAnimNextExportAccessSpecifier& InSpecifier) const
diff --git a/Engine/Plugins/Experimental/UAF/UAF/Source/UAFEditor/Private/Variables/VariableCustomization.cpp b/Engine/Plugins/Experimental/UAF/UAF/Source/UAFEditor/Private/Variables/VariableCustomization.cpp
--- a/Engine/Plugins/Experimental/UAF/UAF/Source/UAFEditor/Private/Variables/VariableCustomization.cpp
+++ b/Engine/Plugins/Experimental/UAF/UAF/Source/UAFEditor/Private/Variables/VariableCustomization.cpp
@@ -5,6 +5,7 @@
 #include "DetailLayoutBuilder.h"
 #include "DetailWidgetRow.h"
 #include "EdGraphSchema_K2.h"
+#include "EditorUtils.h"
 #include "InstancedPropertyBagStructureDataProvider.h"
 #include "UncookedOnlyUtils.h"
 #include "PropertyHandle.h"
@@ -188,6 +189,24 @@ void FVariableCustomization::CustomizeDetails(IDetailLayoutBuilder& DetailBuilde
 				Handle->SetOnChildPropertyValueChangedWithData(TDelegate<void(const FPropertyChangedEvent&)>::CreateLambda(OnPropertyValueChange));
 			}
 		}
+		
+		
+		// Handle renaming of the parameter name and validating that name before commit 
+		const TSharedPtr<IPropertyHandle> ParameterNameProperty = DetailBuilder.GetProperty(GET_MEMBER_NAME_CHECKED(UAnimNextVariableEntry, ParameterName));
+		IDetailPropertyRow& ParameterNameRow = DetailBuilder.AddPropertyToCategory(ParameterNameProperty);	
+		ParameterNameRow.CustomWidget()
+		.NameContent()
+		[
+			ParameterNameProperty->CreatePropertyNameWidget()
+		]
+		.ValueContent()
+		[
+			SNew(SEditableTextBox)
+			.OnVerifyTextChanged(this, &FVariableCustomization::OnVerifyNameChanged)
+			.Text(this, &FVariableCustomization::GetName)
+			.OnTextCommitted(this, &FVariableCustomization::SetName)
+		
+		];
 	}
 }
 
@@ -199,16 +218,32 @@ void FVariableCustomization::CustomizeDetails(const TSharedPtr<IDetailLayoutBuil
 
 FText FVariableCustomization::GetName() const
 {
+	// Only allow name interaction in single select case 
+	if (Variables.Num() == 1 && Variables[0].IsValid())
+	{
+		return Variables[0]->GetDisplayName();
+	}
 	return FText();
 }
 
 void FVariableCustomization::SetName(const FText& InNewText, ETextCommit::Type InCommitType)
 {
+	// Only allow name interaction in single select case 
+	if (Variables.Num() == 1 && Variables[0].IsValid())
+	{
+		UE::UAF::UncookedOnly::FUtils::RenameVariable(Variables[0].Get(), FName(InNewText.ToString()));
+	}
 }
 
 bool FVariableCustomization::OnVerifyNameChanged(const FText& InText, FText& OutErrorMessage)
 {
-	return true;
+	// Only allow name interaction in single select case 
+	if (Variables.Num() == 1 && Variables[0].IsValid())
+	{
+		return FUtils::IsValidParameterNameForEntry(Variables[0].Get(), InText, OutErrorMessage);
+	}
+	
+	return false;
 }
 
 FText FVariableCustomization::GetCategory() const
