Add context menu item to set binding editor to include children into the selection

#Synced-CL 51464644

diff --git a/Engine/Plugins/Experimental/UAF/UAF/Source/UAFEditor/Private/UAF/AbstractSkeleton/Sets/Commands.cpp b/Engine/Plugins/Experimental/UAF/UAF/Source/UAFEditor/Private/UAF/AbstractSkeleton/Sets/Commands.cpp
--- a/Engine/Plugins/Experimental/UAF/UAF/Source/UAFEditor/Private/UAF/AbstractSkeleton/Sets/Commands.cpp
+++ b/Engine/Plugins/Experimental/UAF/UAF/Source/UAFEditor/Private/UAF/AbstractSkeleton/Sets/Commands.cpp
@@ -11,6 +11,7 @@ namespace UE::UAF::Editor
 	void FSetBindingEditorCommands::RegisterCommands()
 	{
 		UI_COMMAND(UnbindSelection, "Unbind Selection", "Unbinds the current selection", EUserInterfaceActionType::Button, FInputChord(EKeys::BackSpace));
+		UI_COMMAND(SelectWithChildren, "Select With Children", "Sets the current selection to this item and its children", EUserInterfaceActionType::Button, FInputChord(EModifierKey::Control | EModifierKey::Shift, EKeys::A));
 	}
 }
 
diff --git a/Engine/Plugins/Experimental/UAF/UAF/Source/UAFEditor/Private/UAF/AbstractSkeleton/Sets/Commands.h b/Engine/Plugins/Experimental/UAF/UAF/Source/UAFEditor/Private/UAF/AbstractSkeleton/Sets/Commands.h
--- a/Engine/Plugins/Experimental/UAF/UAF/Source/UAFEditor/Private/UAF/AbstractSkeleton/Sets/Commands.h
+++ b/Engine/Plugins/Experimental/UAF/UAF/Source/UAFEditor/Private/UAF/AbstractSkeleton/Sets/Commands.h
@@ -27,6 +27,7 @@ namespace UE::UAF::Editor
 		virtual void RegisterCommands() override;
 
 		TSharedPtr<FUICommandInfo> UnbindSelection;
+		TSharedPtr<FUICommandInfo> SelectWithChildren;
 	};
 
 }
\ No newline at end of file
diff --git a/Engine/Plugins/Experimental/UAF/UAF/Source/UAFEditor/Private/UAF/AbstractSkeleton/Sets/SSetSkeletonTree.cpp b/Engine/Plugins/Experimental/UAF/UAF/Source/UAFEditor/Private/UAF/AbstractSkeleton/Sets/SSetSkeletonTree.cpp
--- a/Engine/Plugins/Experimental/UAF/UAF/Source/UAFEditor/Private/UAF/AbstractSkeleton/Sets/SSetSkeletonTree.cpp
+++ b/Engine/Plugins/Experimental/UAF/UAF/Source/UAFEditor/Private/UAF/AbstractSkeleton/Sets/SSetSkeletonTree.cpp
@@ -257,14 +257,15 @@ namespace UE::UAF::Editor
 		FMenuBuilder MenuBuilder(true, CommandList.ToSharedRef());
 
 		MenuBuilder.AddMenuEntry(Commands.UnbindSelection);
+		MenuBuilder.AddMenuEntry(Commands.SelectWithChildren);
 
 		return MenuBuilder.MakeWidget();
 	}
 
-	TArray<TSharedPtr<SSetsSkeletonTree::FTreeItem>> SSetsSkeletonTree::GetAllTreeItems()
+	TArray<TSharedPtr<SSetsSkeletonTree::FTreeItem>> SSetsSkeletonTree::GetAllDescendentTreeItems(const TArray<TSharedPtr<FTreeItem>> ParentItems)
 	{
 		TArray<TSharedPtr<FTreeItem>> AllItems;
-		AllItems.Append(RootItems);
+		AllItems.Append(ParentItems);
 
 		for (int32 Index = 0; Index < AllItems.Num(); ++Index)
 		{
@@ -295,7 +296,7 @@ namespace UE::UAF::Editor
 	{
 		// Save expansion state of each tree item to restore later
 		TSet<FName> CollapsedBoneNames;
-		for (const TSharedPtr<FTreeItem>& TreeItem : GetAllTreeItems())
+		for (const TSharedPtr<FTreeItem>& TreeItem : GetAllDescendentTreeItems(RootItems))
 		{
 			if (!TreeView->IsItemExpanded(TreeItem))
 			{
@@ -337,7 +338,7 @@ namespace UE::UAF::Editor
 
 		TreeView->RebuildList();
 
-		for (TSharedPtr<FTreeItem> TreeItem : GetAllTreeItems())
+		for (TSharedPtr<FTreeItem> TreeItem : GetAllDescendentTreeItems(RootItems))
 		{
 			const bool bExpanded = !CollapsedBoneNames.Contains(TreeItem->BoneName);
 			TreeView->SetItemExpansion(TreeItem, bExpanded);
@@ -353,6 +354,10 @@ namespace UE::UAF::Editor
 		CommandList->MapAction(Commands.UnbindSelection,
 			FExecuteAction::CreateSP(this, &SSetsSkeletonTree::HandleUnbindSelection),
 			FCanExecuteAction::CreateSP(this, &SSetsSkeletonTree::IsSelectionNonEmpty));
+
+		CommandList->MapAction(Commands.SelectWithChildren,
+			FExecuteAction::CreateSP(this, &SSetsSkeletonTree::HandleSelectWithChildren),
+			FCanExecuteAction::CreateSP(this, &SSetsSkeletonTree::IsSelectionSingle));
 	}
 
 	void SSetsSkeletonTree::HandleBindingsChanged()
@@ -373,10 +378,26 @@ namespace UE::UAF::Editor
 		}
 	}
 
+	void SSetsSkeletonTree::HandleSelectWithChildren()
+	{
+		// No transaction as selection is transient
+
+		TArray<TSharedPtr<FTreeItem>> Selection = TreeView->GetSelectedItems();
+		check(Selection.Num() == 1);
+
+		TArray<TSharedPtr<FTreeItem>> NewSelection = GetAllDescendentTreeItems(Selection);
+		TreeView->SetItemSelection(NewSelection, true);
+	}
+
 	bool SSetsSkeletonTree::IsSelectionNonEmpty() const
 	{
 		return TreeView->GetNumItemsSelected() > 0;
 	}
+
+	bool SSetsSkeletonTree::IsSelectionSingle() const
+	{
+		return TreeView->GetNumItemsSelected() == 1;
+	}
 }
 
 #undef LOCTEXT_NAMESPACE
\ No newline at end of file
diff --git a/Engine/Plugins/Experimental/UAF/UAF/Source/UAFEditor/Private/UAF/AbstractSkeleton/Sets/SSetSkeletonTree.h b/Engine/Plugins/Experimental/UAF/UAF/Source/UAFEditor/Private/UAF/AbstractSkeleton/Sets/SSetSkeletonTree.h
--- a/Engine/Plugins/Experimental/UAF/UAF/Source/UAFEditor/Private/UAF/AbstractSkeleton/Sets/SSetSkeletonTree.h
+++ b/Engine/Plugins/Experimental/UAF/UAF/Source/UAFEditor/Private/UAF/AbstractSkeleton/Sets/SSetSkeletonTree.h
@@ -51,10 +51,16 @@ namespace UE::UAF::Editor
 
 		void HandleUnbindSelection();
 
-		bool IsSelectionNonEmpty() const;
+		void HandleSelectWithChildren();
 
-		TArray<TSharedPtr<FTreeItem>> GetAllTreeItems();
+		bool IsSelectionNonEmpty() const;
+		
+		bool IsSelectionSingle() const;
 
+		// Returns a list of the provided parents and all of their descendents
+		// Beware, the descendents of ParentItems should be non-overlapping!
+		TArray<TSharedPtr<FTreeItem>> GetAllDescendentTreeItems(const TArray<TSharedPtr<FTreeItem>> ParentItems);
+		
 		TSharedRef<ITableRow> TreeView_OnGenerateRow(TSharedPtr<FTreeItem> Item, const TSharedRef<STableViewBase>& OwnerTable);
 		
 		void TreeView_OnGetChildren(TSharedPtr<FTreeItem> InParent, TArray<TSharedPtr<FTreeItem>>& OutChildren);
diff --git a/Engine/Plugins/Experimental/UAF/UAF/Source/UAFEditor/Private/UAFStyle.cpp b/Engine/Plugins/Experimental/UAF/UAF/Source/UAFEditor/Private/UAFStyle.cpp
--- a/Engine/Plugins/Experimental/UAF/UAF/Source/UAFEditor/Private/UAFStyle.cpp
+++ b/Engine/Plugins/Experimental/UAF/UAF/Source/UAFEditor/Private/UAFStyle.cpp
@@ -115,6 +115,7 @@ void FUAFStyle::SetSetBindingStyles()
 	const FVector2D Icon16x16(16.0f, 16.0f);
 
 	Set("SetBindingEditorCommands.UnbindSelection", new CORE_IMAGE_BRUSH_SVG("Starship/Common/Unlinked", Icon16x16));
+	Set("SetBindingEditorCommands.SelectWithChildren", new IMAGE_PLUGIN_BRUSH_SVG("Icons/SelectAll_16.svg", Icon16x16));
 }
 
 void FUAFStyle::SetBlendProfileStyles()
