Update UAF GetVariable to first check the proxy write buffer for staged writes

#Synced-CL 51464644

diff --git a/Engine/Plugins/Experimental/UAF/UAF/Source/UAF/Internal/Variables/UAFInstanceVariableDataProxy.h b/Engine/Plugins/Experimental/UAF/UAF/Source/UAF/Internal/Variables/UAFInstanceVariableDataProxy.h
--- a/Engine/Plugins/Experimental/UAF/UAF/Source/UAF/Internal/Variables/UAFInstanceVariableDataProxy.h
+++ b/Engine/Plugins/Experimental/UAF/UAF/Source/UAF/Internal/Variables/UAFInstanceVariableDataProxy.h
@@ -50,12 +50,18 @@ struct FUAFInstanceVariableDataProxy
 	// Sets the value of the specified variable. If the type does not match exactly then a conversion will be attempted.
 	EPropertyBagResult SetVariable(const FAnimNextVariableReference& InVariable, const FAnimNextParamType& InType, TConstArrayView<uint8> InNewValue);
 
+	// Access the memory of the variable directly. Type must match strictly, no conversions are performed.
+	EPropertyBagResult AccessVariable(const FAnimNextVariableReference& InVariable, const FAnimNextParamType& InType, TArrayView<uint8>& OutValue);
+
 	// Access the memory of the variables set directly for writing. All variables in the set are dirtied.
 	bool WriteVariables(const UScriptStruct* InStruct, TFunctionRef<void(FStructView)> InFunction);
 
 	// Copies any dirty variables over to the corresponding Variables 
 	void CopyDirty();
 
+	// Returns whether the variable has been marked as dirty
+	bool IsVariableDirty(const FAnimNextVariableReference& InVariable) const;
+
 	// Variables that this proxy corresponds to
 	const FUAFInstanceVariableData* Variables = nullptr;
 
diff --git a/Engine/Plugins/Experimental/UAF/UAF/Source/UAF/Private/Module/AnimNextModuleInstance.cpp b/Engine/Plugins/Experimental/UAF/UAF/Source/UAF/Private/Module/AnimNextModuleInstance.cpp
--- a/Engine/Plugins/Experimental/UAF/UAF/Source/UAF/Private/Module/AnimNextModuleInstance.cpp
+++ b/Engine/Plugins/Experimental/UAF/UAF/Source/UAF/Private/Module/AnimNextModuleInstance.cpp
@@ -572,6 +572,27 @@ EPropertyBagResult FAnimNextModuleInstance::WriteProxyVariable(const FAnimNextVa
 
 EPropertyBagResult FAnimNextModuleInstance::ReadVariableFromExternal(const FAnimNextVariableReference& InVariable, const FAnimNextParamType& InType, TFunctionRef<void(TConstArrayView<uint8>)> InFunction)
 {
+	// First check if this variable has a dirty value in the write proxy buffer.
+	// If so, an external system has set this value before it has been copied to the live variables so we should return it.
+	// If not, continue reading the value from the live variables.
+	{
+		UE::TReadScopeLock ReadLock(ProxyLock);
+		FUAFInstanceVariableDataProxy& Proxy = ProxyVariables[ProxyWriteIndex];
+		if (Proxy.bIsDirty)
+		{
+			if (Proxy.IsVariableDirty(InVariable))
+			{
+				TArrayView<uint8> Data;
+				EPropertyBagResult Result = Proxy.AccessVariable(InVariable, InType, Data);
+				if (Result == EPropertyBagResult::Success)
+				{
+					InFunction(Data);
+					return EPropertyBagResult::Success;
+				}
+			}
+		}
+	}
+
 	UE::TReadScopeLock Lock(SystemReferenceLock);
 
 	TArrayView<uint8> Data;
diff --git a/Engine/Plugins/Experimental/UAF/UAF/Source/UAF/Private/Variables/UAFInstanceVariableDataProxy.cpp b/Engine/Plugins/Experimental/UAF/UAF/Source/UAF/Private/Variables/UAFInstanceVariableDataProxy.cpp
--- a/Engine/Plugins/Experimental/UAF/UAF/Source/UAF/Private/Variables/UAFInstanceVariableDataProxy.cpp
+++ b/Engine/Plugins/Experimental/UAF/UAF/Source/UAF/Private/Variables/UAFInstanceVariableDataProxy.cpp
@@ -101,6 +101,23 @@ EPropertyBagResult FUAFInstanceVariableDataProxy::SetVariable(const FAnimNextVar
 	return Result;
 }
 
+EPropertyBagResult FUAFInstanceVariableDataProxy::AccessVariable(const FAnimNextVariableReference& InVariable, const FAnimNextParamType& InType, TArrayView<uint8>& OutValue)
+{
+	if (Variables == nullptr)
+	{
+		// Not yet initialized
+		return EPropertyBagResult::PropertyNotFound;
+	}
+
+	const FUAFInstanceVariableData::FVariableMapping* FoundMapping =  Variables->FindMapping(InVariable);
+	if (FoundMapping == nullptr)
+	{
+		return EPropertyBagResult::PropertyNotFound;
+	}
+
+	return ProxyVariableSets[FoundMapping->AllSharedVariableContainersIndex].AccessVariable(FoundMapping->VariableIndex, InType, OutValue);
+}
+
 bool FUAFInstanceVariableDataProxy::WriteVariables(const UScriptStruct* InStruct, TFunctionRef<void(FStructView)> InFunction)
 {
 	if(Variables == nullptr)
@@ -204,3 +221,22 @@ void FUAFInstanceVariableDataProxy::CopyDirty()
 		bIsDirty = false;
 	}
 }
+
+bool FUAFInstanceVariableDataProxy::IsVariableDirty(const FAnimNextVariableReference& InVariable) const
+{
+	if (Variables == nullptr)
+	{
+		return false;
+	}
+
+	const FUAFInstanceVariableData::FVariableMapping* FoundMapping = Variables->FindMapping(InVariable);
+	if (FoundMapping == nullptr)
+	{
+		return false;
+	}
+
+	const FUAFInstanceVariableContainerProxy& VariableSet = ProxyVariableSets[FoundMapping->AllSharedVariableContainersIndex];
+	const bool bIsVariableDirty = VariableSet.DirtyFlags[FoundMapping->VariableIndex];
+
+	return bIsVariableDirty;
+}
\ No newline at end of file
