Update attribute binding data to use skeletal mesh

#Synced-CL 51464644

diff --git a/Engine/Plugins/Experimental/UAF/UAF/Source/UAF/Private/UAF/AbstractSkeleton/AbstractSkeletonSetBinding.cpp b/Engine/Plugins/Experimental/UAF/UAF/Source/UAF/Private/UAF/AbstractSkeleton/AbstractSkeletonSetBinding.cpp
--- a/Engine/Plugins/Experimental/UAF/UAF/Source/UAF/Private/UAF/AbstractSkeleton/AbstractSkeletonSetBinding.cpp
+++ b/Engine/Plugins/Experimental/UAF/UAF/Source/UAF/Private/UAF/AbstractSkeleton/AbstractSkeletonSetBinding.cpp
@@ -3,6 +3,7 @@
 #include "UAF/AbstractSkeleton/AbstractSkeletonSetBinding.h"
 
 #include "UAF/AbstractSkeleton/AbstractSkeletonSetCollection.h"
+#include "UAF/Attributes/AttributeBindingDataCache.h"
 #include "UAF/Attributes/EngineAttributes.h"
 
 #include UE_INLINE_GENERATED_CPP_BY_NAME(AbstractSkeletonSetBinding)
@@ -15,8 +16,6 @@ void UAbstractSkeletonSetBinding::PostLoad()
 	FilterOrphanedBindings();
 #endif
 
-	RebuildBindingData();
-
 #if WITH_EDITOR
 	if (Skeleton)
 	{
@@ -254,11 +253,6 @@ const TConstArrayView<FAbstractSkeleton_AttributeBinding> UAbstractSkeletonSetBi
 	return AttributeBindings;
 }
 
-const UE::UAF::FAttributeBindingDataPtr& UAbstractSkeletonSetBinding::GetAttributeBindingData() const
-{
-	return AttributeBindingData;
-}
-
 void UAbstractSkeletonSetBinding::OnBindingsChanged()
 {
 	RebuildBindingData();
@@ -271,16 +265,7 @@ void UAbstractSkeletonSetBinding::OnBindingsChanged()
 
 void UAbstractSkeletonSetBinding::RebuildBindingData()
 {
-	// Clear our old data
-	AttributeBindingData = nullptr;
-
-	if (!Skeleton)
-	{
-		return;
-	}
-
-	// Build fresh data
-	AttributeBindingData = UE::UAF::MakeAttributeBindingData(*this);
+	UE::UAF::GAttributeBindingDataCache.ResetSetBinding(this);
 }
 
 #if WITH_EDITOR
diff --git a/Engine/Plugins/Experimental/UAF/UAF/Source/UAF/Private/UAF/Attributes/AttributeBindingData.cpp b/Engine/Plugins/Experimental/UAF/UAF/Source/UAF/Private/UAF/Attributes/AttributeBindingData.cpp
--- a/Engine/Plugins/Experimental/UAF/UAF/Source/UAF/Private/UAF/Attributes/AttributeBindingData.cpp
+++ b/Engine/Plugins/Experimental/UAF/UAF/Source/UAF/Private/UAF/Attributes/AttributeBindingData.cpp
@@ -3,6 +3,7 @@
 #include "UAF/Attributes/AttributeBindingData.h"
 
 #include "Animation/BuiltInAttributeTypes.h"
+#include "Engine/SkeletalMesh.h"
 #include "UAF/AbstractSkeleton/AbstractSkeletonSetBinding.h"
 #include "UAF/AbstractSkeleton/AbstractSkeletonSetCollection.h"
 #include "UAF/Attributes/EngineAttributes.h"
@@ -104,15 +105,12 @@ namespace UE::UAF
 		}
 
 		// Returns a partition per LOD
-		static TArray<FLODPartition> PartitionAttributes(const UAbstractSkeletonSetBinding& Binding, const TArray<FName>& SetNames, const TArray<UScriptStruct*>& AttributeTypes)
+		static TArray<FLODPartition> PartitionAttributes(TNonNullPtr<const UAbstractSkeletonSetBinding> Binding, const FReferenceSkeleton& RefSkeleton, const TArray<FName>& SetNames, const TArray<UScriptStruct*>& AttributeTypes)
 		{
-			const USkeleton* Skeleton = Binding.GetSkeleton();
-			const UAbstractSkeletonSetCollection* SetCollection = Binding.GetSetCollection();
+			const UAbstractSkeletonSetCollection* SetCollection = Binding->GetSetCollection();
 
-			const FReferenceSkeleton& RefSkeleton = Skeleton->GetReferenceSkeleton();
-
-			const TConstArrayView<FAbstractSkeleton_BoneBinding> BoneBindings = Binding.GetBoneBindings();
-			const TConstArrayView<FAbstractSkeleton_AttributeBinding> AttributeBindings = Binding.GetAttributeBindings();
+			const TConstArrayView<FAbstractSkeleton_BoneBinding> BoneBindings = Binding->GetBoneBindings();
+			const TConstArrayView<FAbstractSkeleton_AttributeBinding> AttributeBindings = Binding->GetAttributeBindings();
 
 			// TODO: Hook up LODs
 			const int32 NumLODs = 1;
@@ -213,7 +211,7 @@ namespace UE::UAF
 		}
 	}
 
-	FAttributeBindingData::FAttributeBindingData(const UAbstractSkeletonSetBinding& Binding)
+	FAttributeBindingData::FAttributeBindingData(TNonNullPtr<const UAbstractSkeletonSetBinding> Binding, const USkeletalMesh* SkeletalMesh)
 	{
 		using namespace Private;
 
@@ -221,20 +219,26 @@ namespace UE::UAF
 		// This reference is our external one that is creating this instance
 		AddRef();
 
-		const USkeleton* Skeleton = Binding.GetSkeleton();
+		const USkeleton* Skeleton = Binding->GetSkeleton();
+		if (!ensure(Skeleton)) // Considering the case where a user creates a binding for a skeleton and then the skeleton is deleted, maybe need a Binding->IsValid()? :/
+		{
+			return;
+		}
+
+		const FReferenceSkeleton& SkeletonRefSkeleton = Skeleton->GetReferenceSkeleton();
+		const FReferenceSkeleton* SkeletalMeshRefSkeleton = SkeletalMesh ? &SkeletalMesh->GetRefSkeleton() : nullptr;
 
-		const TConstArrayView<FAbstractSkeleton_BoneBinding> BoneBindings = Binding.GetBoneBindings();
-		const TConstArrayView<FAbstractSkeleton_AttributeBinding> AttributeBindings = Binding.GetAttributeBindings();
+		const TConstArrayView<FAbstractSkeleton_BoneBinding> BoneBindings = Binding->GetBoneBindings();
+		const TConstArrayView<FAbstractSkeleton_AttributeBinding> AttributeBindings = Binding->GetAttributeBindings();
 
-		const FReferenceSkeleton& RefSkeleton = Skeleton->GetReferenceSkeleton();
 
 		UScriptStruct* BoneTransformType = FBoneTransformAnimationAttribute::StaticStruct();
 
 		// Collect every set names and attribute type
-		const TArray<FName> SetNames = CollectSetNames(&Binding);
-		const TArray<UScriptStruct*> AttributeTypes = CollectAttributeTypes(RefSkeleton, AttributeBindings);
+		const TArray<FName> SetNames = CollectSetNames(Binding);
+		const TArray<UScriptStruct*> AttributeTypes = CollectAttributeTypes(SkeletonRefSkeleton, AttributeBindings);
 
-		const TArray<FLODPartition> LODPartitions = PartitionAttributes(Binding, SetNames, AttributeTypes);
+		const TArray<FLODPartition> LODPartitions = PartitionAttributes(Binding, SkeletonRefSkeleton, SetNames, AttributeTypes);
 
 		const int32 NumLODs = LODPartitions.Num();
 		const int32 NumNamedSets = SetNames.Num();
@@ -250,7 +254,7 @@ namespace UE::UAF
 			AttributeNamesByType.Reserve(NumAttributeTypes);
 
 			{
-				const TArray<FMeshBoneInfo>& BoneInfoList = RefSkeleton.GetRefBoneInfo();
+				const TArray<FMeshBoneInfo>& BoneInfoList = SkeletonRefSkeleton.GetRefBoneInfo();
 				const int32 NumSkeletonBones = BoneInfoList.Num();
 
 				TArray<FName>& BoneAttributeNames = AttributeNamesByType.Add(BoneTransformType);
@@ -329,7 +333,9 @@ namespace UE::UAF
 			// TODO: Should we cache all the user authored attributes within the same collection per named set? e.g. translation retarget
 			// This would avoid having 1 collection per authored type and reduce waste
 
-			const TArray<FTransform>& RefPose = Skeleton->GetRefLocalPoses();
+			const TArray<FTransform>& RefPose = SkeletalMeshRefSkeleton
+				? SkeletalMeshRefSkeleton->GetRefBonePose()
+				: SkeletonRefSkeleton.GetRefBonePose();
 
 			for (auto NamedSetIt = CreateNamedSetIterator(); NamedSetIt; ++NamedSetIt)
 			{
@@ -350,14 +356,36 @@ namespace UE::UAF
 					FAttributeMappingKey MappingKey = FAttributeMappingKey::MakeFrom(TypedSetIt->GetType()).To<FBoneTransformAnimationAttribute>();
 					TBoundValueMap<FBoneTransformAnimationAttribute>* DefaultValues = DefaultAttributeValues->GetBoundValueMaps().Add<FBoneTransformAnimationAttribute>(MappingKey);
 
-					// Set every value within our map
 					const int32 NumValues = DefaultValues->Num();
-					for (FAttributeSetIndex ValueIndex(0); ValueIndex < NumValues; ++ValueIndex)
+
+					// Set every value within our map
+					if (SkeletalMeshRefSkeleton)
+					{
+						for (FAttributeSetIndex ValueIndex(0); ValueIndex < NumValues; ++ValueIndex)
+						{
+							FAttributeBindingIndex BindingIndex = TypedSetIt->GetBindingIndex(ValueIndex);
+							FSkeletonPoseBoneIndex SkeletonBoneIndex = BindingIndex.GetSkeletonBoneIndex();
+
+							FName BoneName = SkeletonRefSkeleton.GetBoneName(SkeletonBoneIndex.GetInt());
+							int32 MeshBoneIndex = SkeletalMeshRefSkeleton->FindBoneIndex(BoneName);
+
+							SkeletonToMeshIndicesMap.Add(SkeletonBoneIndex, FMeshPoseBoneIndex(MeshBoneIndex));
+
+							if (MeshBoneIndex != INDEX_NONE)
+							{
+								(*DefaultValues)[ValueIndex] = FBoneTransformAnimationAttribute{ RefPose[MeshBoneIndex] };
+							}
+						}
+					}
+					else
 					{
-						FAttributeBindingIndex BindingIndex = TypedSetIt->GetBindingIndex(ValueIndex);
-						FSkeletonPoseBoneIndex SkeletonBoneIndex = BindingIndex.GetSkeletonBoneIndex();
+						for (FAttributeSetIndex ValueIndex(0); ValueIndex < NumValues; ++ValueIndex)
+						{
+							FAttributeBindingIndex BindingIndex = TypedSetIt->GetBindingIndex(ValueIndex);
+							FSkeletonPoseBoneIndex SkeletonBoneIndex = BindingIndex.GetSkeletonBoneIndex();
 
-						(*DefaultValues)[ValueIndex] = FBoneTransformAnimationAttribute{ RefPose[SkeletonBoneIndex.GetInt()] };
+							(*DefaultValues)[ValueIndex] = FBoneTransformAnimationAttribute{ RefPose[SkeletonBoneIndex.GetInt()] };
+						}
 					}
 
 					// Our set map has a reference to a typed set
@@ -371,7 +399,7 @@ namespace UE::UAF
 			}
 		}
 
-		BindingAsset = &Binding;
+		BindingAsset = Binding;
 	}
 
 	int32 FAttributeBindingData::NumNamedSets() const
@@ -456,16 +484,16 @@ namespace UE::UAF
 		UE::Private::FQueryableRefCountedObject::Release();
 	}
 
-	FAttributeBindingDataPtr MakeAttributeBindingData(const UAbstractSkeletonSetBinding& Binding)
+	FAttributeBindingDataPtr MakeAttributeBindingData(TNonNullPtr<const UAbstractSkeletonSetBinding> Binding, const USkeletalMesh* SkeletalMesh)
 	{
-		const USkeleton* Skeleton = Binding.GetSkeleton();
+		const USkeleton* Skeleton = Binding->GetSkeleton();
 
 		if (Skeleton == nullptr)
 		{
 			return FAttributeBindingDataPtr();
 		}
 
-		FAttributeBindingData* BindingData = new FAttributeBindingData(Binding);
+		FAttributeBindingData* BindingData = new FAttributeBindingData(Binding, SkeletalMesh);
 
 		// We already added the reference within the constructor
 		const bool bAddRef = false;
diff --git a/Engine/Plugins/Experimental/UAF/UAF/Source/UAF/Private/UAF/Attributes/AttributeBindingDataCache.cpp b/Engine/Plugins/Experimental/UAF/UAF/Source/UAF/Private/UAF/Attributes/AttributeBindingDataCache.cpp
new file mode 100644
--- /dev/null
+++ b/Engine/Plugins/Experimental/UAF/UAF/Source/UAF/Private/UAF/Attributes/AttributeBindingDataCache.cpp
@@ -0,0 +1,126 @@
+// Copyright Epic Games, Inc. All Rights Reserved.
+
+#include "UAF/Attributes/AttributeBindingDataCache.h"
+#include "UAF/AbstractSkeleton/AbstractSkeletonSetBinding.h"
+#include "Engine/SkeletalMesh.h"
+
+namespace UE::UAF
+{
+	FAutoConsoleCommand ClearCacheCommand(
+		TEXT("UAF.ClearAttributeBindingCache"),
+		TEXT("A debug command to clear the UAF attribute binding data cache"),
+		FConsoleCommandDelegate::CreateLambda([]()
+			{
+				GAttributeBindingDataCache.Reset();
+			})
+	);
+
+	namespace Private
+	{
+		FSetBindingAttributeBindingDataCache::FSetBindingAttributeBindingDataCache(TNonNullPtr<const UAbstractSkeletonSetBinding> InSetBinding)
+			: SkeletonBindingData(MakeAttributeBindingData(InSetBinding, nullptr))
+			, WeakSetBinding(InSetBinding)
+		{
+		}
+
+		FAttributeBindingDataPtr FSetBindingAttributeBindingDataCache::GetOrAdd(const USkeletalMesh* SkeletalMesh)
+		{		
+			if (!SkeletalMesh)
+			{
+				return SkeletonBindingData;
+			}
+
+			const UAbstractSkeletonSetBinding* SetBinding = WeakSetBinding.Get();
+			check(SetBinding);
+		
+			check(SkeletalMesh->GetSkeleton() == SetBinding->GetSkeleton());
+
+			{
+				FReadScopeLock Lock(CacheLock);
+				if (FAttributeBindingDataPtr* CachedBindingData = SkeletalMeshBindingData.Find(SkeletalMesh))
+				{
+					return *CachedBindingData;
+				}
+			}
+
+			FAttributeBindingDataPtr NewAttributeBindingData = MakeAttributeBindingData(SetBinding, SkeletalMesh);
+
+			{
+				FWriteScopeLock Lock(CacheLock);
+
+				// Recheck for update between read-lock released and write-lock obtained
+				if (FAttributeBindingDataPtr* CachedBindingData = SkeletalMeshBindingData.Find(SkeletalMesh))
+				{
+					return *CachedBindingData;
+				}
+
+				return SkeletalMeshBindingData.Add(SkeletalMesh, MoveTemp(NewAttributeBindingData));
+			}
+		}
+	}
+	
+	FAttributeBindingDataPtr FAttributeBindingDataCache::GetOrAdd(
+		TNonNullPtr<const UAbstractSkeletonSetBinding> SetBinding,
+		const USkeletalMesh* SkeletalMesh)
+	{
+		{
+			FSetBindingAttributeBindingDataCachePtr CachedBindingData = nullptr;
+			{
+				FReadScopeLock Lock(CacheLock);
+				
+				if (FSetBindingAttributeBindingDataCachePtr* CachedBindingDataPtr = SetBindingToCache.Find(SetBinding))
+				{
+					CachedBindingData = *CachedBindingDataPtr;
+				}
+			}
+
+			if (CachedBindingData)
+			{
+				return CachedBindingData->GetOrAdd(SkeletalMesh);
+			}
+		}
+
+		FSetBindingAttributeBindingDataCachePtr SetBindingCache = MakeShared<Private::FSetBindingAttributeBindingDataCache>(SetBinding);
+		
+		{
+			FWriteScopeLock Lock(CacheLock);
+
+			// Recheck for update between read-lock released and write-lock obtained
+			if (const FSetBindingAttributeBindingDataCachePtr* CachedBindingData = SetBindingToCache.Find(SetBinding))
+			{
+				SetBindingCache = *CachedBindingData;
+			}
+			else
+			{
+				SetBindingToCache.Add(SetBinding, SetBindingCache);
+			}
+		}
+
+		return SetBindingCache->GetOrAdd(SkeletalMesh);
+	}
+
+	void FAttributeBindingDataCache::ResetSetBinding(TNonNullPtr<const UAbstractSkeletonSetBinding> SetBinding)
+	{
+		// Early out without acquiring write access if there's nothing to clear
+		{
+			FReadScopeLock Lock(CacheLock);
+			if (!SetBindingToCache.Contains(SetBinding))
+			{
+				return;
+			}
+		}
+
+		{
+			FWriteScopeLock Lock(CacheLock);
+			SetBindingToCache.Remove(SetBinding);
+		}
+	}
+
+	void FAttributeBindingDataCache::Reset()
+	{
+		FWriteScopeLock Lock(CacheLock);
+		SetBindingToCache.Reset();
+	}
+	
+	FAttributeBindingDataCache GAttributeBindingDataCache;
+}
\ No newline at end of file
diff --git a/Engine/Plugins/Experimental/UAF/UAF/Source/UAF/Private/UAF/UAFSkeletonUserData.cpp b/Engine/Plugins/Experimental/UAF/UAF/Source/UAF/Private/UAF/UAFSkeletonUserData.cpp
--- a/Engine/Plugins/Experimental/UAF/UAF/Source/UAF/Private/UAF/UAFSkeletonUserData.cpp
+++ b/Engine/Plugins/Experimental/UAF/UAF/Source/UAF/Private/UAF/UAFSkeletonUserData.cpp
@@ -158,18 +158,18 @@ namespace UE::UAF::Private
 	}
 }
 
-UUAFSkeletonUserData* UUAFSkeletonUserData::FromSkeleton(USkeleton& Skeleton)
+UUAFSkeletonUserData* UUAFSkeletonUserData::FromSkeleton(TNonNullPtr<USkeleton> Skeleton)
 {
 	using namespace UE::UAF::Private;
 
-	UUAFSkeletonUserData* UserData = Cast<UUAFSkeletonUserData>(Skeleton.GetAssetUserDataOfClass(UUAFSkeletonUserData::StaticClass()));
+	UUAFSkeletonUserData* UserData = Cast<UUAFSkeletonUserData>(Skeleton->GetAssetUserDataOfClass(UUAFSkeletonUserData::StaticClass()));
 	if (UserData != nullptr && UserData->GetSetBinding() != nullptr)
 	{
 		// We found our user specified data and its valid
 		return UserData;
 	}
 
-	const TObjectKey<USkeleton> SkeletonKey(&Skeleton);
+	const TObjectKey<USkeleton> SkeletonKey(Skeleton);
 	const uint32 SkeletonKeyHash = GetTypeHash(SkeletonKey);
 
 	// No user specified data found on the skeleton itself
@@ -209,9 +209,9 @@ UUAFSkeletonUserData* UUAFSkeletonUserData::FromSkeleton(USkeleton& Skeleton)
 			UserData->SetBinding = NewObject<UAbstractSkeletonSetBinding>();
 			(void)UserData->SetBinding->AtomicallyClearInternalFlags(EInternalObjectFlags::Async);
 
-			UserData->SetBinding->SetSkeleton(&Skeleton);
+			UserData->SetBinding->SetSkeleton(Skeleton);
 
-			CachedData.Skeleton = &Skeleton;
+			CachedData.Skeleton = Skeleton;
 			CachedData.UserData = TStrongObjectPtr<UUAFSkeletonUserData>(UserData);
 		}
 
diff --git a/Engine/Plugins/Experimental/UAF/UAF/Source/UAF/Public/UAF/AbstractSkeleton/AbstractSkeletonSetBinding.h b/Engine/Plugins/Experimental/UAF/UAF/Source/UAF/Public/UAF/AbstractSkeleton/AbstractSkeletonSetBinding.h
--- a/Engine/Plugins/Experimental/UAF/UAF/Source/UAF/Public/UAF/AbstractSkeleton/AbstractSkeletonSetBinding.h
+++ b/Engine/Plugins/Experimental/UAF/UAF/Source/UAF/Public/UAF/AbstractSkeleton/AbstractSkeletonSetBinding.h
@@ -91,9 +91,6 @@ public:
 
 	UE_API const TConstArrayView<FAbstractSkeleton_AttributeBinding> GetAttributeBindings() const;
 
-	// Returns the baked binding data
-	UE_API const UE::UAF::FAttributeBindingDataPtr& GetAttributeBindingData() const;
-
 #if WITH_EDITOR
 	UE_API FDelegateHandle RegisterOnBindingsChanged(FSimpleMulticastDelegate::FDelegate&& InDelegate);
 	UE_API bool UnregisterOnBindingsChanged(const FDelegateHandle InHandle);
@@ -129,10 +126,6 @@ private:
 
 	// TODO: Cache sets in PostLoad for cooked builds
 
-	// Attribute binding data, computed on demand and cached in editor, built in PostLoad in cooked builds
-	// This is a baked and immutable representation of the binding
-	UE::UAF::FAttributeBindingDataPtr AttributeBindingData;
-
 #if WITH_EDITOR
 	FDelegateHandle SetCollectionChangedHandle;
 
diff --git a/Engine/Plugins/Experimental/UAF/UAF/Source/UAF/Public/UAF/Attributes/AttributeBindingData.h b/Engine/Plugins/Experimental/UAF/UAF/Source/UAF/Public/UAF/Attributes/AttributeBindingData.h
--- a/Engine/Plugins/Experimental/UAF/UAF/Source/UAF/Public/UAF/Attributes/AttributeBindingData.h
+++ b/Engine/Plugins/Experimental/UAF/UAF/Source/UAF/Public/UAF/Attributes/AttributeBindingData.h
@@ -10,6 +10,7 @@
 
 #define UE_API UAF_API
 
+class USkeletalMesh;
 class UAbstractSkeletonSetBinding;
 
 namespace UE::UAF
@@ -100,7 +101,7 @@ namespace UE::UAF
 	private:
 		UE_NONCOPYABLE(FAttributeBindingData);	// We disallow copy/move semantics
 
-		explicit FAttributeBindingData(const UAbstractSkeletonSetBinding& Binding);
+		FAttributeBindingData(TNonNullPtr<const UAbstractSkeletonSetBinding> Binding, const USkeletalMesh* SkeletalMesh);
 		void ResetInternalReferences();
 
 		// TODO: We should allocate the binding data, named sets, and typed sets data contiguously in a single buffer that is
@@ -125,18 +126,21 @@ namespace UE::UAF
 		// Default attribute values for each named set
 		TMap<FName, FValueBundlePtr> DefaultAttributeValuesByName;
 
+		// Only filled when executing with a skeletal mesh
+		TMap<FSkeletonPoseBoneIndex, FMeshPoseBoneIndex> SkeletonToMeshIndicesMap;
+
 		// The source binding asset used to create this binding data
 		TWeakObjectPtr<const UAbstractSkeletonSetBinding> BindingAsset;
 
 		// The number of internal references for the purposes of reference counting
 		uint32 NumInternalReferences = 0;
 
-		friend UE_API FAttributeBindingDataPtr MakeAttributeBindingData(const UAbstractSkeletonSetBinding& Binding);
+		friend UE_API FAttributeBindingDataPtr MakeAttributeBindingData(TNonNullPtr<const UAbstractSkeletonSetBinding> Binding, const USkeletalMesh* SkeletalMesh);
 	};
 
 	// Creates a new attribute binding data instance from the specified binding
 	// Returns null if the binding isn't valid
-	[[nodiscard]] UE_API FAttributeBindingDataPtr MakeAttributeBindingData(const UAbstractSkeletonSetBinding& Binding);
+	[[nodiscard]] UE_API FAttributeBindingDataPtr MakeAttributeBindingData(TNonNullPtr<const UAbstractSkeletonSetBinding> Binding, const USkeletalMesh* SkeletalMesh);
 }
 
 #undef UE_API
diff --git a/Engine/Plugins/Experimental/UAF/UAF/Source/UAF/Public/UAF/Attributes/AttributeBindingDataCache.h b/Engine/Plugins/Experimental/UAF/UAF/Source/UAF/Public/UAF/Attributes/AttributeBindingDataCache.h
new file mode 100644
--- /dev/null
+++ b/Engine/Plugins/Experimental/UAF/UAF/Source/UAF/Public/UAF/Attributes/AttributeBindingDataCache.h
@@ -0,0 +1,52 @@
+// Copyright Epic Games, Inc. All Rights Reserved.
+
+#pragma once
+
+#include "CoreMinimal.h"
+#include "AttributeBindingData.h"
+#include "Misc/TransactionallySafeRWLock.h"
+#include "UObject/ObjectKey.h"
+
+#define UE_API UAF_API
+
+namespace UE::UAF
+{
+	namespace Private
+	{
+		class FSetBindingAttributeBindingDataCache
+		{
+		public:
+			FSetBindingAttributeBindingDataCache() = delete;
+			explicit FSetBindingAttributeBindingDataCache(TNonNullPtr<const UAbstractSkeletonSetBinding> SetBinding);
+
+			FAttributeBindingDataPtr GetOrAdd(const USkeletalMesh* SkeletalMesh);
+
+		private:
+			const FAttributeBindingDataPtr SkeletonBindingData;
+			const TWeakObjectPtr<const UAbstractSkeletonSetBinding> WeakSetBinding;
+
+			TMap<FObjectKey, FAttributeBindingDataPtr> SkeletalMeshBindingData;
+			FTransactionallySafeRWLock CacheLock;
+		};
+	}
+
+	using FSetBindingAttributeBindingDataCachePtr = TSharedPtr<Private::FSetBindingAttributeBindingDataCache>;
+
+	// A thread-safe global cache for mapping set bindings to the cache of attribute data for each skeletal mesh using that binding
+	class FAttributeBindingDataCache
+	{
+	public:
+		UE_API FAttributeBindingDataPtr GetOrAdd(TNonNullPtr<const UAbstractSkeletonSetBinding> SetBinding, const USkeletalMesh* SkeletalMesh);
+
+		UE_API void ResetSetBinding(TNonNullPtr<const UAbstractSkeletonSetBinding> SetBinding);
+		UE_API void Reset();
+		
+	private:
+		FTransactionallySafeRWLock CacheLock;
+		TMap<FObjectKey, FSetBindingAttributeBindingDataCachePtr> SetBindingToCache;
+	};
+
+	UE_API extern FAttributeBindingDataCache GAttributeBindingDataCache;
+}
+
+#undef UE_API
\ No newline at end of file
diff --git a/Engine/Plugins/Experimental/UAF/UAF/Source/UAF/Public/UAF/UAFSkeletonUserData.h b/Engine/Plugins/Experimental/UAF/UAF/Source/UAF/Public/UAF/UAFSkeletonUserData.h
--- a/Engine/Plugins/Experimental/UAF/UAF/Source/UAF/Public/UAF/UAFSkeletonUserData.h
+++ b/Engine/Plugins/Experimental/UAF/UAF/Source/UAF/Public/UAF/UAFSkeletonUserData.h
@@ -22,7 +22,7 @@ class UUAFSkeletonUserData : public UAssetUserData
 public:
 	// Returns the UAF skeleton user data associated with the specified skeleton
 	// If one isn't set by the user, a new one is allocated on demand and cached
-	static UE_API UUAFSkeletonUserData* FromSkeleton(USkeleton& Skeleton);
+	static UE_API UUAFSkeletonUserData* FromSkeleton(TNonNullPtr<USkeleton> Skeleton);
 
 	// The default abstract skeleton label binding asset to use
 	UE_API UAbstractSkeletonLabelBinding* GetLabelBinding() const;
diff --git a/Engine/Plugins/Experimental/UAF/UAFAnimGraph/Source/UAFAnimGraph/Private/EvaluationVM/EvaluationTaskContext.cpp b/Engine/Plugins/Experimental/UAF/UAFAnimGraph/Source/UAFAnimGraph/Private/EvaluationVM/EvaluationTaskContext.cpp
--- a/Engine/Plugins/Experimental/UAF/UAFAnimGraph/Source/UAFAnimGraph/Private/EvaluationVM/EvaluationTaskContext.cpp
+++ b/Engine/Plugins/Experimental/UAF/UAFAnimGraph/Source/UAFAnimGraph/Private/EvaluationVM/EvaluationTaskContext.cpp
@@ -3,19 +3,21 @@
 #include "EvaluationVM/EvaluationTaskContext.h"
 
 #include "UAF/AbstractSkeleton/AbstractSkeletonSetBinding.h"
+#include "UAF/Attributes/AttributeBindingDataCache.h"
 
 namespace UE::UAF
 {
 	FEvaluationTaskContext::FEvaluationTaskContext() = default;
 
-	FEvaluationTaskContext::FEvaluationTaskContext(const UAbstractSkeletonSetBinding& InSetBinding, FName InSetName)
+	FEvaluationTaskContext::FEvaluationTaskContext(TNonNullPtr<const UAbstractSkeletonSetBinding> InSetBinding, const USkeletalMesh* InSkeletalMesh, FName InSetName)
 	{
-		if (const FAttributeBindingDataPtr& InBindingData = InSetBinding.GetAttributeBindingData())
+		if (const FAttributeBindingDataPtr& InBindingData = GAttributeBindingDataCache.GetOrAdd(InSetBinding, InSkeletalMesh))
 		{
 			if (FAttributeNamedSetPtr InNamedSet = InBindingData->FindNamedSet(InSetName))
 			{
-				SetBinding = &InSetBinding;
-				Skeleton = InSetBinding.GetSkeleton();
+				SetBinding = InSetBinding;
+				SkeletalMesh = InSkeletalMesh;
+				Skeleton = InSetBinding->GetSkeleton();
 				BindingData = InBindingData;
 				NamedSet = MoveTemp(InNamedSet);
 				DefaultAttributeValues = NamedSet->GetDefaultAttributeValues();
@@ -33,6 +35,11 @@ namespace UE::UAF
 		return Skeleton;
 	}
 
+	const USkeletalMesh* FEvaluationTaskContext::GetSkeletalMesh() const
+	{
+		return SkeletalMesh;
+	}
+
 	const FAttributeNamedSetPtr& FEvaluationTaskContext::GetNamedSet() const
 	{
 		return NamedSet;
diff --git a/Engine/Plugins/Experimental/UAF/UAFAnimGraph/Source/UAFAnimGraph/Private/EvaluationVM/EvaluationVM.cpp b/Engine/Plugins/Experimental/UAF/UAFAnimGraph/Source/UAFAnimGraph/Private/EvaluationVM/EvaluationVM.cpp
--- a/Engine/Plugins/Experimental/UAF/UAFAnimGraph/Source/UAFAnimGraph/Private/EvaluationVM/EvaluationVM.cpp
+++ b/Engine/Plugins/Experimental/UAF/UAFAnimGraph/Source/UAFAnimGraph/Private/EvaluationVM/EvaluationVM.cpp
@@ -38,16 +38,16 @@ namespace UE::UAF
 		EvaluationContextStack.Push(FEvaluationTaskContext());
 	}
 
-	FEvaluationVM::FEvaluationVM(EEvaluationFlags InEvaluationFlags, const UAbstractSkeletonSetBinding& InSetBinding, FName InSetName, int32 InCurrentLOD)
+	FEvaluationVM::FEvaluationVM(EEvaluationFlags InEvaluationFlags, TNonNullPtr<const UAbstractSkeletonSetBinding> InSetBinding, const USkeletalMesh* InSkeletalMesh, FName InSetName, int32 InCurrentLOD)
 		: ValueTransformerMap(FValueRuntimeRegistry::Get().GetTransformerMap())
 		, CurrentLOD(InCurrentLOD)
 		, EvaluationFlags(InEvaluationFlags)
 	{
-		EvaluationContextStack.Push(FEvaluationTaskContext(InSetBinding, InSetName));
+		EvaluationContextStack.Push(FEvaluationTaskContext(InSetBinding, InSkeletalMesh, InSetName));
 	}
 
 	// REMOVE ME: Exclusively for testing that before/after yields the same results
-	FEvaluationVM::FEvaluationVM(EEvaluationFlags InEvaluationFlags, const FReferencePose& InReferencePose, const UAbstractSkeletonSetBinding* InSetBinding, FName InSetName, int32 InCurrentLOD)
+	FEvaluationVM::FEvaluationVM(EEvaluationFlags InEvaluationFlags, const FReferencePose& InReferencePose, const UAbstractSkeletonSetBinding* InSetBinding, const USkeletalMesh* InSkeletalMesh, FName InSetName, int32 InCurrentLOD)
 		: ReferencePose(&InReferencePose)
 		, ValueTransformerMap(FValueRuntimeRegistry::Get().GetTransformerMap())
 		, CurrentLOD(InCurrentLOD)
@@ -55,7 +55,7 @@ namespace UE::UAF
 	{
 		if (InSetBinding != nullptr)
 		{
-			EvaluationContextStack.Push(FEvaluationTaskContext(*InSetBinding, InSetName));
+			EvaluationContextStack.Push(FEvaluationTaskContext(InSetBinding, InSkeletalMesh, InSetName));
 		}
 		else
 		{
@@ -129,9 +129,9 @@ namespace UE::UAF
 		return EvaluationContextStack.Last().GetNamedSet();
 	}
 
-	bool FEvaluationVM::PushEvaluationContext(const UAbstractSkeletonSetBinding& SetBinding, FName SetName)
+	bool FEvaluationVM::PushEvaluationContext(TNonNullPtr<const UAbstractSkeletonSetBinding> SetBinding, const USkeletalMesh* SkeletalMesh, FName SetName)
 	{
-		FEvaluationTaskContext NewContext(SetBinding, SetName);
+		FEvaluationTaskContext NewContext(SetBinding, SkeletalMesh, SetName);
 		if (NewContext.IsValid())
 		{
 			EvaluationContextStack.Push(MoveTemp(NewContext));
diff --git a/Engine/Plugins/Experimental/UAF/UAFAnimGraph/Source/UAFAnimGraph/Private/Graph/AnimGraph/AnimNode_AnimNextGraph.cpp b/Engine/Plugins/Experimental/UAF/UAFAnimGraph/Source/UAFAnimGraph/Private/Graph/AnimGraph/AnimNode_AnimNextGraph.cpp
--- a/Engine/Plugins/Experimental/UAF/UAFAnimGraph/Source/UAFAnimGraph/Private/Graph/AnimGraph/AnimNode_AnimNextGraph.cpp
+++ b/Engine/Plugins/Experimental/UAF/UAFAnimGraph/Source/UAFAnimGraph/Private/Graph/AnimGraph/AnimNode_AnimNextGraph.cpp
@@ -128,14 +128,18 @@ void FAnimNode_AnimNextGraph::Evaluate_AnyThread(FPoseContext& Output)
 
 			if (UE::UAF::UseAbstractSkeletonRuntime())
 			{
+                USkeletalMesh* SkeletalMesh = SkeletalMeshComponent->GetSkeletalMeshAsset();
+                check(SkeletalMesh);
+                
 				// By default, use the everything set
 				USkeleton* Skeleton = const_cast<USkeleton*>(RefPose.GetSkeletonAsset());
-				UUAFSkeletonUserData* UserData = UUAFSkeletonUserData::FromSkeleton(*Skeleton);
-
-				UAbstractSkeletonSetBinding* Binding = UserData != nullptr ? UserData->GetSetBinding() : nullptr;
+				check(Skeleton);
+				
+				UUAFSkeletonUserData* UserData = UUAFSkeletonUserData::FromSkeleton(Skeleton);
+				UAbstractSkeletonSetBinding* Binding = UserData ? UserData->GetSetBinding() : nullptr;
 				FName SetName;
 
-				EvaluateGraphContext.SetSkeletonSet(Binding, SetName);
+				EvaluateGraphContext.SetSkeletonSet(Binding, SkeletalMesh, SetName);
 			}
 
 			const FEvaluationProgram EvaluationProgram = UE::UAF::EvaluateGraph(EvaluateGraphContext);
diff --git a/Engine/Plugins/Experimental/UAF/UAFAnimGraph/Source/UAFAnimGraph/Private/Graph/RigUnit_AnimNextRunAnimationGraph_v1.cpp b/Engine/Plugins/Experimental/UAF/UAFAnimGraph/Source/UAFAnimGraph/Private/Graph/RigUnit_AnimNextRunAnimationGraph_v1.cpp
--- a/Engine/Plugins/Experimental/UAF/UAFAnimGraph/Source/UAFAnimGraph/Private/Graph/RigUnit_AnimNextRunAnimationGraph_v1.cpp
+++ b/Engine/Plugins/Experimental/UAF/UAFAnimGraph/Source/UAFAnimGraph/Private/Graph/RigUnit_AnimNextRunAnimationGraph_v1.cpp
@@ -12,6 +12,7 @@
 #include "Graph/RigVMTrait_AnimNextPublicVariables.h"
 #include "TraitCore/TraitEventList.h"
 #include "UAF/UAFSkeletonUserData.h"
+#include "Graph/UAFSystemOutputComponent.h"
 
 #include UE_INLINE_GENERATED_CPP_BY_NAME(RigUnit_AnimNextRunAnimationGraph_v1)
 
@@ -101,14 +102,20 @@ FRigUnit_AnimNextRunAnimationGraph_v1_Execute()
 
 	if (UseAbstractSkeletonRuntime())
 	{
+		FUAFSystemOutputComponent& OutputPoseComponent = ModuleInstance.GetComponent<FUAFSystemOutputComponent>();
+    		
+    	USkeletalMesh* SkeletalMesh = OutputPoseComponent.GetSkeletalMesh();
+    	check(SkeletalMesh);
+    	
 		// By default, use the everything set
 		USkeleton* Skeleton = const_cast<USkeleton*>(RefPose.GetSkeletonAsset());
-		UUAFSkeletonUserData* UserData = UUAFSkeletonUserData::FromSkeleton(*Skeleton);
-
+		check(Skeleton)
+		
+		UUAFSkeletonUserData* UserData = UUAFSkeletonUserData::FromSkeleton(Skeleton);
 		UAbstractSkeletonSetBinding* Binding = UserData != nullptr ? UserData->GetSetBinding() : nullptr;
 		FName SetName;
 
-		EvaluateGraphContext.SetSkeletonSet(Binding, SetName);
+		EvaluateGraphContext.SetSkeletonSet(Binding, SkeletalMesh, SetName);
 	}
 
 	EvaluateGraph(EvaluateGraphContext, Result);
diff --git a/Engine/Plugins/Experimental/UAF/UAFAnimGraph/Source/UAFAnimGraph/Private/Graph/RigUnit_AnimNextRunAnimationGraph_v2.cpp b/Engine/Plugins/Experimental/UAF/UAFAnimGraph/Source/UAFAnimGraph/Private/Graph/RigUnit_AnimNextRunAnimationGraph_v2.cpp
--- a/Engine/Plugins/Experimental/UAF/UAFAnimGraph/Source/UAFAnimGraph/Private/Graph/RigUnit_AnimNextRunAnimationGraph_v2.cpp
+++ b/Engine/Plugins/Experimental/UAF/UAFAnimGraph/Source/UAFAnimGraph/Private/Graph/RigUnit_AnimNextRunAnimationGraph_v2.cpp
@@ -16,6 +16,7 @@
 #include "Factory/AnimGraphFactory.h"
 #include "UAF/UAFSkeletonUserData.h"
 #include "UAFLogging.h"
+#include "Engine/SkeletalMesh.h"
 #include "Graph/RigUnit_UAFRunAsset.h"
 #include "Graph/UAFSystemOutputComponent.h"
 
@@ -141,13 +142,17 @@ FRigUnit_AnimNextRunAnimationGraph_v2_Execute()
 
 	if (UseAbstractSkeletonRuntime())
 	{
-		// By default, use the everything set
+		FUAFSystemOutputComponent& OutputPoseComponent = ModuleInstance.GetComponent<FUAFSystemOutputComponent>();
+		
+		USkeletalMesh* SkeletalMesh = OutputPoseComponent.GetSkeletalMesh();
+		
 		USkeleton* Skeleton = const_cast<USkeleton*>(RefPose.GetSkeletonAsset());
-		UUAFSkeletonUserData* UserData = UUAFSkeletonUserData::FromSkeleton(*Skeleton);
-
+		check(Skeleton);
+		
+		UUAFSkeletonUserData* UserData = UUAFSkeletonUserData::FromSkeleton(Skeleton);
 		UAbstractSkeletonSetBinding* Binding = UserData != nullptr ? UserData->GetSetBinding() : nullptr;
 
-		EvaluateGraphContext.SetSkeletonSet(Binding, AttributeSet);
+		EvaluateGraphContext.SetSkeletonSet(Binding, SkeletalMesh, AttributeSet);
 	}
 
 	EvaluateGraph(EvaluateGraphContext, Result);
diff --git a/Engine/Plugins/Experimental/UAF/UAFAnimGraph/Source/UAFAnimGraph/Private/Graph/RigUnit_UAFRunAsset.cpp b/Engine/Plugins/Experimental/UAF/UAFAnimGraph/Source/UAFAnimGraph/Private/Graph/RigUnit_UAFRunAsset.cpp
--- a/Engine/Plugins/Experimental/UAF/UAFAnimGraph/Source/UAFAnimGraph/Private/Graph/RigUnit_UAFRunAsset.cpp
+++ b/Engine/Plugins/Experimental/UAF/UAFAnimGraph/Source/UAFAnimGraph/Private/Graph/RigUnit_UAFRunAsset.cpp
@@ -126,13 +126,16 @@ FRigUnit_UAFRunAsset_Execute()
 
 	if (UseAbstractSkeletonRuntime())
 	{
-		// By default, use the everything set
-		USkeleton* Skeleton = const_cast<USkeleton*>(RefPose.GetSkeletonAsset());
-		UUAFSkeletonUserData* UserData = UUAFSkeletonUserData::FromSkeleton(*Skeleton);
+    	USkeletalMesh* SkeletalMesh = OutputPoseComponent.GetSkeletalMesh();
+    	check(SkeletalMesh);
 
+		USkeleton* Skeleton = const_cast<USkeleton*>(RefPose.GetSkeletonAsset());
+		check(Skeleton);
+		
+		UUAFSkeletonUserData* UserData = UUAFSkeletonUserData::FromSkeleton(Skeleton);
 		UAbstractSkeletonSetBinding* Binding = UserData != nullptr ? UserData->GetSetBinding() : nullptr;
 
-		EvaluateGraphContext.SetSkeletonSet(Binding, AttributeSet);
+		EvaluateGraphContext.SetSkeletonSet(Binding, SkeletalMesh, AttributeSet);
 	}
 
 	EvaluateGraph(EvaluateGraphContext, GraphPose);
diff --git a/Engine/Plugins/Experimental/UAF/UAFAnimGraph/Source/UAFAnimGraph/Private/TraitInterfaces/IEvaluate.cpp b/Engine/Plugins/Experimental/UAF/UAFAnimGraph/Source/UAFAnimGraph/Private/TraitInterfaces/IEvaluate.cpp
--- a/Engine/Plugins/Experimental/UAF/UAFAnimGraph/Source/UAFAnimGraph/Private/TraitInterfaces/IEvaluate.cpp
+++ b/Engine/Plugins/Experimental/UAF/UAFAnimGraph/Source/UAFAnimGraph/Private/TraitInterfaces/IEvaluate.cpp
@@ -258,8 +258,9 @@ namespace UE::UAF
 
 			UAbstractSkeletonSetBinding* Binding = EvaluateGraphContext.GetSkeletonSetBinding();
 			FName SetName = EvaluateGraphContext.GetSkeletonSetName();
+			USkeletalMesh* SkeletalMesh = EvaluateGraphContext.GetSkeletalMesh();
 
-			FEvaluationVM EvaluationVM(EEvaluationFlags::All, RefPose, Binding, SetName, GraphLODLevel);
+			FEvaluationVM EvaluationVM(EEvaluationFlags::All, RefPose, Binding, SkeletalMesh, SetName, GraphLODLevel);
 			bool bHasValidOutput = false;
 
 			if (!EvaluationProgram.IsEmpty())
diff --git a/Engine/Plugins/Experimental/UAF/UAFAnimGraph/Source/UAFAnimGraph/Public/EvaluationVM/EvaluationTaskContext.h b/Engine/Plugins/Experimental/UAF/UAFAnimGraph/Source/UAFAnimGraph/Public/EvaluationVM/EvaluationTaskContext.h
--- a/Engine/Plugins/Experimental/UAF/UAFAnimGraph/Source/UAFAnimGraph/Public/EvaluationVM/EvaluationTaskContext.h
+++ b/Engine/Plugins/Experimental/UAF/UAFAnimGraph/Source/UAFAnimGraph/Public/EvaluationVM/EvaluationTaskContext.h
@@ -22,7 +22,7 @@ namespace UE::UAF
 	{
 	public:
 		UE_API FEvaluationTaskContext();
-		UE_API FEvaluationTaskContext(const UAbstractSkeletonSetBinding& SetBinding, FName SetName);
+		UE_API FEvaluationTaskContext(TNonNullPtr<const UAbstractSkeletonSetBinding> SetBinding, const USkeletalMesh* SkeletalMesh, FName SetName);
 
 		// Returns whether or not this context is valid
 		[[nodiscard]] UE_API bool IsValid() const;
@@ -33,6 +33,9 @@ namespace UE::UAF
 		// Returns the concrete skeleton we are evaluating with (or nullptr if we aren't valid)
 		[[nodiscard]] UE_API const USkeleton* GetSkeleton() const;
 
+		// Returns the skeletal mesh we are evaluating with, may be nullptr)
+		[[nodiscard]] UE_API const USkeletalMesh* GetSkeletalMesh() const;
+
 		// Returns the named set we are currently evaluating with (or nullptr if we aren't valid)
 		[[nodiscard]] UE_API const FAttributeNamedSetPtr& GetNamedSet() const;
 
@@ -47,6 +50,9 @@ namespace UE::UAF
 		// The set binding we are evaluating with
 		const UAbstractSkeletonSetBinding* SetBinding = nullptr;
 
+		// The target skeletal mesh with are evaluating with, may be nullptr
+		const USkeletalMesh* SkeletalMesh = nullptr;
+		
 		// The target skeleton we are evaluating with
 		const USkeleton* Skeleton = nullptr;
 
diff --git a/Engine/Plugins/Experimental/UAF/UAFAnimGraph/Source/UAFAnimGraph/Public/EvaluationVM/EvaluationVM.h b/Engine/Plugins/Experimental/UAF/UAFAnimGraph/Source/UAFAnimGraph/Public/EvaluationVM/EvaluationVM.h
--- a/Engine/Plugins/Experimental/UAF/UAFAnimGraph/Source/UAFAnimGraph/Public/EvaluationVM/EvaluationVM.h
+++ b/Engine/Plugins/Experimental/UAF/UAFAnimGraph/Source/UAFAnimGraph/Public/EvaluationVM/EvaluationVM.h
@@ -152,10 +152,10 @@ namespace UE::UAF
 	struct FEvaluationVM
 	{
 		UE_API FEvaluationVM(EEvaluationFlags InEvaluationFlags, const FReferencePose& InReferencePose, int32 InCurrentLOD);
-		UE_API FEvaluationVM(EEvaluationFlags EvaluationFlags, const UAbstractSkeletonSetBinding& SetBinding, FName SetName, int32 CurrentLOD);
+		UE_API FEvaluationVM(EEvaluationFlags EvaluationFlags, TNonNullPtr<const UAbstractSkeletonSetBinding> SetBinding, const USkeletalMesh* SkeletalMesh, FName SetName, int32 CurrentLOD);
 
 		// REMOVE ME: Exclusively for testing that before/after yields the same results
-		UE_API FEvaluationVM(EEvaluationFlags InEvaluationFlags, const FReferencePose& IReferencePose, const UAbstractSkeletonSetBinding* SetBinding, FName SetName, int32 InCurrentLOD);
+		UE_API FEvaluationVM(EEvaluationFlags InEvaluationFlags, const FReferencePose& IReferencePose, const UAbstractSkeletonSetBinding* SetBinding, const USkeletalMesh* SkeletalMesh, FName SetName, int32 InCurrentLOD);
 
 		// Allow move
 		FEvaluationVM(FEvaluationVM&&) = default;
@@ -205,7 +205,7 @@ namespace UE::UAF
 
 		// Pushes a new binding and named set to evaluate with
 		// Returns true if we succeeded and the set exists in the specified binding, false otherwise
-		[[nodiscard]] UE_API bool PushEvaluationContext(const UAbstractSkeletonSetBinding& SetBinding, FName SetName);
+		[[nodiscard]] UE_API bool PushEvaluationContext(TNonNullPtr<const UAbstractSkeletonSetBinding> SetBinding, const USkeletalMesh* SkeletalMesh, FName SetName);
 
 		// Pushes a new named set to evaluate with
 		// Returns true if we succeeded and the set exists in the currently active binding, false otherwise
diff --git a/Engine/Plugins/Experimental/UAF/UAFAnimGraph/Source/UAFAnimGraph/Public/TraitInterfaces/IEvaluate.h b/Engine/Plugins/Experimental/UAF/UAFAnimGraph/Source/UAFAnimGraph/Public/TraitInterfaces/IEvaluate.h
--- a/Engine/Plugins/Experimental/UAF/UAFAnimGraph/Source/UAFAnimGraph/Public/TraitInterfaces/IEvaluate.h
+++ b/Engine/Plugins/Experimental/UAF/UAFAnimGraph/Source/UAFAnimGraph/Public/TraitInterfaces/IEvaluate.h
@@ -43,14 +43,16 @@ namespace UE::UAF
 		// Only set if the new abstract skeleton runtime is enabled
 		UAbstractSkeletonSetBinding* GetSkeletonSetBinding() const { return SkeletonSetBinding; }
 		FName GetSkeletonSetName() const { return SkeletonSetName; }
+		USkeletalMesh* GetSkeletalMesh() const { return SkeletalMesh; }
 
 	protected:
 		UE_API TWeakObjectPtr<const USkeletalMeshComponent> GetBindingObject() const;
 
-		void SetSkeletonSet(UAbstractSkeletonSetBinding* InSkeletonSetBinding, FName InSkeletonSetName)
+		void SetSkeletonSet(UAbstractSkeletonSetBinding* InSkeletonSetBinding, USkeletalMesh* InSkeletalMesh, FName InSkeletonSetName)
 		{
 			SkeletonSetBinding = InSkeletonSetBinding;
 			SkeletonSetName = InSkeletonSetName;
+			SkeletalMesh = InSkeletalMesh;
 		}
 
 		FAnimNextGraphInstance& GraphInstance;
@@ -58,6 +60,7 @@ namespace UE::UAF
 		const FReferencePose& RefPose;
 
 		UAbstractSkeletonSetBinding* SkeletonSetBinding = nullptr;
+		USkeletalMesh* SkeletalMesh = nullptr;
 		FName SkeletonSetName;
 
 		int32 GraphLODLevel = 0;
diff --git a/Engine/Plugins/Experimental/UAF/UAFAnimNode/Source/UAFAnimNode/Private/RigUnit_RunAnimNode.cpp b/Engine/Plugins/Experimental/UAF/UAFAnimNode/Source/UAFAnimNode/Private/RigUnit_RunAnimNode.cpp
--- a/Engine/Plugins/Experimental/UAF/UAFAnimNode/Source/UAFAnimNode/Private/RigUnit_RunAnimNode.cpp
+++ b/Engine/Plugins/Experimental/UAF/UAFAnimNode/Source/UAFAnimNode/Private/RigUnit_RunAnimNode.cpp
@@ -8,6 +8,7 @@
 #include "EvaluationVM/EvaluationVM.h"
 #include "EvaluationVM/Tasks/PushAnimSequenceKeyframe.h"
 #include "Graph/UAFAnimGraphTrace.h"
+#include "Graph/UAFSystemOutputComponent.h"
 #include "TraitInterfaces/IEvaluate.h"
 #include "UAF/AnimNodeCore/UAFAnimNode.h"
 #include "UAF/AnimNodeCore/UAFAnimNodeUpdate.h"
@@ -101,13 +102,18 @@ FRigUnit_RunAnimNode_Execute()
 
 			if (UseAbstractSkeletonRuntime())
 			{
-				// By default, use the everything set
+				FUAFSystemOutputComponent& OutputPoseComponent = ModuleInstance.GetComponent<FUAFSystemOutputComponent>();
+		
+				USkeletalMesh* SkeletalMesh = OutputPoseComponent.GetSkeletalMesh();
+				check(SkeletalMesh);
+				
 				USkeleton* Skeleton = const_cast<USkeleton*>(RefPose.GetSkeletonAsset());
-				UUAFSkeletonUserData* UserData = UUAFSkeletonUserData::FromSkeleton(*Skeleton);
+				check(Skeleton);
+				
+				TNonNullPtr<UUAFSkeletonUserData> UserData = UUAFSkeletonUserData::FromSkeleton(Skeleton);
+				UAbstractSkeletonSetBinding* Binding = UserData->GetSetBinding();
 
-				UAbstractSkeletonSetBinding* Binding = UserData != nullptr ? UserData->GetSetBinding() : nullptr;
-
-				EvaluateArgs.SetSkeletonSet(Binding, AttributeSet);
+				EvaluateArgs.SetSkeletonSet(Binding, SkeletalMesh, AttributeSet);
 			}
 
 			EvaluateValues(EvaluateArgs, Result);
diff --git a/Engine/Plugins/Experimental/UAF/UAFAnimNode/Source/UAFAnimNode/Private/UAF/AnimNodeCore/UAFAnimNodeUpdate.cpp b/Engine/Plugins/Experimental/UAF/UAFAnimNode/Source/UAFAnimNode/Private/UAF/AnimNodeCore/UAFAnimNodeUpdate.cpp
--- a/Engine/Plugins/Experimental/UAF/UAFAnimNode/Source/UAFAnimNode/Private/UAF/AnimNodeCore/UAFAnimNodeUpdate.cpp
+++ b/Engine/Plugins/Experimental/UAF/UAFAnimNode/Source/UAFAnimNode/Private/UAF/AnimNodeCore/UAFAnimNodeUpdate.cpp
@@ -257,8 +257,9 @@ namespace UE::UAF
 
 		UAbstractSkeletonSetBinding* Binding = Args.GetSkeletonSetBinding();
 		FName SetName = Args.GetSkeletonSetName();
+		USkeletalMesh* SkeletalMesh = Args.GetSkeletalMesh();
 
-		FEvaluationVM EvaluationVM(EEvaluationFlags::All, RefPose, Binding, SetName, LOD);
+		FEvaluationVM EvaluationVM(EEvaluationFlags::All, RefPose, Binding, SkeletalMesh, SetName, LOD);
 		bool bHasValidOutput = false;
 
 		if (!EvaluationProgram.IsEmpty())
diff --git a/Engine/Plugins/Experimental/UAF/UAFAnimNode/Source/UAFAnimNode/Public/UAF/AnimNodeCore/UAFAnimNodeUpdate.h b/Engine/Plugins/Experimental/UAF/UAFAnimNode/Source/UAFAnimNode/Public/UAF/AnimNodeCore/UAFAnimNodeUpdate.h
--- a/Engine/Plugins/Experimental/UAF/UAFAnimNode/Source/UAFAnimNode/Public/UAF/AnimNodeCore/UAFAnimNodeUpdate.h
+++ b/Engine/Plugins/Experimental/UAF/UAFAnimNode/Source/UAFAnimNode/Public/UAF/AnimNodeCore/UAFAnimNodeUpdate.h
@@ -105,18 +105,21 @@ namespace UE::UAF
 		// Only set if the new abstract skeleton runtime is enabled
 		UAbstractSkeletonSetBinding* GetSkeletonSetBinding() const { return SkeletonSetBinding; }
 		FName GetSkeletonSetName() const { return SkeletonSetName; }
+		USkeletalMesh* GetSkeletalMesh() const { return SkeletalMesh; }
 
 	private:
-		void SetSkeletonSet(UAbstractSkeletonSetBinding* InSkeletonSetBinding, FName InSkeletonSetName)
+		void SetSkeletonSet(UAbstractSkeletonSetBinding* InSkeletonSetBinding, USkeletalMesh* InSkeletalMesh, FName InSkeletonSetName)
 		{
 			SkeletonSetBinding = InSkeletonSetBinding;
 			SkeletonSetName = InSkeletonSetName;
+			SkeletalMesh = InSkeletalMesh;
 		}
 
 		const FReferencePose& ReferencePose;
 		const FEvaluationProgram& EvaluationProgram;
 
 		UAbstractSkeletonSetBinding* SkeletonSetBinding = nullptr;
+		USkeletalMesh* SkeletalMesh = nullptr;
 		FName SkeletonSetName;
 
 		int32 LOD = 0;
diff --git a/Engine/Plugins/Experimental/UAF/UAFTestSuites/Source/UAFTestSuite/Private/UAF/Attributes/AttributeBindingDataTest.cpp b/Engine/Plugins/Experimental/UAF/UAFTestSuites/Source/UAFTestSuite/Private/UAF/Attributes/AttributeBindingDataTest.cpp
--- a/Engine/Plugins/Experimental/UAF/UAFTestSuites/Source/UAFTestSuite/Private/UAF/Attributes/AttributeBindingDataTest.cpp
+++ b/Engine/Plugins/Experimental/UAF/UAFTestSuites/Source/UAFTestSuite/Private/UAF/Attributes/AttributeBindingDataTest.cpp
@@ -31,7 +31,7 @@ namespace UE::UAF::Tests
 		UE_RETURN_ON_ERROR(Binding != nullptr, "Failed to create binding asset");
 
 		{
-			FAttributeBindingDataPtr BindingData = MakeAttributeBindingData(*Binding);
+			FAttributeBindingDataPtr BindingData = MakeAttributeBindingData(Binding, nullptr);
 			AddErrorIfFalse(!BindingData.IsValid(), TEXT("Attribute binding data should be invalid"));
 		}
 
@@ -64,7 +64,7 @@ namespace UE::UAF::Tests
 
 		{
 			// We have the everything set
-			FAttributeBindingDataPtr BindingData = MakeAttributeBindingData(*Binding);
+			FAttributeBindingDataPtr BindingData = MakeAttributeBindingData(Binding, nullptr);
 			
 			UE_RETURN_ON_ERROR(BindingData.IsValid(), TEXT("Attribute binding data should be valid"));
 			AddErrorIfFalse(!BindingData->IsEmpty(), TEXT("Attribute binding data should not be empty"));
@@ -178,7 +178,7 @@ namespace UE::UAF::Tests
 
 			{
 				// We have the everything set
-				FAttributeBindingDataPtr BindingData = MakeAttributeBindingData(*Binding);
+				FAttributeBindingDataPtr BindingData = MakeAttributeBindingData(Binding, nullptr);
 
 				UE_RETURN_ON_ERROR(BindingData.IsValid(), TEXT("Attribute binding data should be valid"));
 				AddErrorIfFalse(!BindingData->IsEmpty(), TEXT("Attribute binding data should not be empty"));
@@ -293,7 +293,7 @@ namespace UE::UAF::Tests
 			FloatAttributeNames.Add(TEXT("AttrFloat1"));
 
 			{
-				FAttributeBindingDataPtr BindingData = MakeAttributeBindingData(*Binding);
+				FAttributeBindingDataPtr BindingData = MakeAttributeBindingData(Binding, nullptr);
 				UE_RETURN_ON_ERROR(BindingData.IsValid(), TEXT("Attribute binding data should be valid"));
 				AddErrorIfFalse(!BindingData->IsEmpty(), TEXT("Attribute binding data should not be empty"));
 				AddErrorIfFalse(BindingData->NumNamedSets() == 3, TEXT("Unexpected number of named sets"));
@@ -797,7 +797,7 @@ namespace UE::UAF::Tests
 			bSuccess = SetBinding2->AddBoneToSet(TEXT("Spine_01"), TEXT("UpperBody"));
 			UE_RETURN_ON_ERROR(bSuccess, "Failed to bind Spine_01 to UpperBody");
 
-			FAttributeBindingDataPtr BindingData = MakeAttributeBindingData(*SetBinding2);
+			FAttributeBindingDataPtr BindingData = MakeAttributeBindingData(SetBinding2, nullptr);
 			UE_RETURN_ON_ERROR(BindingData.IsValid(), TEXT("Attribute binding data should be valid"));
 
 			AddErrorIfFalse(!BindingData->IsEmpty(), TEXT("Attribute binding data should not be empty"));
diff --git a/Engine/Plugins/Experimental/UAF/UAFTestSuites/Source/UAFTestSuite/Private/UAF/ValueRuntime/BoundMapCollectionTest.cpp b/Engine/Plugins/Experimental/UAF/UAFTestSuites/Source/UAFTestSuite/Private/UAF/ValueRuntime/BoundMapCollectionTest.cpp
--- a/Engine/Plugins/Experimental/UAF/UAFTestSuites/Source/UAFTestSuite/Private/UAF/ValueRuntime/BoundMapCollectionTest.cpp
+++ b/Engine/Plugins/Experimental/UAF/UAFTestSuites/Source/UAFTestSuite/Private/UAF/ValueRuntime/BoundMapCollectionTest.cpp
@@ -81,7 +81,7 @@ namespace UE::UAF::Tests
 			Binding->AddBoneToSet(TEXT("RightArm"), TEXT("UpperBody"));
 			Binding->AddAttributeToSet(FAnimationAttributeIdentifier(TEXT("AttrFloat1"), INDEX_NONE, NAME_None, FFloatAnimationAttribute::StaticStruct()), TEXT("UpperBody"));
 
-			BindingData = MakeAttributeBindingData(*Binding);
+			BindingData = MakeAttributeBindingData(Binding, nullptr);
 			UE_RETURN_ON_ERROR(BindingData.IsValid(), TEXT("Attribute binding data should be valid"));
 		}
 
diff --git a/Engine/Plugins/Experimental/UAF/UAFTestSuites/Source/UAFTestSuite/Private/UAF/ValueRuntime/BoundValueMapTest.cpp b/Engine/Plugins/Experimental/UAF/UAFTestSuites/Source/UAFTestSuite/Private/UAF/ValueRuntime/BoundValueMapTest.cpp
--- a/Engine/Plugins/Experimental/UAF/UAFTestSuites/Source/UAFTestSuite/Private/UAF/ValueRuntime/BoundValueMapTest.cpp
+++ b/Engine/Plugins/Experimental/UAF/UAFTestSuites/Source/UAFTestSuite/Private/UAF/ValueRuntime/BoundValueMapTest.cpp
@@ -80,7 +80,7 @@ namespace UE::UAF::Tests
 			Binding->AddBoneToSet(TEXT("RightArm"), TEXT("UpperBody"));
 			Binding->AddAttributeToSet(FAnimationAttributeIdentifier(TEXT("AttrFloat1"), INDEX_NONE, NAME_None, FFloatAnimationAttribute::StaticStruct()), TEXT("UpperBody"));
 
-			BindingData = MakeAttributeBindingData(*Binding);
+			BindingData = MakeAttributeBindingData(Binding, nullptr);
 			UE_RETURN_ON_ERROR(BindingData.IsValid(), TEXT("Attribute binding data should be valid"));
 		}
 
diff --git a/Engine/Plugins/Experimental/UAF/UAFTestSuites/Source/UAFTestSuite/Private/UAF/ValueRuntime/IteratorUtilsTest.cpp b/Engine/Plugins/Experimental/UAF/UAFTestSuites/Source/UAFTestSuite/Private/UAF/ValueRuntime/IteratorUtilsTest.cpp
--- a/Engine/Plugins/Experimental/UAF/UAFTestSuites/Source/UAFTestSuite/Private/UAF/ValueRuntime/IteratorUtilsTest.cpp
+++ b/Engine/Plugins/Experimental/UAF/UAFTestSuites/Source/UAFTestSuite/Private/UAF/ValueRuntime/IteratorUtilsTest.cpp
@@ -105,7 +105,7 @@ namespace UE::UAF::Tests
 			Binding->AddBoneToSet(TEXT("RightArm"), TEXT("UpperBody"));
 			Binding->AddAttributeToSet(FAnimationAttributeIdentifier(TEXT("AttrFloat1"), INDEX_NONE, NAME_None, FFloatAnimationAttribute::StaticStruct()), TEXT("UpperBody"));
 
-			BindingData = MakeAttributeBindingData(*Binding);
+			BindingData = MakeAttributeBindingData(Binding, nullptr);
 			UE_RETURN_ON_ERROR(BindingData.IsValid(), TEXT("Attribute binding data should be valid"));
 		}
 
@@ -973,7 +973,7 @@ namespace UE::UAF::Tests
 			Binding->AddBoneToSet(TEXT("RightArm"), TEXT("UpperBody"));
 			Binding->AddAttributeToSet(FAnimationAttributeIdentifier(TEXT("AttrFloat1"), INDEX_NONE, NAME_None, FFloatAnimationAttribute::StaticStruct()), TEXT("UpperBody"));
 
-			BindingData = MakeAttributeBindingData(*Binding);
+			BindingData = MakeAttributeBindingData(Binding, nullptr);
 			UE_RETURN_ON_ERROR(BindingData.IsValid(), TEXT("Attribute binding data should be valid"));
 		}
 
