UAF: Fix for some variable nodes not resolving in the details panel and also resolve entry/return nodes for functions

#Synced-CL 51579867

diff --git a/Engine/Plugins/Experimental/UAF/UAF/Source/UAFEditor/Private/AnimNextEditorModule.cpp b/Engine/Plugins/Experimental/UAF/UAF/Source/UAFEditor/Private/AnimNextEditorModule.cpp
--- a/Engine/Plugins/Experimental/UAF/UAF/Source/UAFEditor/Private/AnimNextEditorModule.cpp
+++ b/Engine/Plugins/Experimental/UAF/UAF/Source/UAFEditor/Private/AnimNextEditorModule.cpp
@@ -1258,38 +1258,73 @@ void FAnimNextEditorModule::RegisterWorkspaceDocumentTypes(Workspace::IWorkspace
 						}
 					}
 				}
+				else if (const URigVMFunctionInterfaceNode* InterfaceNode = Cast<URigVMFunctionInterfaceNode>(RigVMEdGraphNode->GetModelNode()))
+				{
+					if (const URigVMCollapseNode* CollapseNodeOwner = InterfaceNode->FindReferencedNode())
+					{
+						const URigVMGraph* ContainedGraph = CollapseNodeOwner->GetContainedGraph();
+						const IRigVMClientHost* RigVMClientHost = CollapseNodeOwner->GetImplementingOuter<IRigVMClientHost>();
+						if (ContainedGraph && RigVMClientHost)
+						{
+							if (URigVMEdGraph* ContainedEdGraph = Cast<URigVMEdGraph>(RigVMClientHost->GetEditorObjectForRigVMGraph(ContainedGraph)))
+							{							
+								ObjectsToSelect.Add(ContainedEdGraph);
+								continue;
+							}
+						}
+					}
+				}
 				else if (const URigVMVariableNode* VariableNode = Cast<URigVMVariableNode>(RigVMEdGraphNode->GetModelNode()))
 				{
-					UUAFRigVMAssetEditorData* EditorData = nullptr;
+					const UUAFRigVMAssetEditorData* EditorData = nullptr;
+					FGuid VariableGUID;
 					if (const UUAFSharedVariableNode* SharedVarNode = Cast<UUAFSharedVariableNode>(VariableNode))
 					{
 						EditorData = SharedVarNode->Asset ? UncookedOnly::FUtils::GetEditorData<UUAFRigVMAssetEditorData>(SharedVarNode->Asset.Get()) : nullptr;
+						VariableGUID = SharedVarNode->CachedGuid;
 					}
 					else
 					{
 						EditorData = VariableNode->GetTypedOuter<UUAFRigVMAssetEditorData>();
+						VariableGUID = VariableNode->GetVariableGuid();
 					}
-					
-					if (EditorData)
+
+					if (VariableGUID.IsValid())
 					{
-						UAnimNextVariableEntry* Entry = nullptr;
-						EditorData->ForEachEntryOfType<UAnimNextVariableEntry>([InGuid = VariableNode->GetVariableGuid(), &Entry](UAnimNextVariableEntry* VariableEntry)
+						if (EditorData)
 						{
-							if (VariableEntry->GetGuid() == InGuid)
+							UAnimNextVariableEntry* Entry = nullptr;
+							EditorData->ForEachEntryOfType<UAnimNextVariableEntry>([VariableGUID, &Entry](UAnimNextVariableEntry* VariableEntry)
 							{
-								Entry = VariableEntry;
-								return false;
-							}
+								if (VariableEntry->GetGuid() == VariableGUID)
+								{
+									Entry = VariableEntry;
+									return false;
+								}
 						
-							return true;
-						});
+								return true;
+							});
 
-						if (Entry)
+							if (Entry)
+							{
+								ObjectsToSelect.Add(Entry);
+								continue;
+							}
+							else
+							{
+								UE_LOGFMT(LogAnimation, Warning, "Could not find a valid entry for [{0}] in [{1}]", VariableNode->GetVariableName(), EditorData->GetFName());
+							}
+						}
+						else
 						{
-							ObjectsToSelect.Add(Entry);
-							continue;
+							UE_LOGFMT(LogAnimation, Warning, "Could not resolve valid editor data for [{0}]", VariableNode->GetVariableName());
 						}
 					}
+					else
+					{
+						UE_LOGFMT(LogAnimation, Warning, "Could not resolve variable GUID for [{0}]", VariableNode->GetVariableName());
+					}
+					
 				}
 			}
 			else if (const UEdGraphNode* Node = Cast<UEdGraphNode>(Object))
