Fine tune allocator specializations for allocator type traits

#Synced-CL 51666402

diff --git a/Engine/Plugins/Experimental/UAF/UAF/Source/UAF/Public/UAF/ValueRuntime/IndirectAllocator.h b/Engine/Plugins/Experimental/UAF/UAF/Source/UAF/Public/UAF/ValueRuntime/IndirectAllocator.h
--- a/Engine/Plugins/Experimental/UAF/UAF/Source/UAF/Public/UAF/ValueRuntime/IndirectAllocator.h
+++ b/Engine/Plugins/Experimental/UAF/UAF/Source/UAF/Public/UAF/ValueRuntime/IndirectAllocator.h
@@ -15,21 +15,23 @@ namespace UE::UAF
 	//    - When Ptr is valid and NewSize is not 0, we are re-allocating memory
 	//    - It assumes that the elements that live in said memory can be trivially copyable
 	//    - Allocations use default alignment
-	typedef uint8* (*FReallocFun)(uint8* OldPtr, uint32 OldSize, uint32 NewSize);
+	typedef uint8* (*FReallocFun)(uint8* OldPtr, size_t OldSize, size_t NewSize);
 
 	// This type can be specialized for template style allocators for easy interop
-	// with the engine allocator types
+	// with the engine allocator types.
+	// Note that we cannot support some allocators that require 'NeedsElementType'
+	// such as the inline allocator.
 	template<class AllocatorType>
 	struct FAllocatorTypeTrait
 	{
 		// Specialize this struct for your allocator by implementing this function
-		static uint8* Realloc(uint8* OldPtr, uint32 OldSize, uint32 NewSize) = delete;
+		static uint8* Realloc(uint8* OldPtr, size_t OldSize, size_t NewSize) = delete;
 	};
 
-	template<>
-	struct FAllocatorTypeTrait<FDefaultAllocator>
+	template<int IndexSize>
+	struct FAllocatorTypeTrait<TSizedDefaultAllocator<IndexSize>>
 	{
-		static uint8* Realloc(uint8* OldPtr, uint32 OldSize, uint32 NewSize)
+		static uint8* Realloc(uint8* OldPtr, size_t OldSize, size_t NewSize)
 		{
 			if (OldPtr == nullptr && NewSize == 0)
 			{
@@ -37,14 +39,14 @@ namespace UE::UAF
 				return nullptr;
 			}
 
-			return (uint8*)FMemory::Realloc(OldPtr, NewSize, DEFAULT_ALIGNMENT);
+			return (uint8*)FMemory::Realloc(OldPtr, NewSize);
 		}
 	};
 
-	template<>
-	struct FAllocatorTypeTrait<TMemStackAllocator<>>
+	template<uint32 Alignment>
+	struct FAllocatorTypeTrait<TMemStackAllocator<Alignment>>
 	{
-		static uint8* Realloc(uint8* OldPtr, uint32 OldSize, uint32 NewSize)
+		static uint8* Realloc(uint8* OldPtr, size_t OldSize, size_t NewSize)
 		{
 			if (NewSize == 0)
 			{
@@ -53,7 +55,7 @@ namespace UE::UAF
 			}
 
 			// Allocate memory from the stack.
-			uint8* NewPtr = FMemStack::Get().PushBytes(NewSize, DEFAULT_ALIGNMENT);
+			uint8* NewPtr = FMemStack::Get().PushBytes(NewSize, Alignment);
 
 			// If we are growing an existing allocation, copy from the old
 			if (OldPtr != nullptr)
