Update skeleton asset registry tags on UAF blend mask/profile assets to match other animation assets

#rnx
#Synced-CL 51636272

diff --git a/Engine/Plugins/Experimental/UAF/UAF/Source/UAF/Private/UAF/BlendMask/UAFBlendMask.cpp b/Engine/Plugins/Experimental/UAF/UAF/Source/UAF/Private/UAF/BlendMask/UAFBlendMask.cpp
--- a/Engine/Plugins/Experimental/UAF/UAF/Source/UAF/Private/UAF/BlendMask/UAFBlendMask.cpp
+++ b/Engine/Plugins/Experimental/UAF/UAF/Source/UAF/Private/UAF/BlendMask/UAFBlendMask.cpp
@@ -207,18 +207,6 @@ void UUAFBlendMask::PreSave(FObjectPreSaveContext ObjectSaveContext)
 	Super::PreSave(ObjectSaveContext);
 }
 
-void UUAFBlendMask::GetAssetRegistryTags(FAssetRegistryTagsContext Context) const
-{
-	Super::GetAssetRegistryTags(Context);
-
-	if (Table)
-	{
-		const FHierarchyTable_TableType_Skeleton& TableMetadata = Table->GetTableMetadata<FHierarchyTable_TableType_Skeleton>();
-
-		Context.AddTag( FAssetRegistryTag("Skeleton", FAssetData(TableMetadata.Skeleton).ToSoftObjectPath().ToString(), FAssetRegistryTag::TT_Alphabetical) );
-	}
-}
-
 void UUAFBlendMask::SetSkeleton(TObjectPtr<USkeleton> InSkeleton)
 {
 	if (Skeleton)
diff --git a/Engine/Plugins/Experimental/UAF/UAF/Source/UAF/Private/UAF/BlendProfile/UAFBlendProfile.cpp b/Engine/Plugins/Experimental/UAF/UAF/Source/UAF/Private/UAF/BlendProfile/UAFBlendProfile.cpp
--- a/Engine/Plugins/Experimental/UAF/UAF/Source/UAF/Private/UAF/BlendProfile/UAFBlendProfile.cpp
+++ b/Engine/Plugins/Experimental/UAF/UAF/Source/UAF/Private/UAF/BlendProfile/UAFBlendProfile.cpp
@@ -223,13 +223,6 @@ void UUAFBlendProfile::GetAssetRegistryTags(FAssetRegistryTagsContext Context) c
 {
 	Super::GetAssetRegistryTags(Context);
 
-	if (Table)
-	{
-		const FHierarchyTable_TableType_Skeleton& TableMetadata = Table->GetTableMetadata<FHierarchyTable_TableType_Skeleton>();
-
-		Context.AddTag( FAssetRegistryTag("Skeleton", FAssetData(TableMetadata.Skeleton).ToSoftObjectPath().ToString(), FAssetRegistryTag::TT_Alphabetical) );
-	}
-
 	using namespace UE::UAF;
 	
 	Context.AddTag(FAssetRegistryTag(
diff --git a/Engine/Plugins/Experimental/UAF/UAF/Source/UAF/Public/UAF/BlendMask/UAFBlendMask.h b/Engine/Plugins/Experimental/UAF/UAF/Source/UAF/Public/UAF/BlendMask/UAFBlendMask.h
--- a/Engine/Plugins/Experimental/UAF/UAF/Source/UAF/Public/UAF/BlendMask/UAFBlendMask.h
+++ b/Engine/Plugins/Experimental/UAF/UAF/Source/UAF/Public/UAF/BlendMask/UAFBlendMask.h
@@ -32,7 +32,6 @@ public:
 #if WITH_EDITOR
 	UE_API virtual void BeginDestroy() override;
 	UE_API virtual void PreSave(FObjectPreSaveContext ObjectSaveContext) override;
-	UE_API virtual void GetAssetRegistryTags(FAssetRegistryTagsContext Context) const override;
 #endif
 	// End UObject
 
@@ -66,7 +65,7 @@ private:
 	void HandleSkeletonHierarchyChanged();
 #endif // WITH_EDITOR
 
-	UPROPERTY(VisibleAnywhere, Category = Default)
+	UPROPERTY(VisibleAnywhere, AssetRegistrySearchable, Category = Default)
 	TObjectPtr<USkeleton> Skeleton;
 
 	friend class UUAFBlendMaskFactory;
diff --git a/Engine/Plugins/Experimental/UAF/UAF/Source/UAF/Public/UAF/BlendProfile/UAFBlendProfile.h b/Engine/Plugins/Experimental/UAF/UAF/Source/UAF/Public/UAF/BlendProfile/UAFBlendProfile.h
--- a/Engine/Plugins/Experimental/UAF/UAF/Source/UAF/Public/UAF/BlendProfile/UAFBlendProfile.h
+++ b/Engine/Plugins/Experimental/UAF/UAF/Source/UAF/Public/UAF/BlendProfile/UAFBlendProfile.h
@@ -88,7 +88,7 @@ private:
 	void HandleSkeletonHierarchyChanged();
 #endif // WITH_EDITOR
 
-	UPROPERTY(VisibleAnywhere, Category = Default)
+	UPROPERTY(VisibleAnywhere, AssetRegistrySearchable, Category = Default)
 	TObjectPtr<USkeleton> Skeleton;
 
 	friend class UUAFBlendProfileFactory;
