UAF: Include subclasses when querying for supported object classes in the AssetDataFactory

#Synced-CL 51756718

diff --git a/Engine/Plugins/Experimental/UAF/UAF/Source/UAF/Private/Factory/UAFAssetFactory.cpp b/Engine/Plugins/Experimental/UAF/UAF/Source/UAF/Private/Factory/UAFAssetFactory.cpp
--- a/Engine/Plugins/Experimental/UAF/UAF/Source/UAF/Private/Factory/UAFAssetFactory.cpp
+++ b/Engine/Plugins/Experimental/UAF/UAF/Source/UAF/Private/Factory/UAFAssetFactory.cpp
@@ -59,18 +59,24 @@ void FAssetDataFactory::UnregisterAsset(const FTopLevelAssetPath& InObjectClassP
 
 TArray<UClass*> FAssetDataFactory::GetRegisteredObjectClasses()
 {
-	TArray<UClass*> AllSupportedObjectClasses;
+	TSet<UClass*> AllSupportedObjectClasses;
 	AllSupportedObjectClasses.Reserve(AllSupportedObjectClassPaths.Num());
 	for (const FTopLevelAssetPath& ClassName : AllSupportedObjectClassPaths)
 	{
 		if (UClass* Class = FindObject<UClass>(nullptr, ClassName.ToString()))
 		{
 			AllSupportedObjectClasses.Add(Class);
+			
+			// Add all derived classes 
+			TArray<UClass*> DerivedClasses;
+			GetDerivedClasses(Class, DerivedClasses, /*bRecursive=*/true);
+			AllSupportedObjectClasses.Append(DerivedClasses);
 		}
 	}
-	return AllSupportedObjectClasses;
+	return AllSupportedObjectClasses.Array();
 }
 
+	
 TArray<UClass*> FAssetDataFactory::GetRegisteredObjectClasses(TSubScriptStructOf<FUAFAssetData> BaseAssetType)
 {
 	TArray<UClass*> CompatibleClasses;
@@ -83,12 +89,36 @@ TArray<UClass*> FAssetDataFactory::GetRegisteredObjectClasses(TSubScriptStructOf
 			if (UClass* Class = FindObject<UClass>(nullptr, Info.Key.ToString()))
 			{
 				CompatibleClasses.Add(Class);
+				
+				// Add any valid subclasses of this class that are not registered themselves or have a closer parent with a registered initializer
+				GatherDerivedCompatibleObjectClasses(Class, CompatibleClasses);
 			}
 		}
 	}
 
 	return CompatibleClasses;
 }
+	
+void FAssetDataFactory::GatherDerivedCompatibleObjectClasses(const UClass* SuperClass, TArray<UClass*>& InOutCompatibleClasses)
+{
+	if (!SuperClass)
+	{
+		return;
+	}
+	
+	TArray<UClass*> DerivedClasses;
+	GetDerivedClasses(SuperClass, DerivedClasses, false);
+	for (UClass* DerivedClass : DerivedClasses)
+	{
+		// If the derived class does not have a registered initializer it's a valid subclass we want to add and explore it's derived classes as well
+		if (DerivedClass && !AssetInitializers.Contains(DerivedClass->GetClassPathName()))
+		{
+			InOutCompatibleClasses.Add(DerivedClass);
+			GatherDerivedCompatibleObjectClasses(DerivedClass, InOutCompatibleClasses);
+		}
+	}
+}
+	
 }
 
 
diff --git a/Engine/Plugins/Experimental/UAF/UAF/Source/UAF/Public/UAF/UAFAssetFactory.h b/Engine/Plugins/Experimental/UAF/UAF/Source/UAF/Public/UAF/UAFAssetFactory.h
--- a/Engine/Plugins/Experimental/UAF/UAF/Source/UAF/Public/UAF/UAFAssetFactory.h
+++ b/Engine/Plugins/Experimental/UAF/UAF/Source/UAF/Public/UAF/UAFAssetFactory.h
@@ -107,7 +107,12 @@ private:
 	static UAF_API const FInitializerInfo* FindAssetInitializer(const UClass* Class);
 	static TMap<FTopLevelAssetPath, FInitializerInfo> AssetInitializers;
 	static TArray<FTopLevelAssetPath> AllSupportedObjectClassPaths;
-
+	
+	// Helper function to recursively gather all valid non-registered subclasses of SuperClass and add them into InOutCompatibleClasses
+	// Stops each branch at the first subclass that has its own registered initializer in AssetInitializers to avoid adding subclasses that have their 
+	// own or closer related registered initializer already
+	static void GatherDerivedCompatibleObjectClasses(const UClass* SuperClass, TArray<UClass*>& InOutCompatibleClasses); 
+	
 	friend class FUAFAssetDataPropertyCustomization;
 };
 }
