Automatically mark function as public when attempting to add a reference to it outside of its own graph(s)

#Synced-CL 51820569

diff --git a/Engine/Plugins/Experimental/UAF/UAF/Source/UAFEditor/Private/Common/GraphEditorSchemaActions.cpp b/Engine/Plugins/Experimental/UAF/UAF/Source/UAFEditor/Private/Common/GraphEditorSchemaActions.cpp
--- a/Engine/Plugins/Experimental/UAF/UAF/Source/UAFEditor/Private/Common/GraphEditorSchemaActions.cpp
+++ b/Engine/Plugins/Experimental/UAF/UAF/Source/UAFEditor/Private/Common/GraphEditorSchemaActions.cpp
@@ -482,8 +482,26 @@ UEdGraphNode* FAnimNextSchemaAction_Function::PerformAction(UEdGraph* ParentGrap
 		FName Name = UE::UAF::Editor::FUtils::ValidateName(Cast<UObject>(Host), ReferencedPublicFunctionHeader.Name.ToString());
 		URigVMController* Controller = EdGraph->GetController();
 
+		FScopedTransaction Transaction(FText::Format(LOCTEXT("AddFunctionReferenceTransaction", "Add {0} Function Reference Node"), FText::FromName(Name)));
 		Controller->OpenUndoBracket(FString::Printf(TEXT("Add '%s' Node"), *Name.ToString()));
 
+		URigVMController* FunctionGraphController = nullptr;
+		// Make function public automatically when attempting to place a reference to it within a different outer
+		TScriptInterface<IRigVMGraphFunctionHost> FunctionHost = ReferencedPublicFunctionHeader.GetFunctionHostObject();
+		IRigVMGraphFunctionHost* ReferenceHost = EdGraph->GetImplementingOuter<IRigVMGraphFunctionHost>();
+		if (FunctionHost && FunctionHost != ReferenceHost)
+		{
+			if (!FunctionHost->GetRigVMGraphFunctionStore()->IsFunctionPublic(ReferencedPublicFunctionHeader.LibraryPointer))
+			{
+				if (const UUAFRigVMAssetEditorData* FunctionEditorData = Cast<UUAFRigVMAssetEditorData>(ReferencedPublicFunctionHeader.LibraryPointer.HostObject.TryLoad()))
+				{
+					FunctionGraphController = FunctionEditorData->GetController(FunctionEditorData->GetLocalFunctionLibrary()->GetRootGraph());
+					FunctionGraphController->OpenUndoBracket(FString::Printf(TEXT("Mark '%s' as Public"), *Name.ToString()));
+					FunctionGraphController->MarkFunctionAsPublic(ReferencedPublicFunctionHeader.Name, true, true, true);
+				}
+			}
+		}
+
 		if (URigVMFunctionReferenceNode* ModelNode = Controller->AddFunctionReferenceNodeFromDescription(ReferencedPublicFunctionHeader, FDeprecateSlateVector2D(Location), Name.ToString(), true, true))
 		{
 			NewNode = Cast<URigVMEdGraphNode>(EdGraph->FindNodeForModelNodeName(ModelNode->GetFName()));
@@ -505,10 +523,21 @@ UEdGraphNode* FAnimNextSchemaAction_Function::PerformAction(UEdGraph* ParentGrap
 			}
 
 			Controller->CloseUndoBracket();
+			if (FunctionGraphController)
+			{
+				FunctionGraphController->CloseUndoBracket();
+			}
 		}
 		else
 		{
 			Controller->CancelUndoBracket();
+			
+			// Also cancel/undo marking referenced function as public 
+			if (FunctionGraphController)
+			{
+				FunctionGraphController->CancelUndoBracket();
+			}
+			Transaction.Cancel();
 		}
 	}
 
diff --git a/Engine/Plugins/Experimental/UAF/UAF/Source/UAFUncookedOnly/Internal/AnimNextRigVMAssetEditorData.h b/Engine/Plugins/Experimental/UAF/UAF/Source/UAFUncookedOnly/Internal/AnimNextRigVMAssetEditorData.h
--- a/Engine/Plugins/Experimental/UAF/UAF/Source/UAFUncookedOnly/Internal/AnimNextRigVMAssetEditorData.h
+++ b/Engine/Plugins/Experimental/UAF/UAF/Source/UAFUncookedOnly/Internal/AnimNextRigVMAssetEditorData.h
@@ -268,6 +268,7 @@ protected:
 	friend class UE::UAF::Editor::FVariableBindingPropertyCustomization;
 	friend class UAnimNextAssetFindReplaceVariables;
 	friend struct FAnimNextSchemaAction_Variable;
+	friend struct FAnimNextSchemaAction_Function;
 	friend struct FAnimNextModuleTests;
 	friend struct FAnimNextAnimationGraphTests;
 	friend FAnimNextRigVMAssetCompileContext;
