UAF: Fix user defined enums not being fully supported by param type utils and fix for set variable node wildcard pin not updating correctly after link disconnect

#Synced-CL 51881469

diff --git a/Engine/Plugins/Experimental/UAF/UAF/Source/UAF/Private/Param/ParamUtils.cpp b/Engine/Plugins/Experimental/UAF/UAF/Source/UAF/Private/Param/ParamUtils.cpp
--- a/Engine/Plugins/Experimental/UAF/UAF/Source/UAF/Private/Param/ParamUtils.cpp
+++ b/Engine/Plugins/Experimental/UAF/UAF/Source/UAF/Private/Param/ParamUtils.cpp
@@ -199,6 +199,17 @@ FParamCompatibility FParamUtils::GetCompatibility(const FAnimNextParamType& InLH
 			}
 		}
 		break;
+	case EPropertyBagPropertyType::Enum:
+		switch (InLHS.GetValueType())
+		{
+			case EPropertyBagPropertyType::Enum:
+			if(InLHS.GetValueTypeObject() == InRHS.GetValueTypeObject())
+			{
+				return EParamCompatibility::Compatible_Equal;
+			}
+			break;
+		}
+		break;
 	}
 
 	return EParamCompatibility::Incompatible;
diff --git a/Engine/Plugins/Experimental/UAF/UAF/Source/UAFUncookedOnly/Internal/Variables/K2Node_AnimNextComponentSetVariable.h b/Engine/Plugins/Experimental/UAF/UAF/Source/UAFUncookedOnly/Internal/Variables/K2Node_AnimNextComponentSetVariable.h
--- a/Engine/Plugins/Experimental/UAF/UAF/Source/UAFUncookedOnly/Internal/Variables/K2Node_AnimNextComponentSetVariable.h
+++ b/Engine/Plugins/Experimental/UAF/UAF/Source/UAFUncookedOnly/Internal/Variables/K2Node_AnimNextComponentSetVariable.h
@@ -18,6 +18,7 @@ private:
 	// //~ Begin UEdGraphNode Interface.
 	UAFUNCOOKEDONLY_API virtual void ReconstructNode() override;
 	UAFUNCOOKEDONLY_API virtual void PinDefaultValueChanged(UEdGraphPin* Pin) override;
+	UAFUNCOOKEDONLY_API virtual void PinConnectionListChanged(UEdGraphPin* Pin) override;
 	// //~ End UEdGraphNode Interface.
 
 	//~ Begin K2Node Interface
diff --git a/Engine/Plugins/Experimental/UAF/UAF/Source/UAFUncookedOnly/Private/UncookedOnlyUtils.cpp b/Engine/Plugins/Experimental/UAF/UAF/Source/UAFUncookedOnly/Private/UncookedOnlyUtils.cpp
--- a/Engine/Plugins/Experimental/UAF/UAF/Source/UAFUncookedOnly/Private/UncookedOnlyUtils.cpp
+++ b/Engine/Plugins/Experimental/UAF/UAF/Source/UAFUncookedOnly/Private/UncookedOnlyUtils.cpp
@@ -1091,6 +1091,14 @@ FEdGraphPinType FUtils::GetPinTypeFromParamType(const FAnimNextParamType& InPara
 	case EPropertyBagPropertyType::Enum:
 		PinType.PinCategory = UEdGraphSchema_K2::PC_Enum;
 		PinType.PinSubCategoryObject = const_cast<UObject*>(InParamType.ValueTypeObject.Get());
+		if (UEnum* Enum = Cast<UEnum>(PinType.PinSubCategoryObject))
+		{
+			if (Enum->GetMaxEnumValue() <= (int64)std::numeric_limits<uint8>::max())
+			{
+				// Use byte for BP. It will use the correct picker and enum k2 node.
+				PinType.PinCategory = UEdGraphSchema_K2::PC_Byte;
+			}
+		}
 		break;
 	case EPropertyBagPropertyType::Struct:
 		PinType.PinCategory = UEdGraphSchema_K2::PC_Struct;
diff --git a/Engine/Plugins/Experimental/UAF/UAF/Source/UAFUncookedOnly/Private/Variables/K2Node_AnimNextComponentSetVariable.cpp b/Engine/Plugins/Experimental/UAF/UAF/Source/UAFUncookedOnly/Private/Variables/K2Node_AnimNextComponentSetVariable.cpp
--- a/Engine/Plugins/Experimental/UAF/UAF/Source/UAFUncookedOnly/Private/Variables/K2Node_AnimNextComponentSetVariable.cpp
+++ b/Engine/Plugins/Experimental/UAF/UAF/Source/UAFUncookedOnly/Private/Variables/K2Node_AnimNextComponentSetVariable.cpp
@@ -41,6 +41,17 @@ void UK2Node_AnimNextComponentSetVariable::PinDefaultValueChanged(UEdGraphPin* P
 	}
 }
 
+void UK2Node_AnimNextComponentSetVariable::PinConnectionListChanged(UEdGraphPin* Pin)
+{
+	Super::PinConnectionListChanged(Pin);
+	
+	if (Pin == FindPinChecked(TEXT("Value")))
+	{
+		CachePinType();
+		PropagateCachedTypeToPin();
+	}
+}
+
 void UK2Node_AnimNextComponentSetVariable::CachePinType()
 {
 	// Cache the pin type if it has been set as a literal type via default value
