Fix value pin type not updating when changing variable type in Set/GetVariable

#rnx
#Synced-CL 51636272

diff --git a/Engine/Plugins/Experimental/UAF/UAF/Source/UAFUncookedOnly/Private/Variables/K2Node_AnimNextComponentSetVariable.cpp b/Engine/Plugins/Experimental/UAF/UAF/Source/UAFUncookedOnly/Private/Variables/K2Node_AnimNextComponentSetVariable.cpp
--- a/Engine/Plugins/Experimental/UAF/UAF/Source/UAFUncookedOnly/Private/Variables/K2Node_AnimNextComponentSetVariable.cpp
+++ b/Engine/Plugins/Experimental/UAF/UAF/Source/UAFUncookedOnly/Private/Variables/K2Node_AnimNextComponentSetVariable.cpp
@@ -76,9 +76,15 @@ void UK2Node_AnimNextComponentSetVariable::PropagateCachedTypeToPin()
 	// Find the value of the variable pin
 	UEdGraphPin* VariablePin = FindPinChecked(TEXT("Variable"));
 	UEdGraphPin* ValuePin = FindPinChecked(TEXT("Value"));
-	if (!CachedPinType.PinCategory.IsNone() && VariablePin->LinkedTo.Num() == 0 && ValuePin->LinkedTo.Num() == 0)
+
+	if (!CachedPinType.PinCategory.IsNone() && ValuePin->PinType != CachedPinType)
 	{
+		ValuePin->BreakAllPinLinks();
 		ValuePin->PinType = CachedPinType;
+	}
+	
+	if (!CachedPinType.PinCategory.IsNone() && VariablePin->LinkedTo.Num() == 0 && ValuePin->LinkedTo.Num() == 0)
+	{
 		if (!ValuePin->GetSchema()->IsCurrentPinDefaultValid(ValuePin).IsEmpty())
 		{
 			ValuePin->GetSchema()->ResetPinToAutogeneratedDefaultValue(ValuePin);
diff --git a/Engine/Plugins/Experimental/UAF/UAF/Source/UAFUncookedOnly/Private/Variables/K2Node_UAFComponentGetVariable.cpp b/Engine/Plugins/Experimental/UAF/UAF/Source/UAFUncookedOnly/Private/Variables/K2Node_UAFComponentGetVariable.cpp
--- a/Engine/Plugins/Experimental/UAF/UAF/Source/UAFUncookedOnly/Private/Variables/K2Node_UAFComponentGetVariable.cpp
+++ b/Engine/Plugins/Experimental/UAF/UAF/Source/UAFUncookedOnly/Private/Variables/K2Node_UAFComponentGetVariable.cpp
@@ -65,9 +65,15 @@ void UK2Node_UAFComponentGetVariable::PropagateCachedTypeToPin()
 	// Find the value of the variable pin
 	UEdGraphPin* VariablePin = FindPinChecked(TEXT("Variable"));
 	UEdGraphPin* ValuePin = FindPinChecked(TEXT("Value"));
-	if (!CachedPinType.PinCategory.IsNone() && VariablePin->LinkedTo.Num() == 0 && ValuePin->LinkedTo.Num() == 0)
+
+	if (!CachedPinType.PinCategory.IsNone() && ValuePin->PinType != CachedPinType)
 	{
+		ValuePin->BreakAllPinLinks();
 		ValuePin->PinType = CachedPinType;
+	}
+	
+	if (!CachedPinType.PinCategory.IsNone() && VariablePin->LinkedTo.Num() == 0 && ValuePin->LinkedTo.Num() == 0)
+	{
 		if (!ValuePin->GetSchema()->IsCurrentPinDefaultValid(ValuePin).IsEmpty())
 		{
 			ValuePin->GetSchema()->ResetPinToAutogeneratedDefaultValue(ValuePin);
