Remove outdated comment on UUAFComponent Set/GetVariable

#rnx
#Synced-CL 51636272

diff --git a/Engine/Plugins/Experimental/UAF/UAF/Source/UAF/Public/Component/AnimNextComponent.h b/Engine/Plugins/Experimental/UAF/UAF/Source/UAF/Public/Component/AnimNextComponent.h
--- a/Engine/Plugins/Experimental/UAF/UAF/Source/UAF/Public/Component/AnimNextComponent.h
+++ b/Engine/Plugins/Experimental/UAF/UAF/Source/UAF/Public/Component/AnimNextComponent.h
@@ -95,7 +95,6 @@ class UUAFComponent : public UActorComponent, public UE::UAF::ISystemOutputAdapt
 	
 public:
 	// Sets a system variable's value.
-	// Note that variables are double-buffered so writing and reading the same variable won't return the new value unless UAF is updated in between.
 	// @param    Name     The name of the variable to set
 	// @param    Value    The value to set the variable to
 	UE_DEPRECATED(5.6, "This function is no longer used, please use SetVariableReference (for Blueprint) or SetVariable that takes a FAnimNextVariableReference")
@@ -103,21 +102,18 @@ public:
 	UAF_API void BlueprintSetVariable(FName Name, int32 Value);
 
 	// Sets a system variable's value.
-	// Note that variables are double-buffered so writing and reading the same variable won't return the new value unless UAF is updated in between.
 	// @param    Variable The variable to set
 	// @param    Value    The value to set the variable to
 	UFUNCTION(BlueprintCallable, BlueprintInternalUseOnly, Category = "UAF", DisplayName = "Set Variable", CustomThunk, meta = (CustomStructureParam = Value, AutoCreateRefTerm = "Variable, Value", UnsafeDuringActorConstruction))
 	UAF_API void BlueprintSetVariableReference(const FAnimNextVariableReference& Variable, const int32& Value);
 
 	// Gets a system variable's value.
-	// Note that variables are double-buffered so writing and reading the same variable won't return the new value unless UAF is updated in between.
 	// @param    Variable The variable to get
 	// @param    Value    The return value of the variable
 	UFUNCTION(BlueprintCallable, BlueprintInternalUseOnly, Category = "UAF", DisplayName = "Get Variable", CustomThunk, meta = (CustomStructureParam = Value, AutoCreateRefTerm = "Variable, Value", UnsafeDuringActorConstruction))
 	UAF_API void BlueprintGetVariableReference(const FAnimNextVariableReference& Variable, int32& Value);
 	
 	// Set a variable.
-	// Note that variables are double-buffered so writing and reading the same variable won't return the new value unless UAF is updated in between.
 	// @param    InVariable         The variable to set
 	// @param    InValue            The value to set the variable to
 	// @return true if the variable was set correctly
@@ -128,7 +124,6 @@ public:
 	}
 
 	// Get a variable.
-	// Note that variables are double-buffered so writing and reading the same variable won't return the new value unless UAF is updated in between.
 	// @param    InVariable         The variable to get
 	// @param    OutValue           The return value of the variable
 	// @return true if the variable was gotten correctly
@@ -139,7 +134,6 @@ public:
 	}
 
 	// Accesses a variable for writing.
-	// General read access is not provided via this API due to the current double-buffering strategy used to communicate variable writes to worker threads.
 	// This is intended to allow for copy-free writing of larger data structures & arrays, rather than read access. 
 	// @param    InVariable         The variable to set
 	// @param    InFunction         The function used to modify the variable, called back immediately on the game thread if the variable exists.
