UAF: Fix patches not being applied to UAF graphs

#rnx
#Synced-CL 51888919

diff --git a/Engine/Plugins/Experimental/UAF/UAF/Source/UAFUncookedOnly/Internal/AnimNextRigVMAssetSchema.h b/Engine/Plugins/Experimental/UAF/UAF/Source/UAFUncookedOnly/Internal/AnimNextRigVMAssetSchema.h
--- a/Engine/Plugins/Experimental/UAF/UAF/Source/UAFUncookedOnly/Internal/AnimNextRigVMAssetSchema.h
+++ b/Engine/Plugins/Experimental/UAF/UAF/Source/UAFUncookedOnly/Internal/AnimNextRigVMAssetSchema.h
@@ -26,6 +26,8 @@ class UUAFRigVMAssetSchema : public URigVMSchema
 
 protected:
 	UE_API UUAFRigVMAssetSchema();
+	
+	virtual bool ShouldPatchGraphSchema() const override { return false; }
 };
 
 #undef UE_API
diff --git a/Engine/Plugins/Experimental/UAF/UAF/Source/UAFUncookedOnly/Private/AnimNextRigVMAssetEditorData.cpp b/Engine/Plugins/Experimental/UAF/UAF/Source/UAFUncookedOnly/Private/AnimNextRigVMAssetEditorData.cpp
--- a/Engine/Plugins/Experimental/UAF/UAF/Source/UAFUncookedOnly/Private/AnimNextRigVMAssetEditorData.cpp
+++ b/Engine/Plugins/Experimental/UAF/UAF/Source/UAFUncookedOnly/Private/AnimNextRigVMAssetEditorData.cpp
@@ -356,7 +356,11 @@ void UUAFRigVMAssetEditorData::HandlePackageDone()
 {
 	FCoreUObjectDelegates::OnEndLoadPackage.RemoveAll(this);
 
-	GetRigVMClient()->RefreshAllModels(ERigVMLoadType::PostLoad, false, bIsCompiling);
+	{
+		TGuardValue<bool> ReentrantGuardSelf(bSuspendModelNotificationsForSelf, true);
+		GetRigVMClient()->PatchModelsOnLoad();
+		GetRigVMClient()->RefreshAllModels(ERigVMLoadType::PostLoad, false, bIsCompiling);
+	}
 
 	ConditionalPatchFunctionsOnLoad();
 
@@ -1504,114 +1508,117 @@ void UUAFRigVMAssetEditorData::HandleModifiedEvent(ERigVMGraphNotifType InNotifT
 	
 	bool bNotifForOthersPending = true;
 
-	switch(InNotifType)
+	if (!bSuspendModelNotificationsForSelf)
 	{
-	case ERigVMGraphNotifType::InteractionBracketOpened:
+		switch(InNotifType)
 		{
-			IncrementVMRecompileBracket();
-			break;
-		}
-	case ERigVMGraphNotifType::InteractionBracketClosed:
-	case ERigVMGraphNotifType::InteractionBracketCanceled:
-		{
-			DecrementVMRecompileBracket();
-			break;
-		}
-	case ERigVMGraphNotifType::NodeAdded:
-		{
-			if (URigVMCollapseNode* CollapseNode = Cast<URigVMCollapseNode>(InSubject))
+		case ERigVMGraphNotifType::InteractionBracketOpened:
 			{
-				CreateEdGraphForCollapseNode(CollapseNode, false);
+				IncrementVMRecompileBracket();
 				break;
 			}
-			RequestAutoVMRecompilation();
-			break;
-	}
-	case ERigVMGraphNotifType::NodeRemoved:
-	{
-		if (URigVMCollapseNode* CollapseNode = Cast<URigVMCollapseNode>(InSubject))
-		{
-			bNotifForOthersPending = !RemoveEdGraphForCollapseNode(CollapseNode, true);
-			break;
-		}
-		RequestAutoVMRecompilation();
-		break;
-	}
-	case ERigVMGraphNotifType::NodeRenamed:
-	{
-		if (URigVMCollapseNode* CollapseNode = Cast<URigVMCollapseNode>(InSubject))
-		{
-			FString NewNodePath = CollapseNode->GetNodePath(true /* recursive */);
-			FString Left, Right = NewNodePath;
-			URigVMNode::SplitNodePathAtEnd(NewNodePath, Left, Right);
-			FString OldNodePath = CollapseNode->GetPreviousFName().ToString();
-			if (!Left.IsEmpty())
+		case ERigVMGraphNotifType::InteractionBracketClosed:
+		case ERigVMGraphNotifType::InteractionBracketCanceled:
 			{
-				OldNodePath = URigVMNode::JoinNodePath(Left, OldNodePath);
+				DecrementVMRecompileBracket();
+				break;
 			}
-
-			HandleRigVMGraphRenamed(GetRigVMClient(), OldNodePath, NewNodePath);
-
-			if (UEdGraph* ContainedEdGraph = Cast<UEdGraph>(GetEditorObjectForRigVMGraph(CollapseNode->GetContainedGraph())))
+		case ERigVMGraphNotifType::NodeAdded:
 			{
-				ContainedEdGraph->Rename(*CollapseNode->GetEditorSubGraphName(), nullptr);
+				if (URigVMCollapseNode* CollapseNode = Cast<URigVMCollapseNode>(InSubject))
+				{
+					CreateEdGraphForCollapseNode(CollapseNode, false);
+					break;
+				}
+				RequestAutoVMRecompilation();
+				break;
 			}
-		}
-		break;
-	}
-	case ERigVMGraphNotifType::LinkAdded:
-	case ERigVMGraphNotifType::LinkRemoved:
-	case ERigVMGraphNotifType::PinArraySizeChanged:
-	case ERigVMGraphNotifType::PinDirectionChanged:
-		{
-			RequestAutoVMRecompilation();
-			break;
-		}
-
-	case ERigVMGraphNotifType::PinDefaultValueChanged:
-		{
-			if (InGraph->GetRuntimeAST().IsValid())
+		case ERigVMGraphNotifType::NodeRemoved:
 			{
-				URigVMPin* RootPin = CastChecked<URigVMPin>(InSubject)->GetRootPin();
-				FRigVMASTProxy RootPinProxy = FRigVMASTProxy::MakeFromUObject(RootPin);
-				const FRigVMExprAST* Expression = InGraph->GetRuntimeAST()->GetExprForSubject(RootPinProxy);
-				if (Expression == nullptr)
+				if (URigVMCollapseNode* CollapseNode = Cast<URigVMCollapseNode>(InSubject))
 				{
-					InGraph->ClearAST();
+					bNotifForOthersPending = !RemoveEdGraphForCollapseNode(CollapseNode, true);
+					break;
 				}
-				else if (Expression->NumParents() > 1)
+				RequestAutoVMRecompilation();
+				break;
+			}
+		case ERigVMGraphNotifType::NodeRenamed:
+			{
+				if (URigVMCollapseNode* CollapseNode = Cast<URigVMCollapseNode>(InSubject))
 				{
-					InGraph->ClearAST();
+					FString NewNodePath = CollapseNode->GetNodePath(true /* recursive */);
+					FString Left, Right = NewNodePath;
+					URigVMNode::SplitNodePathAtEnd(NewNodePath, Left, Right);
+					FString OldNodePath = CollapseNode->GetPreviousFName().ToString();
+					if (!Left.IsEmpty())
+					{
+						OldNodePath = URigVMNode::JoinNodePath(Left, OldNodePath);
+					}
+
+					HandleRigVMGraphRenamed(GetRigVMClient(), OldNodePath, NewNodePath);
+
+					if (UEdGraph* ContainedEdGraph = Cast<UEdGraph>(GetEditorObjectForRigVMGraph(CollapseNode->GetContainedGraph())))
+					{
+						ContainedEdGraph->Rename(*CollapseNode->GetEditorSubGraphName(), nullptr);
+					}
 				}
+				break;
+			}
+		case ERigVMGraphNotifType::LinkAdded:
+		case ERigVMGraphNotifType::LinkRemoved:
+		case ERigVMGraphNotifType::PinArraySizeChanged:
+		case ERigVMGraphNotifType::PinDirectionChanged:
+			{
+				RequestAutoVMRecompilation();
+				break;
 			}
 
-			RequestAutoVMRecompilation();	// We need to rebuild our metadata when a default value changes
-			break;
-		}
-	case ERigVMGraphNotifType::PinAdded:
-		{
-			// Rebind function variable pins. Needed when a function impl has new referenced vars as it's new var arg pin will be unbound.
-			RefreshFunctionImplementationGraphs();
+		case ERigVMGraphNotifType::PinDefaultValueChanged:
+			{
+				if (InGraph->GetRuntimeAST().IsValid())
+				{
+					URigVMPin* RootPin = CastChecked<URigVMPin>(InSubject)->GetRootPin();
+					FRigVMASTProxy RootPinProxy = FRigVMASTProxy::MakeFromUObject(RootPin);
+					const FRigVMExprAST* Expression = InGraph->GetRuntimeAST()->GetExprForSubject(RootPinProxy);
+					if (Expression == nullptr)
+					{
+						InGraph->ClearAST();
+					}
+					else if (Expression->NumParents() > 1)
+					{
+						InGraph->ClearAST();
+					}
+				}
 
-			if (URigVMPin* Pin = Cast<URigVMPin>(InSubject))
+				RequestAutoVMRecompilation();	// We need to rebuild our metadata when a default value changes
+				break;
+			}
+		case ERigVMGraphNotifType::PinAdded:
 			{
-				if (Pin->IsTraitPin())
+				// Rebind function variable pins. Needed when a function impl has new referenced vars as it's new var arg pin will be unbound.
+				RefreshFunctionImplementationGraphs();
+
+				if (URigVMPin* Pin = Cast<URigVMPin>(InSubject))
 				{
-					RequestAutoVMRecompilation();
+					if (Pin->IsTraitPin())
+					{
+						RequestAutoVMRecompilation();
+					}
 				}
+				break;
+			}
+		case ERigVMGraphNotifType::PinRemoved:
+			{
+				RequestAutoVMRecompilation(); // can not check if it is a trait pin, as it has been already removed
+				break;
+			}
+		case ERigVMGraphNotifType::PinCategoryChanged:
+		case ERigVMGraphNotifType::PinCategoriesChanged:
+			{
+				RequestAutoVMRecompilation();
+				break;
 			}
-			break;
-		}
-	case ERigVMGraphNotifType::PinRemoved:
-		{
-			RequestAutoVMRecompilation(); // can not check if it is a trait pin, as it has been already removed
-			break;
-		}
-	case ERigVMGraphNotifType::PinCategoryChanged:
-	case ERigVMGraphNotifType::PinCategoriesChanged:
-		{
-			RequestAutoVMRecompilation();
-			break;
 		}
 	}
 	
diff --git a/Engine/Plugins/Runtime/RigVM/Source/RigVMDeveloper/Private/RigVMModel/RigVMClient.cpp b/Engine/Plugins/Runtime/RigVM/Source/RigVMDeveloper/Private/RigVMModel/RigVMClient.cpp
--- a/Engine/Plugins/Runtime/RigVM/Source/RigVMDeveloper/Private/RigVMModel/RigVMClient.cpp
+++ b/Engine/Plugins/Runtime/RigVM/Source/RigVMDeveloper/Private/RigVMModel/RigVMClient.cpp
@@ -1359,7 +1359,12 @@ FRigVMClientPatchResult FRigVMClient::PatchModelsOnLoad()
 	for(URigVMGraph* Model : AllModels)
 	{
 		Model->PostLoad();
-		Model->SetSchemaClass(GetDefaultSchemaClass());
+		
+		// Temporary hack until we manage to allow multiple schemas inside a FRigVMClient
+		if (GetDefaultSchemaClass()->GetDefaultObject<URigVMSchema>()->ShouldPatchGraphSchema())
+		{
+			Model->SetSchemaClass(GetDefaultSchemaClass());
+		}
 		
 		URigVMController* Controller = GetOrCreateController(Model);
 		TGuardValue<bool> GuardSuspendTemplateComputation(Controller->bSuspendTemplateComputation, true);
diff --git a/Engine/Plugins/Runtime/RigVM/Source/RigVMDeveloper/Public/RigVMModel/RigVMSchema.h b/Engine/Plugins/Runtime/RigVM/Source/RigVMDeveloper/Public/RigVMModel/RigVMSchema.h
--- a/Engine/Plugins/Runtime/RigVM/Source/RigVMDeveloper/Public/RigVMModel/RigVMSchema.h
+++ b/Engine/Plugins/Runtime/RigVM/Source/RigVMDeveloper/Public/RigVMModel/RigVMSchema.h
@@ -154,6 +154,9 @@ protected:
 
 	UE_API bool IsGraphEditable(const URigVMGraph* InGraph) const;
 	UE_API TObjectPtr<URigVMNode> FindEventNode(URigVMController* InController, const UScriptStruct* InScriptStruct) const;
+	
+	
+	virtual bool ShouldPatchGraphSchema() const { return true; }
 
 	UPROPERTY(transient)
 	TObjectPtr<UScriptStruct> ExecuteContextStruct;
