UAF: Fix a crashing name clash when renaming a collapsed node's EdGraph in a function (e.g. after adding an aggregate pin to a sequence node if there is already another function under the same Outer with sequence node with extra aggregate pins which makes them collapsed nodes)

#Synced-CL 52000284

diff --git a/Engine/Plugins/Experimental/UAF/UAF/Source/UAFUncookedOnly/Private/AnimNextRigVMAssetEditorData.cpp b/Engine/Plugins/Experimental/UAF/UAF/Source/UAFUncookedOnly/Private/AnimNextRigVMAssetEditorData.cpp
--- a/Engine/Plugins/Experimental/UAF/UAF/Source/UAFUncookedOnly/Private/AnimNextRigVMAssetEditorData.cpp
+++ b/Engine/Plugins/Experimental/UAF/UAF/Source/UAFUncookedOnly/Private/AnimNextRigVMAssetEditorData.cpp
@@ -1560,7 +1560,14 @@ void UUAFRigVMAssetEditorData::HandleModifiedEvent(ERigVMGraphNotifType InNotifT
 
 					if (UEdGraph* ContainedEdGraph = Cast<UEdGraph>(GetEditorObjectForRigVMGraph(CollapseNode->GetContainedGraph())))
 					{
-						ContainedEdGraph->Rename(*CollapseNode->GetEditorSubGraphName(), nullptr);
+						UObject* Outer = FindEntryForRigVMGraph(CollapseNode->GetRootGraph());
+						if (Outer == nullptr)
+						{
+							Outer = this; // function library graph has no entry
+						}
+						
+						const FName SubGraphName = RigVMClient.GetUniqueName(Outer, *CollapseNode->GetEditorSubGraphName());
+						ContainedEdGraph->Rename(*SubGraphName.ToString(), nullptr);
 					}
 				}
 				break;
