Fix: ExecuteBindings_GT not gated by external prereqs added via module dependency traits

When a FRigVMTrait_ModuleEventDependency_ActorPrimaryTickFunction (Ordering=Before)
is applied via the initialize event, ApplyDependenciesHelper was only calling
OnAddDependency on the named tick function (e.g. PrePhysics), leaving module-owned
prerequisites in the linear chain (e.g. ExecuteBindings_GT) without the external
prerequisite. ExecuteBindings_GT would therefore run before the pawn primary tick.

Fix: in ApplyDependenciesHelper's game-thread lambda, after applying the dependency
to the matched tick function, walk its prerequisites and re-apply to any that are
module-owned (PrerequisiteObject == Object), for Before-ordered dependencies only.
This also covers removal automatically via FRigUnit_AnimNextModuleRemoveDependencies.

#Synced-CL 51902807

diff --git a/Engine/Plugins/Experimental/UAF/UAF/Source/UAF/Private/Module/RigUnit_AnimNextModuleEventDependencies.cpp b/Engine/Plugins/Experimental/UAF/UAF/Source/UAF/Private/Module/RigUnit_AnimNextModuleEventDependencies.cpp
--- a/Engine/Plugins/Experimental/UAF/UAF/Source/UAF/Private/Module/RigUnit_AnimNextModuleEventDependencies.cpp
+++ b/Engine/Plugins/Experimental/UAF/UAF/Source/UAF/Private/Module/RigUnit_AnimNextModuleEventDependencies.cpp
@@ -88,6 +88,30 @@ static void ApplyDependenciesHelper(const FAnimNextExecuteContext& InContext, co
 			FModuleDependencyContext Context(Object, *TickFunction);
 			const FRigVMTrait_ModuleEventDependency& Dependency = Dependencies[DependencyIndex].Get<FRigVMTrait_ModuleEventDependency>();
 			InPredicate(Context, Dependency);
+
+			// For Before-ordered dependencies, also apply to module-owned prerequisites
+			// (e.g. ExecuteBindings_GT chained before the user event tick function),
+			// so that external prerequisites like a pawn primary tick also gate bindings execution.
+			if(Dependency.Ordering == EAnimNextModuleEventDependencyOrdering::Before)
+			{
+				for(FTickPrerequisite& Prereq : TickFunction->GetPrerequisites())
+				{
+					if(Prereq.PrerequisiteObject == Object)
+					{
+						if(FTickFunction* PrereqTF = Prereq.Get())
+						{
+							// Only propagate to game-thread prerequisites. Worker-thread tick functions
+							// run in parallel with the game thread, so adding a game-thread actor tick
+							// as a prerequisite would incorrectly serialize parallel work.
+							if(!PrereqTF->bRunOnAnyThread)
+							{
+								FModuleDependencyContext PrereqContext(Object, *PrereqTF);
+								InPredicate(PrereqContext, Dependency);
+							}
+						}
+					}
+				}
+			}
 		}
 	});
 }
diff --git a/Engine/Plugins/Experimental/UAF/UAF/Source/UAF/Private/Module/RigVMTrait_ModuleEventDependency_ActorPrimaryTickFunction.cpp b/Engine/Plugins/Experimental/UAF/UAF/Source/UAF/Private/Module/RigVMTrait_ModuleEventDependency_ActorPrimaryTickFunction.cpp
new file mode 100644
--- /dev/null
+++ b/Engine/Plugins/Experimental/UAF/UAF/Source/UAF/Private/Module/RigVMTrait_ModuleEventDependency_ActorPrimaryTickFunction.cpp
@@ -0,0 +1,64 @@
+// Copyright Epic Games, Inc. All Rights Reserved.
+
+#include "RigVMTrait_ModuleEventDependency_ActorPrimaryTickFunction.h"
+
+#include UE_INLINE_GENERATED_CPP_BY_NAME(RigVMTrait_ModuleEventDependency_ActorPrimaryTickFunction)
+
+#if WITH_EDITOR
+FString FRigVMTrait_ModuleEventDependency_ActorPrimaryTickFunction::GetDisplayName() const
+{
+	return StaticStruct()->GetDisplayNameText().ToString();
+}
+#endif
+
+void FRigVMTrait_ModuleEventDependency_ActorPrimaryTickFunction::OnAddDependency(const UE::UAF::FModuleDependencyContext& InContext) const
+{
+	UActorComponent* AnimNextComponent = Cast<UActorComponent>(InContext.Object);
+	if(AnimNextComponent == nullptr)
+	{
+		UE_LOG(LogAnimation, Warning, TEXT("FRigVMTrait_ModuleEventDependency_ActorPrimaryTickFunction: Could not add dependency for Actor - UAF Component is not valid."));
+		return;
+	}
+
+	AActor* Actor = AnimNextComponent->GetOwner();
+	if(Actor == nullptr)
+	{
+		UE_LOG(LogAnimation, Warning, TEXT("FRigVMTrait_ModuleEventDependency_ActorPrimaryTickFunction: Could not add dependency for Actor - UAF Component has no valid owner."));
+		return;
+	}
+
+	if(Ordering == EAnimNextModuleEventDependencyOrdering::Before)
+	{
+		InContext.TickFunction.AddPrerequisite(Actor, Actor->PrimaryActorTick);
+	}
+	else
+	{
+		Actor->PrimaryActorTick.AddPrerequisite(AnimNextComponent, InContext.TickFunction);
+	}
+}
+
+void FRigVMTrait_ModuleEventDependency_ActorPrimaryTickFunction::OnRemoveDependency(const UE::UAF::FModuleDependencyContext& InContext) const
+{
+	UActorComponent* AnimNextComponent = Cast<UActorComponent>(InContext.Object);
+	if(AnimNextComponent == nullptr)
+	{
+		UE_LOG(LogAnimation, Warning, TEXT("FRigVMTrait_ModuleEventDependency_ActorPrimaryTickFunction: Could not remove dependency for Actor - UAF Component is not valid."));
+		return;
+	}
+
+	AActor* Actor = AnimNextComponent->GetOwner();
+	if(Actor == nullptr)
+	{
+		UE_LOG(LogAnimation, Warning, TEXT("FRigVMTrait_ModuleEventDependency_ActorPrimaryTickFunction: Could not remove dependency for Actor"));
+		return;
+	}
+
+	if(Ordering == EAnimNextModuleEventDependencyOrdering::Before)
+	{
+		InContext.TickFunction.RemovePrerequisite(Actor, Actor->PrimaryActorTick);
+	}
+	else
+	{
+		Actor->PrimaryActorTick.RemovePrerequisite(AnimNextComponent, InContext.TickFunction);
+	}
+}
diff --git a/Engine/Plugins/Experimental/UAF/UAF/Source/UAF/Private/Module/RigVMTrait_ModuleEventDependency_ActorPrimaryTickFunction.h b/Engine/Plugins/Experimental/UAF/UAF/Source/UAF/Private/Module/RigVMTrait_ModuleEventDependency_ActorPrimaryTickFunction.h
new file mode 100644
--- /dev/null
+++ b/Engine/Plugins/Experimental/UAF/UAF/Source/UAF/Private/Module/RigVMTrait_ModuleEventDependency_ActorPrimaryTickFunction.h
@@ -0,0 +1,24 @@
+// Copyright Epic Games, Inc. All Rights Reserved.
+
+#pragma once
+
+#include "Module/RigVMTrait_ModuleEventDependency.h"
+#include "RigVMTrait_ModuleEventDependency_ActorPrimaryTickFunction.generated.h"
+
+class UActor;
+
+// A dependency on the primary tick function of the host actor
+USTRUCT(DisplayName="Actor", meta=(ShowTooltip=true))
+struct FRigVMTrait_ModuleEventDependency_ActorPrimaryTickFunction : public FRigVMTrait_ModuleEventDependency
+{
+	GENERATED_BODY()
+
+	// FRigVMTrait interface
+#if WITH_EDITOR
+	virtual FString GetDisplayName() const override;
+#endif
+
+	// FAnimNextModuleEventDependency interface
+	virtual void OnAddDependency(const UE::UAF::FModuleDependencyContext& InContext) const override;
+	virtual void OnRemoveDependency(const UE::UAF::FModuleDependencyContext& InContext) const override;
+};
