UAF - Implement more transformer specializations for remaining blendable animation attributes

#Synced-CL 52003153

diff --git a/Engine/Plugins/Experimental/UAF/UAF/Source/UAF/Private/UAF/Attributes/EngineAttributes.cpp b/Engine/Plugins/Experimental/UAF/UAF/Source/UAF/Private/UAF/Attributes/EngineAttributes.cpp
--- a/Engine/Plugins/Experimental/UAF/UAF/Source/UAF/Private/UAF/Attributes/EngineAttributes.cpp
+++ b/Engine/Plugins/Experimental/UAF/UAF/Source/UAF/Private/UAF/Attributes/EngineAttributes.cpp
@@ -20,28 +20,65 @@ namespace UE::UAF
 		RuntimeRegistry.RegisterBoundValueMapInitializer<FIntegerAnimationAttribute>();
 		RuntimeRegistry.RegisterBoundValueMapInitializer<FStringAnimationAttribute>();
 		RuntimeRegistry.RegisterBoundValueMapInitializer<FBoneTransformAnimationAttribute>();
+		RuntimeRegistry.RegisterBoundValueMapInitializer<FTransformAnimationAttribute>();
+		RuntimeRegistry.RegisterBoundValueMapInitializer<FVectorAnimationAttribute>();
+		RuntimeRegistry.RegisterBoundValueMapInitializer<FQuaternionAnimationAttribute>();
 
 		RuntimeRegistry.RegisterUnboundValueMapInitializer<FFloatAnimationAttribute>();
 		RuntimeRegistry.RegisterUnboundValueMapInitializer<FIntegerAnimationAttribute>();
 		RuntimeRegistry.RegisterUnboundValueMapInitializer<FStringAnimationAttribute>();
 		RuntimeRegistry.RegisterUnboundValueMapInitializer<FBoneTransformAnimationAttribute>();
+		RuntimeRegistry.RegisterUnboundValueMapInitializer<FTransformAnimationAttribute>();
+		RuntimeRegistry.RegisterUnboundValueMapInitializer<FVectorAnimationAttribute>();
+		RuntimeRegistry.RegisterUnboundValueMapInitializer<FQuaternionAnimationAttribute>();
 
 		RuntimeRegistry.RegisterValueTransformer<Transformers::Private::FAccumulate_BoneTransformAttribute>();
 		RuntimeRegistry.RegisterValueTransformer<Transformers::Private::FAccumulate_FloatAttribute>();
-		RuntimeRegistry.RegisterValueTransformer<Transformers::Private::FApplyAdditiveSpace_BoneTransformAttribute>();
-		RuntimeRegistry.RegisterValueTransformer<Transformers::Private::FApplyAdditiveSpace_FloatAttribute>();
+		RuntimeRegistry.RegisterValueTransformer<Transformers::Private::FAccumulate_IntegerAttribute>();
+		RuntimeRegistry.RegisterValueTransformer<Transformers::Private::FAccumulate_TransformAttribute>();
+		RuntimeRegistry.RegisterValueTransformer<Transformers::Private::FAccumulate_VectorAttribute>();
+		RuntimeRegistry.RegisterValueTransformer<Transformers::Private::FAccumulate_QuaternionAttribute>();
+		
 		RuntimeRegistry.RegisterValueTransformer<Transformers::Private::FInterpolate_BoneTransformAttribute>();
 		RuntimeRegistry.RegisterValueTransformer<Transformers::Private::FInterpolate_FloatAttribute>();
+		RuntimeRegistry.RegisterValueTransformer<Transformers::Private::FInterpolate_IntegerAttribute>();
+		RuntimeRegistry.RegisterValueTransformer<Transformers::Private::FInterpolate_TransformAttribute>();
+		RuntimeRegistry.RegisterValueTransformer<Transformers::Private::FInterpolate_VectorAttribute>();
+		RuntimeRegistry.RegisterValueTransformer<Transformers::Private::FInterpolate_QuaternionAttribute>();
+		
 		RuntimeRegistry.RegisterValueTransformer<Transformers::Private::FLayer_BoneTransformAttribute>();
 		RuntimeRegistry.RegisterValueTransformer<Transformers::Private::FLayer_FloatAttribute>();
+		RuntimeRegistry.RegisterValueTransformer<Transformers::Private::FLayer_IntegerAttribute>();
+		RuntimeRegistry.RegisterValueTransformer<Transformers::Private::FLayer_TransformAttribute>();
+		RuntimeRegistry.RegisterValueTransformer<Transformers::Private::FLayer_VectorAttribute>();
+		RuntimeRegistry.RegisterValueTransformer<Transformers::Private::FLayer_QuaternionAttribute>();
+		
 		RuntimeRegistry.RegisterValueTransformer<Transformers::Private::FMakeAdditiveSpace_BoneTransformAttribute>();
 		RuntimeRegistry.RegisterValueTransformer<Transformers::Private::FMakeAdditiveSpace_FloatAttribute>();
+		RuntimeRegistry.RegisterValueTransformer<Transformers::Private::FMakeAdditiveSpace_IntegerAttribute>();
+		RuntimeRegistry.RegisterValueTransformer<Transformers::Private::FMakeAdditiveSpace_TransformAttribute>();
+		RuntimeRegistry.RegisterValueTransformer<Transformers::Private::FMakeAdditiveSpace_VectorAttribute>();
+		RuntimeRegistry.RegisterValueTransformer<Transformers::Private::FMakeAdditiveSpace_QuaternionAttribute>();
+
+		RuntimeRegistry.RegisterValueTransformer<Transformers::Private::FApplyAdditiveSpace_BoneTransformAttribute>();
+		RuntimeRegistry.RegisterValueTransformer<Transformers::Private::FApplyAdditiveSpace_FloatAttribute>();
+		RuntimeRegistry.RegisterValueTransformer<Transformers::Private::FApplyAdditiveSpace_IntegerAttribute>();
+		RuntimeRegistry.RegisterValueTransformer<Transformers::Private::FApplyAdditiveSpace_TransformAttribute>();
+		RuntimeRegistry.RegisterValueTransformer<Transformers::Private::FApplyAdditiveSpace_VectorAttribute>();
+		RuntimeRegistry.RegisterValueTransformer<Transformers::Private::FApplyAdditiveSpace_QuaternionAttribute>();
+		
 		RuntimeRegistry.RegisterValueTransformer<Transformers::Private::FOverwrite_BoneTransformAttribute>();
 		RuntimeRegistry.RegisterValueTransformer<Transformers::Private::FOverwrite_FloatAttribute>();
+		RuntimeRegistry.RegisterValueTransformer<Transformers::Private::FOverwrite_IntegerAttribute>();
+		RuntimeRegistry.RegisterValueTransformer<Transformers::Private::FOverwrite_TransformAttribute>();
+		RuntimeRegistry.RegisterValueTransformer<Transformers::Private::FOverwrite_VectorAttribute>();
+		RuntimeRegistry.RegisterValueTransformer<Transformers::Private::FOverwrite_QuaternionAttribute>();
+		
 		RuntimeRegistry.RegisterValueTransformer<Transformers::Private::FSanitize_BoneTransformAttribute>();
-
 #if UE_UAF_VALIDATE_MAPPED_VALUES
 		RuntimeRegistry.RegisterValueTransformer<Transformers::Private::FSanitize_FloatAttribute>();
 #endif
+		RuntimeRegistry.RegisterValueTransformer<Transformers::Private::FSanitize_TransformAttribute>();
+		RuntimeRegistry.RegisterValueTransformer<Transformers::Private::FSanitize_QuaternionAttribute>();
 	}
 }
diff --git a/Engine/Plugins/Experimental/UAF/UAF/Source/UAF/Private/UAF/ValueRuntime/Transformers/Accumulate.cpp b/Engine/Plugins/Experimental/UAF/UAF/Source/UAF/Private/UAF/ValueRuntime/Transformers/Accumulate.cpp
--- a/Engine/Plugins/Experimental/UAF/UAF/Source/UAF/Private/UAF/ValueRuntime/Transformers/Accumulate.cpp
+++ b/Engine/Plugins/Experimental/UAF/UAF/Source/UAF/Private/UAF/ValueRuntime/Transformers/Accumulate.cpp
@@ -576,5 +576,364 @@ namespace UE::UAF::Transformers
 
 			OutputMap = OutputFloats;
 		}
+
+		void FAccumulate_IntegerAttribute::TransformBoundValueMap(const FBoundValueMap* MapA, const FBoundValueMap* MapB, const FBoundValueMap* PerValueWeightMapB, float DefaultWeightB, FBoundValueMap*& OutputMap, FReallocFun OutputAllocator)
+		{
+			check(MapA->GetTypedSet() == MapB->GetTypedSet());
+			check(MapA->Num() == MapB->Num());
+
+			if (OutputMap == nullptr)
+			{
+				const FBoundValueMap::FConstructArgs Args(MapA->GetTypedSet(), MapA->GetValueType(), OutputAllocator);
+				OutputMap = MakeBoundValueMap<FIntegerAnimationAttribute>(Args);
+			}
+
+			check(MapA->GetTypedSet() == OutputMap->GetTypedSet());
+			check(MapA->Num() == OutputMap->Num());
+
+			const TBoundValueMap<FIntegerAnimationAttribute>* IntegersA = CastChecked<FIntegerAnimationAttribute>(MapA);
+			const TBoundValueMap<FIntegerAnimationAttribute>* IntegersB = CastChecked<FIntegerAnimationAttribute>(MapB);
+			TBoundValueMap<FIntegerAnimationAttribute>* OutputIntegers = CastChecked<FIntegerAnimationAttribute>(OutputMap);
+
+			const FIntegerAnimationAttribute* IntegerPtrA = IntegersA->GetData();
+			const FIntegerAnimationAttribute* IntegerPtrB = IntegersB->GetData();
+			FIntegerAnimationAttribute* OutputIntegerPtr = OutputIntegers->GetData();
+			const FIntegerAnimationAttribute* IntegerPtrAEnd = IntegerPtrA + IntegersA->Num();
+
+			if (const TBoundValueMap<FFloatAnimationAttribute>* PerValueWeightsB = Cast<FFloatAnimationAttribute>(PerValueWeightMapB))
+			{
+				check(MapA->GetTypedSet() == PerValueWeightsB->GetTypedSet());
+				check(MapA->Num() == PerValueWeightsB->Num());
+
+				const FFloatAnimationAttribute* PerValueWeightsBPtr = PerValueWeightsB->GetData();
+
+				while (IntegerPtrA < IntegerPtrAEnd)
+				{
+					OutputIntegerPtr->Value = IntegerPtrA->Value + FMath::TruncToInt32(IntegerPtrB->Value * PerValueWeightsBPtr->Value);
+
+					IntegerPtrA++;
+					IntegerPtrB++;
+					PerValueWeightsBPtr++;
+					OutputIntegerPtr++;
+				}
+			}
+			else
+			{
+				while (IntegerPtrA < IntegerPtrAEnd)
+				{
+					OutputIntegerPtr->Value = IntegerPtrA->Value + FMath::TruncToInt32(IntegerPtrB->Value * DefaultWeightB);
+
+					IntegerPtrA++;
+					IntegerPtrB++;
+					OutputIntegerPtr++;
+				}
+			}
+		}
+
+		void FAccumulate_IntegerAttribute::TransformUnboundValueMap(const FUnboundValueMap* MapA, const FUnboundValueMap* MapB, float WeightB, FUnboundValueMap*& OutputMap, FReallocFun OutputAllocator)
+		{
+			check(OutputMap == nullptr);
+
+			const TUnboundValueMap<FIntegerAnimationAttribute>* IntegersA = CastChecked<FIntegerAnimationAttribute>(MapA);
+			const TUnboundValueMap<FIntegerAnimationAttribute>* IntegersB = CastChecked<FIntegerAnimationAttribute>(MapB);
+			TUnboundValueMap<FIntegerAnimationAttribute>* OutputIntegers = MakeUnboundValueMap<FIntegerAnimationAttribute>(OutputAllocator);
+
+			OuterJoinBy(
+				FUnboundValueMapJoinOp(),
+				[&WeightB, OutputIntegers](FName Name, const FIntegerAnimationAttribute* InputA, const FIntegerAnimationAttribute* InputB)
+				{
+					// Missing values are set to 0
+					const int32 ValueA = InputA != nullptr ? InputA->Value : 0;
+					const int32 ValueB = InputB != nullptr ? InputB->Value : 0;
+
+					FIntegerAnimationAttribute Output;
+					Output.Value = ValueA + FMath::TruncToInt32(ValueB * WeightB);
+
+					OutputIntegers->Append(Name, Output);
+				},
+				IntegersA->CreateConstIterator(),
+				IntegersB->CreateConstIterator());
+
+			OutputMap = OutputIntegers;
+		}
+		
+		void FAccumulate_TransformAttribute::TransformBoundValueMap(const FBoundValueMap* MapA, const FBoundValueMap* MapB, const FBoundValueMap* PerValueWeightMapB, float DefaultWeightB, FBoundValueMap*& OutputMap, FReallocFun OutputAllocator)
+		{
+			check(MapA->GetTypedSet() == MapB->GetTypedSet());
+			check(MapA->Num() == MapB->Num());
+
+			if (OutputMap == nullptr)
+			{
+				const FBoundValueMap::FConstructArgs Args(MapA->GetTypedSet(), MapA->GetValueType(), OutputAllocator);
+				OutputMap = MakeBoundValueMap<FTransformAnimationAttribute>(Args);
+			}
+
+			check(MapA->GetTypedSet() == OutputMap->GetTypedSet());
+			check(MapA->Num() == OutputMap->Num());
+
+			const TBoundValueMap<FTransformAnimationAttribute>* TransformsA = CastChecked<FTransformAnimationAttribute>(MapA);
+			const TBoundValueMap<FTransformAnimationAttribute>* TransformsB = CastChecked<FTransformAnimationAttribute>(MapB);
+			TBoundValueMap<FTransformAnimationAttribute>* OutputTransforms = CastChecked<FTransformAnimationAttribute>(OutputMap);
+
+			const FTransformAnimationAttribute* TransformPtrA = TransformsA->GetData();
+			const FTransformAnimationAttribute* TransformPtrB = TransformsB->GetData();
+			FTransformAnimationAttribute* OutputTransformPtr = OutputTransforms->GetData();
+			const FTransformAnimationAttribute* TransformPtrAEnd = TransformPtrA + TransformsA->Num();
+
+			if (const TBoundValueMap<FFloatAnimationAttribute>* PerValueWeightsB = Cast<FFloatAnimationAttribute>(PerValueWeightMapB))
+			{
+				check(MapA->GetTypedSet() == PerValueWeightsB->GetTypedSet());
+				check(MapA->Num() == PerValueWeightsB->Num());
+
+				const FFloatAnimationAttribute* PerValueWeightsBPtr = PerValueWeightsB->GetData();
+
+				while (TransformPtrA < TransformPtrAEnd)
+				{
+					FTransform Result = TransformPtrA->Value;
+					Result.AccumulateWithShortestRotation(TransformPtrB->Value, ScalarRegister(PerValueWeightsBPtr->Value));
+
+					OutputTransformPtr->Value = Result;
+
+					TransformPtrA++;
+					TransformPtrB++;
+					PerValueWeightsBPtr++;
+					OutputTransformPtr++;
+				}
+			}
+			else
+			{
+				const ScalarRegister VDefaultWeightB(DefaultWeightB);
+
+				while (TransformPtrA < TransformPtrAEnd)
+				{
+					FTransform Result = TransformPtrA->Value;
+					Result.AccumulateWithShortestRotation(TransformPtrB->Value, VDefaultWeightB);
+
+					OutputTransformPtr->Value = Result;
+
+					TransformPtrA++;
+					TransformPtrB++;
+					OutputTransformPtr++;
+				}
+			}
+		}
+
+		void FAccumulate_TransformAttribute::TransformUnboundValueMap(const FUnboundValueMap* MapA, const FUnboundValueMap* MapB, float WeightB, FUnboundValueMap*& OutputMap, FReallocFun OutputAllocator)
+		{
+			check(OutputMap == nullptr);
+
+			const TUnboundValueMap<FTransformAnimationAttribute>* TransformsA = CastChecked<FTransformAnimationAttribute>(MapA);
+			const TUnboundValueMap<FTransformAnimationAttribute>* TransformsB = CastChecked<FTransformAnimationAttribute>(MapB);
+			TUnboundValueMap<FTransformAnimationAttribute>* OutputTransforms = MakeUnboundValueMap<FTransformAnimationAttribute>(OutputAllocator);
+
+			const ScalarRegister VWeightB(WeightB);
+
+			OuterJoinBy(
+				FUnboundValueMapJoinOp(),
+				[VWeightB, OutputTransforms](FName Name, const FTransformAnimationAttribute* InputA, const FTransformAnimationAttribute* InputB)
+				{
+					// Missing values are set to the identity
+					const FTransform& ValueA = InputA != nullptr ? InputA->Value : FTransform::Identity;
+					const FTransform& ValueB = InputB != nullptr ? InputB->Value : FTransform::Identity;
+
+					FTransform Result = ValueA;
+					Result.AccumulateWithShortestRotation(ValueB, VWeightB);
+
+					OutputTransforms->Append(Name, FTransformAnimationAttribute{ Result });
+				},
+				TransformsA->CreateConstIterator(),
+				TransformsB->CreateConstIterator());
+
+			OutputMap = OutputTransforms;
+		}
+
+		void FAccumulate_VectorAttribute::TransformBoundValueMap(const FBoundValueMap* MapA, const FBoundValueMap* MapB, const FBoundValueMap* PerValueWeightMapB, float DefaultWeightB, FBoundValueMap*& OutputMap, FReallocFun OutputAllocator)
+		{
+			check(MapA->GetTypedSet() == MapB->GetTypedSet());
+			check(MapA->Num() == MapB->Num());
+
+			if (OutputMap == nullptr)
+			{
+				const FBoundValueMap::FConstructArgs Args(MapA->GetTypedSet(), MapA->GetValueType(), OutputAllocator);
+				OutputMap = MakeBoundValueMap<FVectorAnimationAttribute>(Args);
+			}
+
+			check(MapA->GetTypedSet() == OutputMap->GetTypedSet());
+			check(MapA->Num() == OutputMap->Num());
+
+			const TBoundValueMap<FVectorAnimationAttribute>* VectorsA = CastChecked<FVectorAnimationAttribute>(MapA);
+			const TBoundValueMap<FVectorAnimationAttribute>* VectorsB = CastChecked<FVectorAnimationAttribute>(MapB);
+			TBoundValueMap<FVectorAnimationAttribute>* OutputVectors = CastChecked<FVectorAnimationAttribute>(OutputMap);
+
+			const FVectorAnimationAttribute* VectorPtrA = VectorsA->GetData();
+			const FVectorAnimationAttribute* VectorPtrB = VectorsB->GetData();
+			FVectorAnimationAttribute* OutputVectorPtr = OutputVectors->GetData();
+			const FVectorAnimationAttribute* VectorPtrAEnd = VectorPtrA + VectorsA->Num();
+
+			if (const TBoundValueMap<FFloatAnimationAttribute>* PerValueWeightsB = Cast<FFloatAnimationAttribute>(PerValueWeightMapB))
+			{
+				check(MapA->GetTypedSet() == PerValueWeightsB->GetTypedSet());
+				check(MapA->Num() == PerValueWeightsB->Num());
+
+				const FFloatAnimationAttribute* PerValueWeightsBPtr = PerValueWeightsB->GetData();
+
+				while (VectorPtrA < VectorPtrAEnd)
+				{
+					OutputVectorPtr->Value = VectorPtrA->Value + (VectorPtrB->Value * PerValueWeightsBPtr->Value);
+
+					VectorPtrA++;
+					VectorPtrB++;
+					PerValueWeightsBPtr++;
+					OutputVectorPtr++;
+				}
+			}
+			else
+			{
+				while (VectorPtrA < VectorPtrAEnd)
+				{
+					OutputVectorPtr->Value = VectorPtrA->Value + (VectorPtrB->Value * DefaultWeightB);
+
+					VectorPtrA++;
+					VectorPtrB++;
+					OutputVectorPtr++;
+				}
+			}
+		}
+
+		void FAccumulate_VectorAttribute::TransformUnboundValueMap(const FUnboundValueMap* MapA, const FUnboundValueMap* MapB, float WeightB, FUnboundValueMap*& OutputMap, FReallocFun OutputAllocator)
+		{
+			check(OutputMap == nullptr);
+
+			const TUnboundValueMap<FVectorAnimationAttribute>* VectorsA = CastChecked<FVectorAnimationAttribute>(MapA);
+			const TUnboundValueMap<FVectorAnimationAttribute>* VectorsB = CastChecked<FVectorAnimationAttribute>(MapB);
+			TUnboundValueMap<FVectorAnimationAttribute>* OutputVectors = MakeUnboundValueMap<FVectorAnimationAttribute>(OutputAllocator);
+
+			OuterJoinBy(
+				FUnboundValueMapJoinOp(),
+				[WeightB, OutputVectors](FName Name, const FVectorAnimationAttribute* InputA, const FVectorAnimationAttribute* InputB)
+				{
+					// Missing values are set to the identity
+					const FVector& ValueA = InputA != nullptr ? InputA->Value : FVector::ZeroVector;
+					const FVector& ValueB = InputB != nullptr ? InputB->Value : FVector::ZeroVector;
+
+					FVector Result = ValueA + (ValueB * WeightB);
+
+					OutputVectors->Append(Name, FVectorAnimationAttribute{ Result });
+				},
+				VectorsA->CreateConstIterator(),
+				VectorsB->CreateConstIterator());
+
+			OutputMap = OutputVectors;
+		}
+
+		void FAccumulate_QuaternionAttribute::TransformBoundValueMap(const FBoundValueMap* MapA, const FBoundValueMap* MapB, const FBoundValueMap* PerValueWeightMapB, float DefaultWeightB, FBoundValueMap*& OutputMap, FReallocFun OutputAllocator)
+		{
+			check(MapA->GetTypedSet() == MapB->GetTypedSet());
+			check(MapA->Num() == MapB->Num());
+
+			if (OutputMap == nullptr)
+			{
+				const FBoundValueMap::FConstructArgs Args(MapA->GetTypedSet(), MapA->GetValueType(), OutputAllocator);
+				OutputMap = MakeBoundValueMap<FQuaternionAnimationAttribute>(Args);
+			}
+
+			check(MapA->GetTypedSet() == OutputMap->GetTypedSet());
+			check(MapA->Num() == OutputMap->Num());
+
+			const TBoundValueMap<FQuaternionAnimationAttribute>* QuaternionsA = CastChecked<FQuaternionAnimationAttribute>(MapA);
+			const TBoundValueMap<FQuaternionAnimationAttribute>* QuaternionsB = CastChecked<FQuaternionAnimationAttribute>(MapB);
+			TBoundValueMap<FQuaternionAnimationAttribute>* OutputQuaternions = CastChecked<FQuaternionAnimationAttribute>(OutputMap);
+
+			const FQuaternionAnimationAttribute* QuaternionPtrA = QuaternionsA->GetData();
+			const FQuaternionAnimationAttribute* QuaternionPtrB = QuaternionsB->GetData();
+			FQuaternionAnimationAttribute* OutputQuaternionPtr = OutputQuaternions->GetData();
+			const FQuaternionAnimationAttribute* QuaternionPtrAEnd = QuaternionPtrA + QuaternionsA->Num();
+
+			if (const TBoundValueMap<FFloatAnimationAttribute>* PerValueWeightsB = Cast<FFloatAnimationAttribute>(PerValueWeightMapB))
+			{
+				check(MapA->GetTypedSet() == PerValueWeightsB->GetTypedSet());
+				check(MapA->Num() == PerValueWeightsB->Num());
+
+				const FFloatAnimationAttribute* PerValueWeightsBPtr = PerValueWeightsB->GetData();
+
+				while (QuaternionPtrA < QuaternionPtrAEnd)
+				{
+					const FQuat WeightedB = QuaternionPtrB->Value * PerValueWeightsBPtr->Value;
+				
+					// From VectorAccumulateQuaternionShortestPath
+					// Blend rotation
+					//     To ensure the 'shortest route', we make sure the dot product between the both rotations is positive.
+					//     const float Bias = (|A.B| >= 0 ? 1 : -1)
+					//     return A + B * Bias;
+					const FQuat::FReal DotResult = (QuaternionPtrA->Value | WeightedB);
+					const FQuat::FReal Bias = FMath::FloatSelect(DotResult, FQuat::FReal(1.0f), FQuat::FReal(-1.0f));
+					
+					OutputQuaternionPtr->Value = QuaternionPtrA->Value + (WeightedB * Bias);
+
+					QuaternionPtrA++;
+					QuaternionPtrB++;
+					PerValueWeightsBPtr++;
+					OutputQuaternionPtr++;
+				}
+			}
+			else
+			{
+				while (QuaternionPtrA < QuaternionPtrAEnd)
+				{
+					const FQuat WeightedB = QuaternionPtrB->Value * DefaultWeightB;
+				
+					// From VectorAccumulateQuaternionShortestPath
+					// Blend rotation
+					//     To ensure the 'shortest route', we make sure the dot product between the both rotations is positive.
+					//     const float Bias = (|A.B| >= 0 ? 1 : -1)
+					//     return A + B * Bias;
+					const FQuat::FReal DotResult = (QuaternionPtrA->Value | WeightedB);
+					const FQuat::FReal Bias = FMath::FloatSelect(DotResult, FQuat::FReal(1.0f), FQuat::FReal(-1.0f));
+					
+					OutputQuaternionPtr->Value = QuaternionPtrA->Value + (WeightedB * Bias);
+
+					QuaternionPtrA++;
+					QuaternionPtrB++;
+					OutputQuaternionPtr++;
+				}
+			}
+		}
+
+		void FAccumulate_QuaternionAttribute::TransformUnboundValueMap(const FUnboundValueMap* MapA, const FUnboundValueMap* MapB, float WeightB, FUnboundValueMap*& OutputMap, FReallocFun OutputAllocator)
+		{
+			check(OutputMap == nullptr);
+
+			const TUnboundValueMap<FQuaternionAnimationAttribute>* QuaternionsA = CastChecked<FQuaternionAnimationAttribute>(MapA);
+			const TUnboundValueMap<FQuaternionAnimationAttribute>* QuaternionsB = CastChecked<FQuaternionAnimationAttribute>(MapB);
+			TUnboundValueMap<FQuaternionAnimationAttribute>* OutputQuaternions = MakeUnboundValueMap<FQuaternionAnimationAttribute>(OutputAllocator);
+			
+			OuterJoinBy(
+				FUnboundValueMapJoinOp(),
+				[WeightB, OutputQuaternions](FName Name, const FQuaternionAnimationAttribute* InputA, const FQuaternionAnimationAttribute* InputB)
+				{
+					// Missing values are set to the identity
+					const FQuat& ValueA = InputA != nullptr ? InputA->Value : FQuat::Identity;
+					const FQuat& ValueB = InputB != nullptr ? InputB->Value : FQuat::Identity;
+
+					const FQuat WeightedB = ValueB * WeightB;
+				
+					// From VectorAccumulateQuaternionShortestPath
+					// Blend rotation
+					//     To ensure the 'shortest route', we make sure the dot product between the both rotations is positive.
+					//     const float Bias = (|A.B| >= 0 ? 1 : -1)
+					//     return A + B * Bias;
+					const FQuat::FReal DotResult = (ValueA | WeightedB);
+					const FQuat::FReal Bias = FMath::FloatSelect(DotResult, FQuat::FReal(1.0f), FQuat::FReal(-1.0f));
+
+					const FQuat Result = ValueA + (WeightedB * Bias);
+
+					OutputQuaternions->Append(Name, FQuaternionAnimationAttribute{ Result });
+				},
+				QuaternionsA->CreateConstIterator(),
+				QuaternionsB->CreateConstIterator());
+
+			OutputMap = OutputQuaternions;
+		}
+		
 	}
 }
diff --git a/Engine/Plugins/Experimental/UAF/UAF/Source/UAF/Private/UAF/ValueRuntime/Transformers/AdditiveSpace.cpp b/Engine/Plugins/Experimental/UAF/UAF/Source/UAF/Private/UAF/ValueRuntime/Transformers/AdditiveSpace.cpp
--- a/Engine/Plugins/Experimental/UAF/UAF/Source/UAF/Private/UAF/ValueRuntime/Transformers/AdditiveSpace.cpp
+++ b/Engine/Plugins/Experimental/UAF/UAF/Source/UAF/Private/UAF/ValueRuntime/Transformers/AdditiveSpace.cpp
@@ -657,6 +657,349 @@ namespace UE::UAF::Transformers
 			Output = OutputFloats;
 		}
 
+		void FMakeAdditiveSpace_IntegerAttribute::TransformBoundValueMap(const FBoundValueMap* Base, const FBoundValueMap* Input, FBoundValueMap*& Output, FReallocFun OutputAllocator)
+		{
+			if (Base == nullptr)
+			{
+				// Our base is missing, use the user authored default values
+				const FAttributeTypedSetPtr& TypedSet = Input->GetTypedSet();
+				const FValueBundlePtr ReferenceValues = TypedSet->GetBindingData()->GetDefaultAttributeValues(TypedSet->GetNamedSet()->GetName());
+				Base = ReferenceValues->GetBoundValueMaps().Find(Input->GetMappingKey());
+			}
+
+			if (Output == nullptr)
+			{
+				const FBoundValueMap::FConstructArgs Args(Input->GetTypedSet(), Input->GetValueType(), OutputAllocator);
+				Output = MakeBoundValueMap<FIntegerAnimationAttribute>(Args);
+			}
+
+			if (Base == nullptr)
+			{
+				UE_LOG(LogAnimation, Warning, TEXT("FMakeAdditiveSpace_IntegerAttribute::TransformBoundValueMap: Could not transform bound value map - no base provided and no reference found. Setting to identity instead."));
+				Output->FillWithIdentity(true);
+				return;
+			}
+
+			// Our output should have the same shape as our base
+			check(Base->GetTypedSet() == Output->GetTypedSet());
+			check(Base->Num() == Output->Num());
+
+			// Our base and input should have the same shape
+			check(Base->GetTypedSet() == Input->GetTypedSet());
+			check(Base->Num() == Input->Num());
+
+			const TBoundValueMap<FIntegerAnimationAttribute>* BaseIntegers = CastChecked<FIntegerAnimationAttribute>(Base);
+			const TBoundValueMap<FIntegerAnimationAttribute>* InputIntegers = CastChecked<FIntegerAnimationAttribute>(Input);
+			TBoundValueMap<FIntegerAnimationAttribute>* OutputIntegers = CastChecked<FIntegerAnimationAttribute>(Output);
+
+			const FIntegerAnimationAttribute* BaseIntegersPtr = BaseIntegers->GetData();
+			const FIntegerAnimationAttribute* InputIntegersPtr = InputIntegers->GetData();
+			FIntegerAnimationAttribute* OutputIntegersPtr = OutputIntegers->GetData();
+			const FIntegerAnimationAttribute* BaseIntegersPtrEnd = BaseIntegersPtr + BaseIntegers->Num();
+
+			while (BaseIntegersPtr < BaseIntegersPtrEnd)
+			{
+				OutputIntegersPtr->Value = InputIntegersPtr->Value - BaseIntegersPtr->Value;
+
+				BaseIntegersPtr++;
+				InputIntegersPtr++;
+				OutputIntegersPtr++;
+			}
+		}
+
+		void FMakeAdditiveSpace_IntegerAttribute::TransformUnboundValueMap(const FUnboundValueMap* Base, const FUnboundValueMap* Input, FUnboundValueMap*& Output, FReallocFun OutputAllocator)
+		{
+			check(Output == nullptr);
+
+			const TUnboundValueMap<FIntegerAnimationAttribute>* BaseIntegers = Cast<FIntegerAnimationAttribute>(Base);
+			const TUnboundValueMap<FIntegerAnimationAttribute>* InputIntegers = CastChecked<FIntegerAnimationAttribute>(Input);
+
+			TUnboundValueMap<FIntegerAnimationAttribute>* OutputInteger = MakeUnboundValueMap<FIntegerAnimationAttribute>(OutputAllocator);
+
+			OuterJoinBy(
+				FUnboundValueMapJoinOp(),
+				[OutputInteger](FName Name, const FIntegerAnimationAttribute* BaseValue, const FIntegerAnimationAttribute* InputValue)
+				{
+					if (InputValue == nullptr)
+					{
+						// We have no input to subtract, skip it
+						return;
+					}
+
+					// Missing values are set to 0
+					const int32 BaseInteger = BaseValue != nullptr ? BaseValue->Value : 0;
+
+					const FIntegerAnimationAttribute Output{ InputValue->Value - BaseInteger };
+					OutputInteger->Append(Name, Output);
+				},
+				BaseIntegers != nullptr ? BaseIntegers->CreateConstIterator() : TUnboundValueMap<FIntegerAnimationAttribute>::FConstIterator(),
+				InputIntegers->CreateConstIterator());
+
+			Output = OutputInteger;
+		}
+
+		void FMakeAdditiveSpace_TransformAttribute::TransformBoundValueMap(const FBoundValueMap* Base, const FBoundValueMap* Input, FBoundValueMap*& Output, FReallocFun OutputAllocator)
+		{
+			if (Base == nullptr)
+			{
+				// Our base is missing, use the reference pose
+				const FAttributeTypedSetPtr& TypedSet = Input->GetTypedSet();
+				const FValueBundlePtr ReferenceValues = TypedSet->GetBindingData()->GetDefaultAttributeValues(TypedSet->GetNamedSet()->GetName());
+				Base = ReferenceValues->GetBoundValueMaps().Find(Input->GetMappingKey());
+			}
+
+			if (Output == nullptr)
+			{
+				const FBoundValueMap::FConstructArgs Args(Input->GetTypedSet(), Input->GetValueType(), OutputAllocator);
+				Output = MakeBoundValueMap<FTransformAnimationAttribute>(Args);
+			}
+
+			if (Base == nullptr)
+			{
+				UE_LOG(LogAnimation, Warning, TEXT("FMakeAdditiveSpace_TransformAttribute::TransformBoundValueMap: Could not transform bound value map - no base provided and no reference found. Setting to identity instead."));
+				Output->FillWithIdentity(true);
+				return;
+			}
+
+			// Our output should have the same shape as our base
+			check(Base->GetTypedSet() == Output->GetTypedSet());
+			check(Base->Num() == Output->Num());
+
+			// Our base and input should have the same shape
+			check(Base->GetTypedSet() == Input->GetTypedSet());
+			check(Base->Num() == Input->Num());
+
+			const TBoundValueMap<FTransformAnimationAttribute>* BaseTransforms = CastChecked<FTransformAnimationAttribute>(Base);
+			const TBoundValueMap<FTransformAnimationAttribute>* InputTransforms = CastChecked<FTransformAnimationAttribute>(Input);
+			TBoundValueMap<FTransformAnimationAttribute>* OutputTransforms = CastChecked<FTransformAnimationAttribute>(Output);
+
+			const FTransformAnimationAttribute* BaseTransformPtr = BaseTransforms->GetData();
+			const FTransformAnimationAttribute* InputTransformPtr = InputTransforms->GetData();
+			FTransformAnimationAttribute* OutputTransformPtr = OutputTransforms->GetData();
+			const FTransformAnimationAttribute* BaseTransformPtrEnd = BaseTransformPtr + BaseTransforms->Num();
+
+			while (BaseTransformPtr < BaseTransformPtrEnd)
+			{
+				OutputTransformPtr->Value.SetRotation(InputTransformPtr->Value.GetRotation() * BaseTransformPtr->Value.GetRotation().Inverse());
+				OutputTransformPtr->Value.SetTranslation(InputTransformPtr->Value.GetTranslation() - BaseTransformPtr->Value.GetTranslation());
+				// additive scale considers how much it grow or lower
+				// in order to support blending between different additive scale, we save [(target scale)/(source scale) - 1.f], and this can blend with 
+				// other delta scale value
+				// when we apply to the another scale, we apply scale * (1 + [additive scale])
+				OutputTransformPtr->Value.SetScale3D(InputTransformPtr->Value.GetScale3D() * FTransform::GetSafeScaleReciprocal(BaseTransformPtr->Value.GetScale3D()) - 1.0f);
+				OutputTransformPtr->Value.NormalizeRotation();
+
+				BaseTransformPtr++;
+				InputTransformPtr++;
+				OutputTransformPtr++;
+			}
+		}
+
+		void FMakeAdditiveSpace_TransformAttribute::TransformUnboundValueMap(const FUnboundValueMap* Base, const FUnboundValueMap* Input, FUnboundValueMap*& Output, FReallocFun OutputAllocator)
+		{
+			check(Output == nullptr);
+
+			const TUnboundValueMap<FTransformAnimationAttribute>* BaseTransforms = Cast<FTransformAnimationAttribute>(Base);
+			const TUnboundValueMap<FTransformAnimationAttribute>* InputTransforms = CastChecked<FTransformAnimationAttribute>(Input);
+
+			TUnboundValueMap<FTransformAnimationAttribute>* OutputTransforms = MakeUnboundValueMap<FTransformAnimationAttribute>(OutputAllocator);
+
+			OuterJoinBy(
+				FUnboundValueMapJoinOp(),
+				[OutputTransforms](FName Name, const FTransformAnimationAttribute* BaseValue, const FTransformAnimationAttribute* InputValue)
+				{
+					if (InputValue == nullptr)
+					{
+						// We have no input to subtract, skip it
+						return;
+					}
+
+					// Missing values are set to the identity
+					const FTransform& BaseTransform = BaseValue != nullptr ? BaseValue->Value : FTransform::Identity;
+
+					FTransformAnimationAttribute Output;
+					Output.Value.SetRotation(InputValue->Value.GetRotation() * BaseTransform.GetRotation().Inverse());
+					Output.Value.SetTranslation(InputValue->Value.GetTranslation() - BaseTransform.GetTranslation());
+					// additive scale considers how much it grow or lower
+					// in order to support blending between different additive scale, we save [(target scale)/(source scale) - 1.f], and this can blend with 
+					// other delta scale value
+					// when we apply to the another scale, we apply scale * (1 + [additive scale])
+					Output.Value.SetScale3D(InputValue->Value.GetScale3D() * FTransform::GetSafeScaleReciprocal(BaseTransform.GetScale3D()) - 1.0f);
+					Output.Value.NormalizeRotation();
+
+					OutputTransforms->Append(Name, Output);
+				},
+				BaseTransforms != nullptr ? BaseTransforms->CreateConstIterator() : TUnboundValueMap<FTransformAnimationAttribute>::FConstIterator(),
+				InputTransforms->CreateConstIterator());
+
+			Output = OutputTransforms;
+		}
+
+		void FMakeAdditiveSpace_VectorAttribute::TransformBoundValueMap(const FBoundValueMap* Base, const FBoundValueMap* Input, FBoundValueMap*& Output, FReallocFun OutputAllocator)
+		{
+			if (Base == nullptr)
+			{
+				// Our base is missing, use the user authored default values
+				const FAttributeTypedSetPtr& TypedSet = Input->GetTypedSet();
+				const FValueBundlePtr ReferenceValues = TypedSet->GetBindingData()->GetDefaultAttributeValues(TypedSet->GetNamedSet()->GetName());
+				Base = ReferenceValues->GetBoundValueMaps().Find(Input->GetMappingKey());
+			}
+
+			if (Output == nullptr)
+			{
+				const FBoundValueMap::FConstructArgs Args(Input->GetTypedSet(), Input->GetValueType(), OutputAllocator);
+				Output = MakeBoundValueMap<FVectorAnimationAttribute>(Args);
+			}
+
+			if (Base == nullptr)
+			{
+				UE_LOG(LogAnimation, Warning, TEXT("FMakeAdditiveSpace_VectorAttribute::TransformBoundValueMap: Could not transform bound value map - no base provided and no reference found. Setting to identity instead."));
+				Output->FillWithIdentity(true);
+				return;
+			}
+
+			// Our output should have the same shape as our base
+			check(Base->GetTypedSet() == Output->GetTypedSet());
+			check(Base->Num() == Output->Num());
+
+			// Our base and input should have the same shape
+			check(Base->GetTypedSet() == Input->GetTypedSet());
+			check(Base->Num() == Input->Num());
+
+			const TBoundValueMap<FVectorAnimationAttribute>* BaseVectors = CastChecked<FVectorAnimationAttribute>(Base);
+			const TBoundValueMap<FVectorAnimationAttribute>* InputVectors = CastChecked<FVectorAnimationAttribute>(Input);
+			TBoundValueMap<FVectorAnimationAttribute>* OutputVectors = CastChecked<FVectorAnimationAttribute>(Output);
+
+			const FVectorAnimationAttribute* BaseVectorPtr = BaseVectors->GetData();
+			const FVectorAnimationAttribute* InputVectorPtr = InputVectors->GetData();
+			FVectorAnimationAttribute* OutputVectorPtr = OutputVectors->GetData();
+			const FVectorAnimationAttribute* BaseVectorPtrEnd = BaseVectorPtr + BaseVectors->Num();
+
+			while (BaseVectorPtr < BaseVectorPtrEnd)
+			{
+				OutputVectorPtr->Value = InputVectorPtr->Value - BaseVectorPtr->Value;
+
+				BaseVectorPtr++;
+				InputVectorPtr++;
+				OutputVectorPtr++;
+			}
+		}
+
+		void FMakeAdditiveSpace_VectorAttribute::TransformUnboundValueMap(const FUnboundValueMap* Base, const FUnboundValueMap* Input, FUnboundValueMap*& Output, FReallocFun OutputAllocator)
+		{
+			check(Output == nullptr);
+
+			const TUnboundValueMap<FVectorAnimationAttribute>* BaseVectors = Cast<FVectorAnimationAttribute>(Base);
+			const TUnboundValueMap<FVectorAnimationAttribute>* InputVectors = CastChecked<FVectorAnimationAttribute>(Input);
+
+			TUnboundValueMap<FVectorAnimationAttribute>* OutputVectors = MakeUnboundValueMap<FVectorAnimationAttribute>(OutputAllocator);
+
+			OuterJoinBy(
+				FUnboundValueMapJoinOp(),
+				[OutputVectors](FName Name, const FVectorAnimationAttribute* BaseValue, const FVectorAnimationAttribute* InputValue)
+				{
+					if (InputValue == nullptr)
+					{
+						// We have no input to subtract, skip it
+						return;
+					}
+
+					// Missing values are set to 0.0
+					const FVector BaseVector = BaseValue != nullptr ? BaseValue->Value : FVector::ZeroVector;
+
+					const FVectorAnimationAttribute Output{ InputValue->Value - BaseVector };
+					OutputVectors->Append(Name, Output);
+				},
+				BaseVectors != nullptr ? BaseVectors->CreateConstIterator() : TUnboundValueMap<FVectorAnimationAttribute>::FConstIterator(),
+				InputVectors->CreateConstIterator());
+
+			Output = OutputVectors;
+		}
+
+		void FMakeAdditiveSpace_QuaternionAttribute::TransformBoundValueMap(const FBoundValueMap* Base, const FBoundValueMap* Input, FBoundValueMap*& Output, FReallocFun OutputAllocator)
+		{
+			if (Base == nullptr)
+			{
+				// Our base is missing, use the user authored default values
+				const FAttributeTypedSetPtr& TypedSet = Input->GetTypedSet();
+				const FValueBundlePtr ReferenceValues = TypedSet->GetBindingData()->GetDefaultAttributeValues(TypedSet->GetNamedSet()->GetName());
+				Base = ReferenceValues->GetBoundValueMaps().Find(Input->GetMappingKey());
+			}
+
+			if (Output == nullptr)
+			{
+				const FBoundValueMap::FConstructArgs Args(Input->GetTypedSet(), Input->GetValueType(), OutputAllocator);
+				Output = MakeBoundValueMap<FQuaternionAnimationAttribute>(Args);
+			}
+
+			if (Base == nullptr)
+			{
+				UE_LOG(LogAnimation, Warning, TEXT("FMakeAdditiveSpace_QuaternionAttribute::TransformBoundValueMap: Could not transform bound value map - no base provided and no reference found. Setting to identity instead."));
+				Output->FillWithIdentity(true);
+				return;
+			}
+
+			// Our output should have the same shape as our base
+			check(Base->GetTypedSet() == Output->GetTypedSet());
+			check(Base->Num() == Output->Num());
+
+			// Our base and input should have the same shape
+			check(Base->GetTypedSet() == Input->GetTypedSet());
+			check(Base->Num() == Input->Num());
+
+			const TBoundValueMap<FQuaternionAnimationAttribute>* BaseQuaternions = CastChecked<FQuaternionAnimationAttribute>(Base);
+			const TBoundValueMap<FQuaternionAnimationAttribute>* InputQuaternions = CastChecked<FQuaternionAnimationAttribute>(Input);
+			TBoundValueMap<FQuaternionAnimationAttribute>* OutputQuaternions = CastChecked<FQuaternionAnimationAttribute>(Output);
+
+			const FQuaternionAnimationAttribute* BaseQuaternionPtr = BaseQuaternions->GetData();
+			const FQuaternionAnimationAttribute* InputQuaternionPtr = InputQuaternions->GetData();
+			FQuaternionAnimationAttribute* OutputQuaternionPtr = OutputQuaternions->GetData();
+			const FQuaternionAnimationAttribute* BaseQuaternionEndPtr = BaseQuaternionPtr + BaseQuaternions->Num();
+
+			while (BaseQuaternionPtr < BaseQuaternionEndPtr)
+			{
+				OutputQuaternionPtr->Value = InputQuaternionPtr->Value * BaseQuaternionPtr->Value.Inverse();
+				OutputQuaternionPtr->Value.Normalize();
+
+				BaseQuaternionPtr++;
+				InputQuaternionPtr++;
+				OutputQuaternionPtr++;
+			}
+		}
+
+		void FMakeAdditiveSpace_QuaternionAttribute::TransformUnboundValueMap(const FUnboundValueMap* Base, const FUnboundValueMap* Input, FUnboundValueMap*& Output, FReallocFun OutputAllocator)
+		{
+			check(Output == nullptr);
+
+			const TUnboundValueMap<FQuaternionAnimationAttribute>* BaseQuaternions = Cast<FQuaternionAnimationAttribute>(Base);
+			const TUnboundValueMap<FQuaternionAnimationAttribute>* InputQuaternions = CastChecked<FQuaternionAnimationAttribute>(Input);
+
+			TUnboundValueMap<FQuaternionAnimationAttribute>* OutputQuaternions = MakeUnboundValueMap<FQuaternionAnimationAttribute>(OutputAllocator);
+
+			OuterJoinBy(
+				FUnboundValueMapJoinOp(),
+				[OutputQuaternions](FName Name, const FQuaternionAnimationAttribute* BaseValue, const FQuaternionAnimationAttribute* InputValue)
+				{
+					if (InputValue == nullptr)
+					{
+						// We have no input to subtract, skip it
+						return;
+					}
+
+					// Missing values are set to identity
+					const FQuat BaseQuaternion = BaseValue != nullptr ? BaseValue->Value : FQuat::Identity;
+
+					FQuaternionAnimationAttribute Output{ InputValue->Value * BaseQuaternion.Inverse() };
+					Output.Value.Normalize();
+						
+					OutputQuaternions->Append(Name, Output);
+				},
+				BaseQuaternions != nullptr ? BaseQuaternions->CreateConstIterator() : TUnboundValueMap<FQuaternionAnimationAttribute>::FConstIterator(),
+				InputQuaternions->CreateConstIterator());
+
+			Output = OutputQuaternions;
+		}
+		
 		//////////////////////////////////////////////////////////////////////////
 		// ApplyAdditiveSpace transformers
 
@@ -913,5 +1256,457 @@ namespace UE::UAF::Transformers
 
 			Output = OutputFloats;
 		}
+
+		void FApplyAdditiveSpace_IntegerAttribute::TransformBoundValueMap(const FBoundValueMap* Base, const FBoundValueMap* Additive, const FBoundValueMap* PerValueWeightMap, float DefaultWeight, FBoundValueMap*& Output, FReallocFun OutputAllocator)
+		{
+			if (Base == nullptr)
+			{
+				// Our base is missing, use the user authored default values
+				const FAttributeTypedSetPtr& TypedSet = Additive->GetTypedSet();
+				const FValueBundlePtr ReferenceValues = TypedSet->GetBindingData()->GetDefaultAttributeValues(TypedSet->GetNamedSet()->GetName());
+				Base = ReferenceValues->GetBoundValueMaps().Find(Additive->GetMappingKey());
+			}
+
+			if (Output == nullptr)
+			{
+				const FBoundValueMap::FConstructArgs Args(Additive->GetTypedSet(), Additive->GetValueType(), OutputAllocator);
+				Output = MakeBoundValueMap<FFloatAnimationAttribute>(Args);
+			}
+
+			if (Base == nullptr)
+			{
+				UE_LOG(LogAnimation, Warning, TEXT("FApplyAdditiveSpace_IntegerAttribute::TransformBoundValueMap: Could not transform bound value map - no base provided and no reference found. Setting to identity instead."));
+				Output->FillWithIdentity(false);
+				return;
+			}
+
+			// Our output should have the same shape as our base
+			check(Base->GetTypedSet() == Output->GetTypedSet());
+			check(Base->Num() == Output->Num());
+
+			// Our base and additive should have the same shape
+			check(Base->GetTypedSet() == Additive->GetTypedSet());
+			check(Base->Num() == Additive->Num());
+
+			const TBoundValueMap<FIntegerAnimationAttribute>* BaseIntegers = CastChecked<FIntegerAnimationAttribute>(Base);
+			const TBoundValueMap<FIntegerAnimationAttribute>* AdditiveIntegers = CastChecked<FIntegerAnimationAttribute>(Additive);
+			TBoundValueMap<FIntegerAnimationAttribute>* OutputIntegers = CastChecked<FIntegerAnimationAttribute>(Output);
+
+			const FIntegerAnimationAttribute* BaseIntegerPtr = BaseIntegers->GetData();
+			const FIntegerAnimationAttribute* AdditiveIntegerPtr = AdditiveIntegers->GetData();
+			FIntegerAnimationAttribute* OutputIntegerPtr = OutputIntegers->GetData();
+			const FIntegerAnimationAttribute* BaseIntegerPtrEnd = BaseIntegerPtr + BaseIntegers->Num();
+
+			if (const TBoundValueMap<FFloatAnimationAttribute>* PerValueWeights = Cast<FFloatAnimationAttribute>(PerValueWeightMap))
+			{
+				check(Base->GetTypedSet() == PerValueWeights->GetTypedSet());
+				check(Base->Num() == PerValueWeights->Num());
+
+				const FFloatAnimationAttribute* PerValueWeightsPtr = PerValueWeights->GetData();
+
+				while (BaseIntegerPtr < BaseIntegerPtrEnd)
+				{
+					OutputIntegerPtr->Value = BaseIntegerPtr->Value + static_cast<int32>(PerValueWeightsPtr->Value * AdditiveIntegerPtr->Value);
+
+					BaseIntegerPtr++;
+					AdditiveIntegerPtr++;
+					PerValueWeightsPtr++;
+					OutputIntegerPtr++;
+				}
+			}
+			else
+			{
+				while (BaseIntegerPtr < BaseIntegerPtrEnd)
+				{
+					OutputIntegerPtr->Value = BaseIntegerPtr->Value + static_cast<int32>(DefaultWeight * AdditiveIntegerPtr->Value);
+
+					BaseIntegerPtr++;
+					AdditiveIntegerPtr++;
+					OutputIntegerPtr++;
+				}
+			}
+		}
+
+		void FApplyAdditiveSpace_IntegerAttribute::TransformUnboundValueMap(const FUnboundValueMap* Base, const FUnboundValueMap* Additive, float AdditiveWeight, FUnboundValueMap*& Output, FReallocFun OutputAllocator)
+		{
+			check(Output == nullptr);
+
+			const TUnboundValueMap<FIntegerAnimationAttribute>* BaseIntegers = Cast<FIntegerAnimationAttribute>(Base);
+			const TUnboundValueMap<FIntegerAnimationAttribute>* AdditiveIntegers = CastChecked<FIntegerAnimationAttribute>(Additive);
+
+			TUnboundValueMap<FIntegerAnimationAttribute>* OutputIntegers = MakeUnboundValueMap<FIntegerAnimationAttribute>(OutputAllocator);
+
+			OuterJoinBy(
+				FUnboundValueMapJoinOp(),
+				[AdditiveWeight, OutputIntegers](FName Name, const FIntegerAnimationAttribute* BaseValue, const FIntegerAnimationAttribute* AdditiveValue)
+				{
+					// Missing values are set to 0
+					const int32 BaseInteger = BaseValue != nullptr ? BaseValue->Value : 0;
+					const int32 AdditiveInteger = AdditiveValue != nullptr ? AdditiveValue->Value : 0;
+
+					const FIntegerAnimationAttribute Output{ BaseInteger + static_cast<int32>(AdditiveWeight * AdditiveInteger) };
+					OutputIntegers->Append(Name, Output);
+				},
+				BaseIntegers != nullptr ? BaseIntegers->CreateConstIterator() : TUnboundValueMap<FIntegerAnimationAttribute>::FConstIterator(),
+				AdditiveIntegers->CreateConstIterator());
+
+			Output = OutputIntegers;
+		}
+
+		void FApplyAdditiveSpace_TransformAttribute::TransformBoundValueMap(const FBoundValueMap* Base, const FBoundValueMap* Additive, const FBoundValueMap* PerValueWeightMap, float DefaultWeight, FBoundValueMap*& Output, FReallocFun OutputAllocator)
+		{
+			if (Base == nullptr)
+			{
+				// Our base is missing, use the reference pose
+				const FAttributeTypedSetPtr& TypedSet = Additive->GetTypedSet();
+				const FValueBundlePtr ReferenceValues = TypedSet->GetBindingData()->GetDefaultAttributeValues(TypedSet->GetNamedSet()->GetName());
+				Base = ReferenceValues->GetBoundValueMaps().Find(Additive->GetMappingKey());
+			}
+
+			if (Output == nullptr)
+			{
+				const FBoundValueMap::FConstructArgs Args(Additive->GetTypedSet(), Additive->GetValueType(), OutputAllocator);
+				Output = MakeBoundValueMap<FTransformAnimationAttribute>(Args);
+			}
+
+			if (Base == nullptr)
+			{
+				UE_LOG(LogAnimation, Warning, TEXT("FApplyAdditiveSpace_TransformAttribute::TransformBoundValueMap: Could not transform bound value map - no base provided and no reference found. Setting to identity instead."));
+				Output->FillWithIdentity(false);
+				return;
+			}
+
+			// Our output should have the same shape as our base
+			check(Base->GetTypedSet() == Output->GetTypedSet());
+			check(Base->Num() == Output->Num());
+
+			// Our base and additive should have the same shape
+			check(Base->GetTypedSet() == Additive->GetTypedSet());
+			check(Base->Num() == Additive->Num());
+
+			const TBoundValueMap<FTransformAnimationAttribute>* BaseTransforms = CastChecked<FTransformAnimationAttribute>(Base);
+			const TBoundValueMap<FTransformAnimationAttribute>* AdditiveTransforms = CastChecked<FTransformAnimationAttribute>(Additive);
+			TBoundValueMap<FTransformAnimationAttribute>* OutputTransforms = CastChecked<FTransformAnimationAttribute>(Output);
+
+			const FTransformAnimationAttribute* BaseTransformPtr = BaseTransforms->GetData();
+			const FTransformAnimationAttribute* AdditiveTransformPtr = AdditiveTransforms->GetData();
+			FTransformAnimationAttribute* OutputTransformPtr = OutputTransforms->GetData();
+			const FTransformAnimationAttribute* BaseTransformPtrEnd = BaseTransformPtr + BaseTransforms->Num();
+
+			if (const TBoundValueMap<FFloatAnimationAttribute>* PerValueWeights = Cast<FFloatAnimationAttribute>(PerValueWeightMap))
+			{
+				check(Base->GetTypedSet() == PerValueWeights->GetTypedSet());
+				check(Base->Num() == PerValueWeights->Num());
+
+				const FFloatAnimationAttribute* PerValueWeightsPtr = PerValueWeights->GetData();
+
+				while (BaseTransformPtr < BaseTransformPtrEnd)
+				{
+					FTransform Result = BaseTransformPtr->Value;
+					FTransform::BlendFromIdentityAndAccumulate(Result, AdditiveTransformPtr->Value, ScalarRegister(PerValueWeightsPtr->Value));
+
+					OutputTransformPtr->Value = Result;
+
+					BaseTransformPtr++;
+					AdditiveTransformPtr++;
+					PerValueWeightsPtr++;
+					OutputTransformPtr++;
+				}
+			}
+			else
+			{
+				const ScalarRegister VDefaultWeight(DefaultWeight);
+
+				if (FAnimWeight::IsFullWeight(DefaultWeight))
+				{
+					// Our additive has full weight, no need to interpolate with identity
+					while (BaseTransformPtr < BaseTransformPtrEnd)
+					{
+						// TODO: The VDefaultWeight is practically 1.0, can assume that it is?
+						FTransform Result = BaseTransformPtr->Value;
+						Result.AccumulateWithAdditiveScale(AdditiveTransformPtr->Value, VDefaultWeight);
+
+						OutputTransformPtr->Value = Result;
+
+						BaseTransformPtr++;
+						AdditiveTransformPtr++;
+						OutputTransformPtr++;
+					}
+				}
+				else
+				{
+					// Interpolate with identity before adding
+					while (BaseTransformPtr < BaseTransformPtrEnd)
+					{
+						FTransform Result = BaseTransformPtr->Value;
+						FTransform::BlendFromIdentityAndAccumulate(Result, AdditiveTransformPtr->Value, VDefaultWeight);
+
+						OutputTransformPtr->Value = Result;
+
+						BaseTransformPtr++;
+						AdditiveTransformPtr++;
+						OutputTransformPtr++;
+					}
+				}
+			}
+		}
+
+		void FApplyAdditiveSpace_TransformAttribute::TransformUnboundValueMap(const FUnboundValueMap* Base, const FUnboundValueMap* Additive, float AdditiveWeight, FUnboundValueMap*& Output, FReallocFun OutputAllocator)
+		{
+			check(Output == nullptr);
+
+			const TUnboundValueMap<FTransformAnimationAttribute>* BaseTransforms = Cast<FTransformAnimationAttribute>(Base);
+			const TUnboundValueMap<FTransformAnimationAttribute>* AdditiveTransforms = CastChecked<FTransformAnimationAttribute>(Additive);
+
+			TUnboundValueMap<FTransformAnimationAttribute>* OutputTransforms = MakeUnboundValueMap<FTransformAnimationAttribute>(OutputAllocator);
+
+			const ScalarRegister VAdditiveWeight(AdditiveWeight);
+
+			if (FAnimWeight::IsFullWeight(AdditiveWeight))
+			{
+				// Our additive has full weight, no need to interpolate with identity
+				OuterJoinBy(
+					FUnboundValueMapJoinOp(),
+					[VAdditiveWeight, OutputTransforms](FName Name, const FTransformAnimationAttribute* BaseValue, const FTransformAnimationAttribute* AdditiveValue)
+					{
+						// Missing values are set to the identity
+						const FTransform& BaseTransform = BaseValue != nullptr ? BaseValue->Value : FTransform::Identity;
+
+						FTransformAnimationAttribute Output{ BaseTransform };
+
+						if (AdditiveValue != nullptr)
+						{
+							// TODO: The VAdditiveWeight is practically 1.0, can assume that it is?
+							Output.Value.AccumulateWithAdditiveScale(AdditiveValue->Value, VAdditiveWeight);
+						}
+
+						OutputTransforms->Append(Name, Output);
+					},
+					BaseTransforms != nullptr ? BaseTransforms->CreateConstIterator() : TUnboundValueMap<FTransformAnimationAttribute>::FConstIterator(),
+					AdditiveTransforms->CreateConstIterator());
+			}
+			else
+			{
+				// Interpolate with identity before adding
+				OuterJoinBy(
+					FUnboundValueMapJoinOp(),
+					[VAdditiveWeight, OutputTransforms](FName Name, const FTransformAnimationAttribute* BaseValue, const FTransformAnimationAttribute* AdditiveValue)
+					{
+						// Missing values are set to the identity
+						const FTransform& BaseTransform = BaseValue != nullptr ? BaseValue->Value : FTransform::Identity;
+
+						FTransformAnimationAttribute Output{ BaseTransform };
+
+						if (AdditiveValue != nullptr)
+						{
+							FTransform::BlendFromIdentityAndAccumulate(Output.Value, AdditiveValue->Value, VAdditiveWeight);
+						}
+
+						OutputTransforms->Append(Name, Output);
+					},
+					BaseTransforms != nullptr ? BaseTransforms->CreateConstIterator() : TUnboundValueMap<FTransformAnimationAttribute>::FConstIterator(),
+					AdditiveTransforms->CreateConstIterator());
+			}
+
+			Output = OutputTransforms;
+		}
+
+		void FApplyAdditiveSpace_VectorAttribute::TransformBoundValueMap(const FBoundValueMap* Base, const FBoundValueMap* Additive, const FBoundValueMap* PerValueWeightMap, float DefaultWeight, FBoundValueMap*& Output, FReallocFun OutputAllocator)
+		{
+			if (Base == nullptr)
+			{
+				// Our base is missing, use the user authored default values
+				const FAttributeTypedSetPtr& TypedSet = Additive->GetTypedSet();
+				const FValueBundlePtr ReferenceValues = TypedSet->GetBindingData()->GetDefaultAttributeValues(TypedSet->GetNamedSet()->GetName());
+				Base = ReferenceValues->GetBoundValueMaps().Find(Additive->GetMappingKey());
+			}
+
+			if (Output == nullptr)
+			{
+				const FBoundValueMap::FConstructArgs Args(Additive->GetTypedSet(), Additive->GetValueType(), OutputAllocator);
+				Output = MakeBoundValueMap<FVectorAnimationAttribute>(Args);
+			}
+
+			if (Base == nullptr)
+			{
+				UE_LOG(LogAnimation, Warning, TEXT("FApplyAdditiveSpace_VectorAttribute::TransformBoundValueMap: Could not transform bound value map - no base provided and no reference found. Setting to identity instead."));
+				Output->FillWithIdentity(false);
+				return;
+			}
+
+			// Our output should have the same shape as our base
+			check(Base->GetTypedSet() == Output->GetTypedSet());
+			check(Base->Num() == Output->Num());
+
+			// Our base and additive should have the same shape
+			check(Base->GetTypedSet() == Additive->GetTypedSet());
+			check(Base->Num() == Additive->Num());
+
+			const TBoundValueMap<FVectorAnimationAttribute>* BaseVectors = CastChecked<FVectorAnimationAttribute>(Base);
+			const TBoundValueMap<FVectorAnimationAttribute>* AdditiveVectors = CastChecked<FVectorAnimationAttribute>(Additive);
+			TBoundValueMap<FVectorAnimationAttribute>* OutputVectors = CastChecked<FVectorAnimationAttribute>(Output);
+
+			const FVectorAnimationAttribute* BaseVectorPtr = BaseVectors->GetData();
+			const FVectorAnimationAttribute* AdditiveVectorPtr = AdditiveVectors->GetData();
+			FVectorAnimationAttribute* OutputVectorPtr = OutputVectors->GetData();
+			const FVectorAnimationAttribute* BaseVectorPtrEnd = BaseVectorPtr + BaseVectors->Num();
+
+			if (const TBoundValueMap<FFloatAnimationAttribute>* PerValueWeights = Cast<FFloatAnimationAttribute>(PerValueWeightMap))
+			{
+				check(Base->GetTypedSet() == PerValueWeights->GetTypedSet());
+				check(Base->Num() == PerValueWeights->Num());
+
+				const FFloatAnimationAttribute* PerValueWeightsPtr = PerValueWeights->GetData();
+
+				while (BaseVectorPtr < BaseVectorPtrEnd)
+				{
+					OutputVectorPtr->Value = BaseVectorPtr->Value + (PerValueWeightsPtr->Value * AdditiveVectorPtr->Value);
+
+					BaseVectorPtr++;
+					AdditiveVectorPtr++;
+					PerValueWeightsPtr++;
+					OutputVectorPtr++;
+				}
+			}
+			else
+			{
+				while (BaseVectorPtr < BaseVectorPtrEnd)
+				{
+					OutputVectorPtr->Value = BaseVectorPtr->Value + (DefaultWeight * AdditiveVectorPtr->Value);
+
+					BaseVectorPtr++;
+					AdditiveVectorPtr++;
+					OutputVectorPtr++;
+				}
+			}
+		}
+
+		void FApplyAdditiveSpace_VectorAttribute::TransformUnboundValueMap(const FUnboundValueMap* Base, const FUnboundValueMap* Additive, float AdditiveWeight, FUnboundValueMap*& Output, FReallocFun OutputAllocator)
+		{
+			check(Output == nullptr);
+
+			const TUnboundValueMap<FVectorAnimationAttribute>* BaseVectors = Cast<FVectorAnimationAttribute>(Base);
+			const TUnboundValueMap<FVectorAnimationAttribute>* AdditiveVectors = CastChecked<FVectorAnimationAttribute>(Additive);
+
+			TUnboundValueMap<FVectorAnimationAttribute>* OutputVectors = MakeUnboundValueMap<FVectorAnimationAttribute>(OutputAllocator);
+
+			OuterJoinBy(
+				FUnboundValueMapJoinOp(),
+				[AdditiveWeight, OutputVectors](FName Name, const FVectorAnimationAttribute* BaseValue, const FVectorAnimationAttribute* AdditiveValue)
+				{
+					// Missing values are set to identity
+					const FVector BaseVector = BaseValue != nullptr ? BaseValue->Value : FVector::ZeroVector;
+					const FVector AdditiveVector = AdditiveValue != nullptr ? AdditiveValue->Value : FVector::ZeroVector;
+
+					const FVectorAnimationAttribute Output{ BaseVector + (AdditiveWeight * AdditiveVector) };
+					OutputVectors->Append(Name, Output);
+				},
+				BaseVectors != nullptr ? BaseVectors->CreateConstIterator() : TUnboundValueMap<FVectorAnimationAttribute>::FConstIterator(),
+				AdditiveVectors->CreateConstIterator());
+
+			Output = OutputVectors;
+		}
+
+		void FApplyAdditiveSpace_QuaternionAttribute::TransformBoundValueMap(const FBoundValueMap* Base, const FBoundValueMap* Additive, const FBoundValueMap* PerValueWeightMap, float DefaultWeight, FBoundValueMap*& Output, FReallocFun OutputAllocator)
+		{
+			if (Base == nullptr)
+			{
+				// Our base is missing, use the user authored default values
+				const FAttributeTypedSetPtr& TypedSet = Additive->GetTypedSet();
+				const FValueBundlePtr ReferenceValues = TypedSet->GetBindingData()->GetDefaultAttributeValues(TypedSet->GetNamedSet()->GetName());
+				Base = ReferenceValues->GetBoundValueMaps().Find(Additive->GetMappingKey());
+			}
+
+			if (Output == nullptr)
+			{
+				const FBoundValueMap::FConstructArgs Args(Additive->GetTypedSet(), Additive->GetValueType(), OutputAllocator);
+				Output = MakeBoundValueMap<FQuaternionAnimationAttribute>(Args);
+			}
+
+			if (Base == nullptr)
+			{
+				UE_LOG(LogAnimation, Warning, TEXT("FApplyAdditiveSpace_VectorAttribute::QuaternionBoundValueMap: Could not transform bound value map - no base provided and no reference found. Setting to identity instead."));
+				Output->FillWithIdentity(false);
+				return;
+			}
+
+			// Our output should have the same shape as our base
+			check(Base->GetTypedSet() == Output->GetTypedSet());
+			check(Base->Num() == Output->Num());
+
+			// Our base and additive should have the same shape
+			check(Base->GetTypedSet() == Additive->GetTypedSet());
+			check(Base->Num() == Additive->Num());
+
+			const TBoundValueMap<FQuaternionAnimationAttribute>* BaseQuaternions = CastChecked<FQuaternionAnimationAttribute>(Base);
+			const TBoundValueMap<FQuaternionAnimationAttribute>* AdditiveQuaternions = CastChecked<FQuaternionAnimationAttribute>(Additive);
+			TBoundValueMap<FQuaternionAnimationAttribute>* OutputQuaternions = CastChecked<FQuaternionAnimationAttribute>(Output);
+
+			const FQuaternionAnimationAttribute* BaseQuaternionPtr = BaseQuaternions->GetData();
+			const FQuaternionAnimationAttribute* AdditiveQuaternionPtr = AdditiveQuaternions->GetData();
+			FQuaternionAnimationAttribute* OutputQuaternionPtr = OutputQuaternions->GetData();
+			const FQuaternionAnimationAttribute* BaseQuaternionPtrEnd = BaseQuaternionPtr + BaseQuaternions->Num();
+
+			if (const TBoundValueMap<FFloatAnimationAttribute>* PerValueWeights = Cast<FFloatAnimationAttribute>(PerValueWeightMap))
+			{
+				check(Base->GetTypedSet() == PerValueWeights->GetTypedSet());
+				check(Base->Num() == PerValueWeights->Num());
+
+				const FFloatAnimationAttribute* PerValueWeightsPtr = PerValueWeights->GetData();
+
+				while (BaseQuaternionPtr < BaseQuaternionPtrEnd)
+				{
+					const FQuat WeightedRotation = FQuat::FastLerp(FQuat::Identity, AdditiveQuaternionPtr->Value, PerValueWeightsPtr->Value).GetNormalized();
+					OutputQuaternionPtr->Value = BaseQuaternionPtr->Value + WeightedRotation;
+					
+					BaseQuaternionPtr++;
+					AdditiveQuaternionPtr++;
+					PerValueWeightsPtr++;
+					OutputQuaternionPtr++;
+				}
+			}
+			else
+			{
+				while (BaseQuaternionPtr < BaseQuaternionPtrEnd)
+				{
+					const FQuat WeightedRotation = FQuat::FastLerp(FQuat::Identity, AdditiveQuaternionPtr->Value, DefaultWeight).GetNormalized();
+					OutputQuaternionPtr->Value = BaseQuaternionPtr->Value + WeightedRotation;
+
+					BaseQuaternionPtr++;
+					AdditiveQuaternionPtr++;
+					OutputQuaternionPtr++;
+				}
+			}
+		}
+
+		void FApplyAdditiveSpace_QuaternionAttribute::TransformUnboundValueMap(const FUnboundValueMap* Base, const FUnboundValueMap* Additive, float AdditiveWeight, FUnboundValueMap*& Output, FReallocFun OutputAllocator)
+		{
+			check(Output == nullptr);
+
+			const TUnboundValueMap<FQuaternionAnimationAttribute>* BaseQuaternions = Cast<FQuaternionAnimationAttribute>(Base);
+			const TUnboundValueMap<FQuaternionAnimationAttribute>* AdditiveQuaternions = CastChecked<FQuaternionAnimationAttribute>(Additive);
+
+			TUnboundValueMap<FQuaternionAnimationAttribute>* OutputQuaternions = MakeUnboundValueMap<FQuaternionAnimationAttribute>(OutputAllocator);
+
+			OuterJoinBy(
+				FUnboundValueMapJoinOp(),
+				[AdditiveWeight, OutputQuaternions](FName Name, const FQuaternionAnimationAttribute* BaseValue, const FQuaternionAnimationAttribute* AdditiveValue)
+				{
+					// Missing values are set to identity
+					const FQuat BaseQuaternion = BaseValue != nullptr ? BaseValue->Value : FQuat::Identity;
+					const FQuat AdditiveQuaternion = AdditiveValue != nullptr ? AdditiveValue->Value : FQuat::Identity;
+
+					const FQuat WeightedRotation = FQuat::FastLerp(FQuat::Identity, AdditiveQuaternion, AdditiveWeight).GetNormalized();
+
+					FQuaternionAnimationAttribute Output;
+					Output.Value = BaseQuaternion + WeightedRotation;
+
+					OutputQuaternions->Append(Name, Output);
+				},
+				BaseQuaternions != nullptr ? BaseQuaternions->CreateConstIterator() : TUnboundValueMap<FQuaternionAnimationAttribute>::FConstIterator(),
+				AdditiveQuaternions->CreateConstIterator());
+
+			Output = OutputQuaternions;
+		}
 	}
 }
diff --git a/Engine/Plugins/Experimental/UAF/UAF/Source/UAF/Private/UAF/ValueRuntime/Transformers/Interpolate.cpp b/Engine/Plugins/Experimental/UAF/UAF/Source/UAF/Private/UAF/ValueRuntime/Transformers/Interpolate.cpp
--- a/Engine/Plugins/Experimental/UAF/UAF/Source/UAF/Private/UAF/ValueRuntime/Transformers/Interpolate.cpp
+++ b/Engine/Plugins/Experimental/UAF/UAF/Source/UAF/Private/UAF/ValueRuntime/Transformers/Interpolate.cpp
@@ -584,5 +584,332 @@ namespace UE::UAF::Transformers
 
 			OutputMap = OutputFloats;
 		}
+
+		void FInterpolate_IntegerAttribute::TransformBoundValueMap(const FBoundValueMap* MapA, const FBoundValueMap* MapB, const FBoundValueMap* PerValueWeightMapB, float DefaultWeightB, FBoundValueMap*& OutputMap, FReallocFun OutputAllocator)
+		{
+			check(MapA->GetTypedSet() == MapB->GetTypedSet());
+			check(MapA->Num() == MapB->Num());
+
+			if (OutputMap == nullptr)
+			{
+				const FBoundValueMap::FConstructArgs Args(MapA->GetTypedSet(), MapA->GetValueType(), OutputAllocator);
+				OutputMap = MakeBoundValueMap<FIntegerAnimationAttribute>(Args);
+			}
+
+			check(MapA->GetTypedSet() == OutputMap->GetTypedSet());
+			check(MapA->Num() == OutputMap->Num());
+
+			const TBoundValueMap<FIntegerAnimationAttribute>* IntegersA = CastChecked<FIntegerAnimationAttribute>(MapA);
+			const TBoundValueMap<FIntegerAnimationAttribute>* IntegersB = CastChecked<FIntegerAnimationAttribute>(MapB);
+			TBoundValueMap<FIntegerAnimationAttribute>* OutputIntegers = CastChecked<FIntegerAnimationAttribute>(OutputMap);
+
+			const FIntegerAnimationAttribute* IntegersPtrA = IntegersA->GetData();
+			const FIntegerAnimationAttribute* IntegersPtrB = IntegersB->GetData();
+			FIntegerAnimationAttribute* OutputIntegersPtr = OutputIntegers->GetData();
+			const FIntegerAnimationAttribute* IntegersPtrAEnd = IntegersPtrA + IntegersA->Num();
+
+			if (const TBoundValueMap<FFloatAnimationAttribute>* PerValueWeightsB = Cast<FFloatAnimationAttribute>(PerValueWeightMapB))
+			{
+				check(MapA->GetTypedSet() == PerValueWeightsB->GetTypedSet());
+				check(MapA->Num() == PerValueWeightsB->Num());
+
+				const FFloatAnimationAttribute* PerValueWeightsBPtr = PerValueWeightsB->GetData();
+
+				while (IntegersPtrA < IntegersPtrAEnd)
+				{
+					OutputIntegersPtr->Value = FMath::TruncToInt32(IntegersPtrA->Value * (1.0 - PerValueWeightsBPtr->Value) + (IntegersPtrB->Value * PerValueWeightsBPtr->Value));
+
+					IntegersPtrA++;
+					IntegersPtrB++;
+					PerValueWeightsBPtr++;
+					OutputIntegersPtr++;
+				}
+			}
+			else
+			{
+				// Use doubles to reduce precision loss during weighted sum
+				const double WeightBDouble = DefaultWeightB;
+				const double WeightADouble = 1.0 - DefaultWeightB;
+				
+				while (IntegersPtrA < IntegersPtrAEnd)
+				{
+					OutputIntegersPtr->Value = FMath::TruncToInt32(IntegersPtrA->Value * WeightADouble + IntegersPtrB->Value * WeightBDouble);
+					
+					IntegersPtrA++;
+					IntegersPtrB++;
+					OutputIntegersPtr++;
+				}
+			}
+		}
+
+		void FInterpolate_IntegerAttribute::TransformUnboundValueMap(const FUnboundValueMap* MapA, const FUnboundValueMap* MapB, float WeightB, FUnboundValueMap*& OutputMap, FReallocFun OutputAllocator)
+		{
+			check(OutputMap == nullptr);
+
+			const TUnboundValueMap<FIntegerAnimationAttribute>* IntegersA = CastChecked<FIntegerAnimationAttribute>(MapA);
+			const TUnboundValueMap<FIntegerAnimationAttribute>* IntegersB = CastChecked<FIntegerAnimationAttribute>(MapB);
+			TUnboundValueMap<FIntegerAnimationAttribute>* OutputIntegers = MakeUnboundValueMap<FIntegerAnimationAttribute>(OutputAllocator);
+
+			const double WeightBDouble = WeightB;
+			const double WeightADouble = 1.0 - WeightB;
+			
+			OuterJoinBy(
+				FUnboundValueMapJoinOp(),
+				[WeightADouble, WeightBDouble, OutputIntegers](FName Name, const FIntegerAnimationAttribute* A, const FIntegerAnimationAttribute* B)
+				{
+					// Missing values are set to 0
+					const int32 ValueA = A != nullptr ? A->Value : 0;
+					const int32 ValueB = B != nullptr ? B->Value : 0;
+
+					FIntegerAnimationAttribute Output;
+					Output.Value = FMath::TruncToInt32(ValueA * WeightADouble + ValueB * WeightBDouble);
+
+					OutputIntegers->Append(Name, Output);
+				},
+				IntegersA->CreateConstIterator(),
+				IntegersB->CreateConstIterator());
+
+			OutputMap = OutputIntegers;
+		}
+
+		void FInterpolate_TransformAttribute::TransformBoundValueMap(const FBoundValueMap* MapA, const FBoundValueMap* MapB, const FBoundValueMap* PerValueWeightMapB, float DefaultWeightB, FBoundValueMap*& OutputMap, FReallocFun OutputAllocator)
+		{
+			check(MapA->GetTypedSet() == MapB->GetTypedSet());
+			check(MapA->Num() == MapB->Num());
+
+			if (OutputMap == nullptr)
+			{
+				const FBoundValueMap::FConstructArgs Args(MapA->GetTypedSet(), MapA->GetValueType(), OutputAllocator);
+				OutputMap = MakeBoundValueMap<FTransformAnimationAttribute>(Args);
+			}
+
+			check(MapA->GetTypedSet() == OutputMap->GetTypedSet());
+			check(MapA->Num() == OutputMap->Num());
+
+			const TBoundValueMap<FTransformAnimationAttribute>* TransformsA = CastChecked<FTransformAnimationAttribute>(MapA);
+			const TBoundValueMap<FTransformAnimationAttribute>* TransformsB = CastChecked<FTransformAnimationAttribute>(MapB);
+			TBoundValueMap<FTransformAnimationAttribute>* OutputTransforms = CastChecked<FTransformAnimationAttribute>(OutputMap);
+
+			const FTransformAnimationAttribute* TransformPtrA = TransformsA->GetData();
+			const FTransformAnimationAttribute* TransformPtrB = TransformsB->GetData();
+			FTransformAnimationAttribute* OutputTransformPtr = OutputTransforms->GetData();
+			const FTransformAnimationAttribute* TransformPtrAEnd = TransformPtrA + TransformsA->Num();
+
+			if (const TBoundValueMap<FFloatAnimationAttribute>* PerValueWeightsB = Cast<FFloatAnimationAttribute>(PerValueWeightMapB))
+			{
+				check(MapA->GetTypedSet() == PerValueWeightsB->GetTypedSet());
+				check(MapA->Num() == PerValueWeightsB->Num());
+
+				const FFloatAnimationAttribute* PerValueWeightsBPtr = PerValueWeightsB->GetData();
+
+				while (TransformPtrA < TransformPtrAEnd)
+				{
+					OutputTransformPtr->Value.Blend(TransformPtrA->Value, TransformPtrB->Value, PerValueWeightsBPtr->Value);
+
+					TransformPtrA++;
+					TransformPtrB++;
+					PerValueWeightsBPtr++;
+					OutputTransformPtr++;
+				}
+			}
+			else
+			{
+				while (TransformPtrA < TransformPtrAEnd)
+				{
+					OutputTransformPtr->Value.Blend(TransformPtrA->Value, TransformPtrB->Value, DefaultWeightB);
+
+					TransformPtrA++;
+					TransformPtrB++;
+					OutputTransformPtr++;
+				}
+			}
+		}
+
+		void FInterpolate_TransformAttribute::TransformUnboundValueMap(const FUnboundValueMap* MapA, const FUnboundValueMap* MapB, float WeightB, FUnboundValueMap*& OutputMap, FReallocFun OutputAllocator)
+		{
+			check(OutputMap == nullptr);
+
+			const TUnboundValueMap<FTransformAnimationAttribute>* TransformsA = CastChecked<FTransformAnimationAttribute>(MapA);
+			const TUnboundValueMap<FTransformAnimationAttribute>* TransformsB = CastChecked<FTransformAnimationAttribute>(MapB);
+			TUnboundValueMap<FTransformAnimationAttribute>* OutputTransforms = MakeUnboundValueMap<FTransformAnimationAttribute>(OutputAllocator);
+
+			OuterJoinBy(
+				FUnboundValueMapJoinOp(),
+				[WeightB, OutputTransforms](FName Name, const FTransformAnimationAttribute* A, const FTransformAnimationAttribute* B)
+				{
+					// Missing values are set to the identity
+					const FTransform& ValueA = A != nullptr ? A->Value : FTransform::Identity;
+					const FTransform& ValueB = B != nullptr ? B->Value : FTransform::Identity;
+
+					FTransformAnimationAttribute Output;
+					Output.Value.Blend(ValueA, ValueB, WeightB);
+
+					OutputTransforms->Append(Name, Output);
+				},
+				TransformsA->CreateConstIterator(),
+				TransformsB->CreateConstIterator());
+
+			OutputMap = OutputTransforms;
+		}
+
+		void FInterpolate_VectorAttribute::TransformBoundValueMap(const FBoundValueMap* MapA, const FBoundValueMap* MapB, const FBoundValueMap* PerValueWeightMapB, float DefaultWeightB, FBoundValueMap*& OutputMap, FReallocFun OutputAllocator)
+		{
+			check(MapA->GetTypedSet() == MapB->GetTypedSet());
+			check(MapA->Num() == MapB->Num());
+
+			if (OutputMap == nullptr)
+			{
+				const FBoundValueMap::FConstructArgs Args(MapA->GetTypedSet(), MapA->GetValueType(), OutputAllocator);
+				OutputMap = MakeBoundValueMap<FVectorAnimationAttribute>(Args);
+			}
+
+			check(MapA->GetTypedSet() == OutputMap->GetTypedSet());
+			check(MapA->Num() == OutputMap->Num());
+
+			const TBoundValueMap<FVectorAnimationAttribute>* VectorsA = CastChecked<FVectorAnimationAttribute>(MapA);
+			const TBoundValueMap<FVectorAnimationAttribute>* VectorsB = CastChecked<FVectorAnimationAttribute>(MapB);
+			TBoundValueMap<FVectorAnimationAttribute>* OutputVectors = CastChecked<FVectorAnimationAttribute>(OutputMap);
+
+			const FVectorAnimationAttribute* VectorPtrA = VectorsA->GetData();
+			const FVectorAnimationAttribute* VectorPtrB = VectorsB->GetData();
+			FVectorAnimationAttribute* OutputVectorPtr = OutputVectors->GetData();
+			const FVectorAnimationAttribute* VectorPtrAEnd = VectorPtrA + VectorsA->Num();
+
+			if (const TBoundValueMap<FFloatAnimationAttribute>* PerValueWeightsB = Cast<FFloatAnimationAttribute>(PerValueWeightMapB))
+			{
+				check(MapA->GetTypedSet() == PerValueWeightsB->GetTypedSet());
+				check(MapA->Num() == PerValueWeightsB->Num());
+
+				const FFloatAnimationAttribute* PerValueWeightsBPtr = PerValueWeightsB->GetData();
+
+				while (VectorPtrA < VectorPtrAEnd)
+				{
+					OutputVectorPtr->Value = VectorPtrA->Value * (1.0f - PerValueWeightsBPtr->Value) + VectorPtrB->Value * PerValueWeightsBPtr->Value;
+
+					VectorPtrA++;
+					VectorPtrB++;
+					PerValueWeightsBPtr++;
+					OutputVectorPtr++;
+				}
+			}
+			else
+			{
+				while (VectorPtrA < VectorPtrAEnd)
+				{
+					OutputVectorPtr->Value = VectorPtrA->Value * (1.0f - DefaultWeightB) + VectorPtrB->Value * DefaultWeightB;
+
+					VectorPtrA++;
+					VectorPtrB++;
+					OutputVectorPtr++;
+				}
+			}
+		}
+
+		void FInterpolate_VectorAttribute::TransformUnboundValueMap(const FUnboundValueMap* MapA, const FUnboundValueMap* MapB, float WeightB, FUnboundValueMap*& OutputMap, FReallocFun OutputAllocator)
+		{
+			check(OutputMap == nullptr);
+
+			const TUnboundValueMap<FVectorAnimationAttribute>* VectorsA = CastChecked<FVectorAnimationAttribute>(MapA);
+			const TUnboundValueMap<FVectorAnimationAttribute>* VectorsB = CastChecked<FVectorAnimationAttribute>(MapB);
+			TUnboundValueMap<FVectorAnimationAttribute>* OutputVectors = MakeUnboundValueMap<FVectorAnimationAttribute>(OutputAllocator);
+
+			OuterJoinBy(
+				FUnboundValueMapJoinOp(),
+				[WeightB, OutputVectors](FName Name, const FVectorAnimationAttribute* A, const FVectorAnimationAttribute* B)
+				{
+					// Missing values are set to the identity
+					const FVector& ValueA = A != nullptr ? A->Value : FVector::ZeroVector;
+					const FVector& ValueB = B != nullptr ? B->Value : FVector::ZeroVector;
+
+					FVectorAnimationAttribute Output;
+					Output.Value = ValueA * (1.0f - WeightB) + ValueB * WeightB;
+
+					OutputVectors->Append(Name, Output);
+				},
+				VectorsA->CreateConstIterator(),
+				VectorsB->CreateConstIterator());
+
+			OutputMap = OutputVectors;
+		}
+
+		void FInterpolate_QuaternionAttribute::TransformBoundValueMap(const FBoundValueMap* MapA, const FBoundValueMap* MapB, const FBoundValueMap* PerValueWeightMapB, float DefaultWeightB, FBoundValueMap*& OutputMap, FReallocFun OutputAllocator)
+		{
+			check(MapA->GetTypedSet() == MapB->GetTypedSet());
+			check(MapA->Num() == MapB->Num());
+
+			if (OutputMap == nullptr)
+			{
+				const FBoundValueMap::FConstructArgs Args(MapA->GetTypedSet(), MapA->GetValueType(), OutputAllocator);
+				OutputMap = MakeBoundValueMap<FQuaternionAnimationAttribute>(Args);
+			}
+
+			check(MapA->GetTypedSet() == OutputMap->GetTypedSet());
+			check(MapA->Num() == OutputMap->Num());
+
+			const TBoundValueMap<FQuaternionAnimationAttribute>* QuaternionsA = CastChecked<FQuaternionAnimationAttribute>(MapA);
+			const TBoundValueMap<FQuaternionAnimationAttribute>* QuaternionsB = CastChecked<FQuaternionAnimationAttribute>(MapB);
+			TBoundValueMap<FQuaternionAnimationAttribute>* OutputQuaternions = CastChecked<FQuaternionAnimationAttribute>(OutputMap);
+
+			const FQuaternionAnimationAttribute* QuaternionPtrA = QuaternionsA->GetData();
+			const FQuaternionAnimationAttribute* QuaternionPtrB = QuaternionsB->GetData();
+			FQuaternionAnimationAttribute* OutputQuaternionPtr = OutputQuaternions->GetData();
+			const FQuaternionAnimationAttribute* QuaternionPtrAEnd = QuaternionPtrA + QuaternionsA->Num();
+
+			if (const TBoundValueMap<FFloatAnimationAttribute>* PerValueWeightsB = Cast<FFloatAnimationAttribute>(PerValueWeightMapB))
+			{
+				check(MapA->GetTypedSet() == PerValueWeightsB->GetTypedSet());
+				check(MapA->Num() == PerValueWeightsB->Num());
+
+				const FFloatAnimationAttribute* PerValueWeightsBPtr = PerValueWeightsB->GetData();
+
+				while (QuaternionPtrA < QuaternionPtrAEnd)
+				{
+					OutputQuaternionPtr->Value = FQuat::FastLerp(QuaternionPtrA->Value, QuaternionPtrB->Value, PerValueWeightsBPtr->Value);
+
+					QuaternionPtrA++;
+					QuaternionPtrB++;
+					PerValueWeightsBPtr++;
+					OutputQuaternionPtr++;
+				}
+			}
+			else
+			{
+				while (QuaternionPtrA < QuaternionPtrAEnd)
+				{
+					OutputQuaternionPtr->Value = FQuat::FastLerp(QuaternionPtrA->Value, QuaternionPtrB->Value, DefaultWeightB);
+
+					QuaternionPtrA++;
+					QuaternionPtrB++;
+					OutputQuaternionPtr++;
+				}
+			}
+		}
+
+		void FInterpolate_QuaternionAttribute::TransformUnboundValueMap(const FUnboundValueMap* MapA, const FUnboundValueMap* MapB, float WeightB, FUnboundValueMap*& OutputMap, FReallocFun OutputAllocator)
+		{
+			check(OutputMap == nullptr);
+
+			const TUnboundValueMap<FQuaternionAnimationAttribute>* QuaternionsA = CastChecked<FQuaternionAnimationAttribute>(MapA);
+			const TUnboundValueMap<FQuaternionAnimationAttribute>* QuaternionsB = CastChecked<FQuaternionAnimationAttribute>(MapB);
+			TUnboundValueMap<FQuaternionAnimationAttribute>* OutputQuaternions = MakeUnboundValueMap<FQuaternionAnimationAttribute>(OutputAllocator);
+
+			OuterJoinBy(
+				FUnboundValueMapJoinOp(),
+				[WeightB, OutputQuaternions](FName Name, const FQuaternionAnimationAttribute* A, const FQuaternionAnimationAttribute* B)
+				{
+					// Missing values are set to the identity
+					const FQuat& ValueA = A != nullptr ? A->Value : FQuat::Identity;
+					const FQuat& ValueB = B != nullptr ? B->Value : FQuat::Identity;
+
+					FQuaternionAnimationAttribute Output;
+					Output.Value = FQuat::FastLerp(ValueA, ValueB, WeightB);
+
+					OutputQuaternions->Append(Name, Output);
+				},
+				QuaternionsA->CreateConstIterator(),
+				QuaternionsB->CreateConstIterator());
+
+			OutputMap = OutputQuaternions;
+		}
 	}
 }
diff --git a/Engine/Plugins/Experimental/UAF/UAF/Source/UAF/Private/UAF/ValueRuntime/Transformers/Layer.cpp b/Engine/Plugins/Experimental/UAF/UAF/Source/UAF/Private/UAF/ValueRuntime/Transformers/Layer.cpp
--- a/Engine/Plugins/Experimental/UAF/UAF/Source/UAF/Private/UAF/ValueRuntime/Transformers/Layer.cpp
+++ b/Engine/Plugins/Experimental/UAF/UAF/Source/UAF/Private/UAF/ValueRuntime/Transformers/Layer.cpp
@@ -244,21 +244,13 @@ namespace UE::UAF::Transformers
 					{
 						if (BaseMap == nullptr)
 						{
-							if (LayerMap == nullptr)
-							{
-								// No maps found for the transformer type
-							}
-							else
-							{
-								// This set map doesn't have a matching pair, duplicate it
-								OutputMaps.Append(LayerMap->Duplicate(OutputMaps.GetAllocator()));
-							}
+							// Nothing to layer into, nothing to do
 						}
 						else
 						{
 							if (LayerMap == nullptr)
 							{
-								// This set map doesn't have a matching pair, duplicate it
+								// No layer map, nothing to do apart from duplicating unmodified base into output
 								OutputMaps.Append(BaseMap->Duplicate(OutputMaps.GetAllocator()));
 							}
 							else
@@ -312,22 +304,14 @@ namespace UE::UAF::Transformers
 					{
 						if (BaseMap == nullptr)
 						{
-							if (LayerMap == nullptr)
-							{
-								// No maps found for the transformer type
-							}
-							else
-							{
-								// This map doesn't have a matching pair, duplicate it
-								OutputMaps.Append(LayerMap->Duplicate(OutputMaps.GetAllocator()));
-							}
+							// Nothing to layer into, nothing to do
 						}
 						else
 						{
 							if (LayerMap == nullptr)
 							{
-								// This map doesn't have a matching pair
-								// Our output is A and thus already contains this
+								// No layer map, nothing to do apart from duplicating unmodified base into output
+								OutputMaps.Append(BaseMap->Duplicate(OutputMaps.GetAllocator()));
 							}
 							else
 							{
@@ -542,11 +526,387 @@ namespace UE::UAF::Transformers
 						FFloatAnimationAttribute Output{ FMath::Lerp(BaseValue->Value, LayerValue->Value, LayerWeight) };
 						OutputFloats->Append(Name, Output);
 					}
+					else if (BaseValue != nullptr)
+					{
+						OutputFloats->Append(Name, *BaseValue);
+					}
 				},
 				BaseFloats->CreateConstIterator(),
 				LayerFloats->CreateConstIterator());
 
 			OutputMap = OutputFloats;
 		}
+
+		void FLayer_IntegerAttribute::TransformBoundValueMap(const FBoundValueMap* Base, const FBoundValueMap* Layer, const FBoundValueMap* PerValueWeights, float DefaultLayerWeight, FBoundValueMap*& OutputMap, FReallocFun OutputAllocator)
+		{
+			if (OutputMap == nullptr)
+			{
+				const FBoundValueMap::FConstructArgs Args(Base->GetTypedSet(), Base->GetValueType(), OutputAllocator);
+				OutputMap = MakeBoundValueMap<FIntegerAnimationAttribute>(Args);
+			}
+
+			// Our output should have the same shape as our base
+			check(Base->GetTypedSet() == OutputMap->GetTypedSet());
+			check(Base->Num() == OutputMap->Num());
+
+			// Our per-value weights should have the same shape as our base
+			check(PerValueWeights == nullptr || Base->GetTypedSet() == PerValueWeights->GetTypedSet());
+			check(PerValueWeights == nullptr || Base->Num() == PerValueWeights->Num());
+
+			const FAttributeTypedSetPtr& BaseTypedSet = Base->GetTypedSet();
+			const FAttributeTypedSetPtr& LayerTypedSet = Layer->GetTypedSet();
+			const FAttributeTypedSet* PerValueWeightTypedSet = PerValueWeights != nullptr ? PerValueWeights->GetTypedSet().GetReference() : nullptr;
+
+			const TBoundValueMap<FIntegerAnimationAttribute>* BaseIntegers = CastChecked<FIntegerAnimationAttribute>(Base);
+			const TBoundValueMap<FIntegerAnimationAttribute>* LayerIntegers = CastChecked<FIntegerAnimationAttribute>(Layer);
+			const TBoundValueMap<FFloatAnimationAttribute>* PerValueWeightIntegers = PerValueWeights != nullptr ? CastChecked<FFloatAnimationAttribute>(PerValueWeights) : nullptr;
+			TBoundValueMap<FIntegerAnimationAttribute>* OutputIntegers = CastChecked<FIntegerAnimationAttribute>(OutputMap);
+
+			// Use an outer join over our base, layer, and weights
+			// If a value is in the base but not the layer, use the base value
+			// If a value is not in the base but in the layer, ignore it
+			// If a value is in both, blend it
+			//
+			// Efficiency of an outer join wouldn't be great and would be difficult to port so ISPC
+			// Instead we iterate over the base set and use the binding index to get layer values and weights
+
+			const int32 NumBaseBones = BaseTypedSet->Num();
+			for (FAttributeSetIndex BaseSetIndex(0); BaseSetIndex < NumBaseBones; ++BaseSetIndex)
+			{
+				const FAttributeBindingIndex BaseBindingIndex = BaseTypedSet->GetBindingIndex(BaseSetIndex);
+				check(BaseBindingIndex.IsValid());
+
+				int32 Result = (*BaseIntegers)[BaseSetIndex].Value;
+
+				const FAttributeSetIndex LayerSetIndex = LayerTypedSet->GetIndex(BaseBindingIndex);
+				if (LayerSetIndex.IsValid())
+				{
+					double LayerWeight = DefaultLayerWeight;
+					
+					if (PerValueWeightTypedSet != nullptr)
+					{
+						const FAttributeSetIndex PerWeightSetIndex = PerValueWeightTypedSet->GetIndex(BaseBindingIndex);
+						if (PerWeightSetIndex.IsValid())
+						{
+							LayerWeight = (*PerValueWeightIntegers)[PerWeightSetIndex].Value;
+							LayerWeight = FMath::Clamp(LayerWeight, 0.0, 1.0);
+						}
+					}
+
+					double BaseWeight = 1.0 - LayerWeight;
+					const int32 LayerValue = (*LayerIntegers)[LayerSetIndex].Value;
+
+					Result = FMath::TruncToInt32(Result * BaseWeight + LayerValue * LayerWeight);
+				}
+
+				(*OutputIntegers)[BaseSetIndex].Value = Result;
+			}
+		}
+
+		void FLayer_IntegerAttribute::TransformUnboundValueMap(const FUnboundValueMap* Base, const FUnboundValueMap* Layer, float LayerWeight, FUnboundValueMap*& OutputMap, FReallocFun OutputAllocator)
+		{
+			check(OutputMap == nullptr);
+
+			const TUnboundValueMap<FIntegerAnimationAttribute>* BaseIntegers = CastChecked<FIntegerAnimationAttribute>(Base);
+			const TUnboundValueMap<FIntegerAnimationAttribute>* LayerIntegers = CastChecked<FIntegerAnimationAttribute>(Layer);
+
+			TUnboundValueMap<FIntegerAnimationAttribute>* OutputIntegers = MakeUnboundValueMap<FIntegerAnimationAttribute>(OutputAllocator);
+
+			const double LayerWeightDouble = LayerWeight;
+			const double BaseWeightDouble = 1.0 - LayerWeight;
+			
+			OuterJoinBy(
+				FUnboundValueMapJoinOp(),
+				[LayerWeightDouble, BaseWeightDouble, OutputIntegers](FName Name, const FIntegerAnimationAttribute* BaseValue, const FIntegerAnimationAttribute* LayerValue)
+				{
+					// Mismatched pairs are ignored
+					if (BaseValue != nullptr && LayerValue != nullptr)
+					{
+						int32 Result = FMath::TruncToInt32(BaseValue->Value * BaseWeightDouble + LayerValue->Value * LayerWeightDouble);
+						
+						OutputIntegers->Append(Name, FIntegerAnimationAttribute { Result });
+					}
+				},
+				BaseIntegers->CreateConstIterator(),
+				LayerIntegers->CreateConstIterator());
+
+			OutputMap = OutputIntegers;
+		}
+
+		void FLayer_TransformAttribute::TransformBoundValueMap(const FBoundValueMap* Base, const FBoundValueMap* Layer, const FBoundValueMap* PerValueWeights, float DefaultLayerWeight, FBoundValueMap*& OutputMap, FReallocFun OutputAllocator)
+		{
+			if (OutputMap == nullptr)
+			{
+				const FBoundValueMap::FConstructArgs Args(Base->GetTypedSet(), Base->GetValueType(), OutputAllocator);
+				OutputMap = MakeBoundValueMap<FTransformAnimationAttribute>(Args);
+			}
+
+			// Our output should have the same shape as our base
+			check(Base->GetTypedSet() == OutputMap->GetTypedSet());
+			check(Base->Num() == OutputMap->Num());
+
+			// Our per-value weights should have the same shape as our base
+			check(PerValueWeights == nullptr || Base->GetTypedSet() == PerValueWeights->GetTypedSet());
+			check(PerValueWeights == nullptr || Base->Num() == PerValueWeights->Num());
+
+			const FAttributeTypedSetPtr& BaseTypedSet = Base->GetTypedSet();
+			const FAttributeTypedSetPtr& LayerTypedSet = Layer->GetTypedSet();
+			const FAttributeTypedSet* PerValueWeightTypedSet = PerValueWeights != nullptr ? PerValueWeights->GetTypedSet().GetReference() : nullptr;
+
+			const TBoundValueMap<FTransformAnimationAttribute>* BaseTransforms = CastChecked<FTransformAnimationAttribute>(Base);
+			const TBoundValueMap<FTransformAnimationAttribute>* LayerTransforms = CastChecked<FTransformAnimationAttribute>(Layer);
+			const TBoundValueMap<FFloatAnimationAttribute>* PerValueWeightFloats = PerValueWeights != nullptr ? CastChecked<FFloatAnimationAttribute>(PerValueWeights) : nullptr;
+			TBoundValueMap<FTransformAnimationAttribute>* OutputTransforms = CastChecked<FTransformAnimationAttribute>(OutputMap);
+
+			// Use an outer join over our base, layer, and weights
+			// If a value is in the base but not the layer, use the base value
+			// If a value is not in the base but in the layer, ignore it
+			// If a value is in both, blend it
+			//
+			// Efficiency of an outer join wouldn't be great and would be difficult to port so ISPC
+			// Instead we iterate over the base set and use the binding index to get layer values and weights
+
+			const int32 NumBaseBones = BaseTypedSet->Num();
+			for (FAttributeSetIndex BaseSetIndex(0); BaseSetIndex < NumBaseBones; ++BaseSetIndex)
+			{
+				const FAttributeBindingIndex BaseBindingIndex = BaseTypedSet->GetBindingIndex(BaseSetIndex);
+				check(BaseBindingIndex.IsValid());
+
+				FTransform Result = (*BaseTransforms)[BaseSetIndex].Value;
+
+				const FAttributeSetIndex LayerSetIndex = LayerTypedSet->GetIndex(BaseBindingIndex);
+				if (LayerSetIndex.IsValid())
+				{
+					float ValueWeight = DefaultLayerWeight;
+					if (PerValueWeightTypedSet != nullptr)
+					{
+						const FAttributeSetIndex PerWeightSetIndex = PerValueWeightTypedSet->GetIndex(BaseBindingIndex);
+						if (PerWeightSetIndex.IsValid())
+						{
+							ValueWeight = (*PerValueWeightFloats)[PerWeightSetIndex].Value;
+							ValueWeight = FMath::Clamp(ValueWeight, 0.0f, 1.0f);
+						}
+					}
+
+					Result.BlendWith((*LayerTransforms)[LayerSetIndex].Value, ValueWeight);
+				}
+
+				(*OutputTransforms)[BaseSetIndex].Value = Result;
+			}
+		}
+
+		void FLayer_TransformAttribute::TransformUnboundValueMap(const FUnboundValueMap* Base, const FUnboundValueMap* Layer, float LayerWeight, FUnboundValueMap*& OutputMap, FReallocFun OutputAllocator)
+		{
+			check(OutputMap == nullptr);
+
+			const TUnboundValueMap<FTransformAnimationAttribute>* BaseTransforms = CastChecked<FTransformAnimationAttribute>(Base);
+			const TUnboundValueMap<FTransformAnimationAttribute>* LayerTransforms = CastChecked<FTransformAnimationAttribute>(Layer);
+
+			TUnboundValueMap<FTransformAnimationAttribute>* OutputBoneTransforms = MakeUnboundValueMap<FTransformAnimationAttribute>(OutputAllocator);
+
+			OuterJoinBy(
+				FUnboundValueMapJoinOp(),
+				[LayerWeight, OutputBoneTransforms](FName Name, const FTransformAnimationAttribute* BaseValue, const FTransformAnimationAttribute* LayerValue)
+				{
+					// Mismatched pairs are ignored
+					if (BaseValue != nullptr && LayerValue != nullptr)
+					{
+						FTransformAnimationAttribute Output = *BaseValue;
+						Output.Value.BlendWith(LayerValue->Value, LayerWeight);
+
+						OutputBoneTransforms->Append(Name, Output);
+					}
+				},
+				BaseTransforms->CreateConstIterator(),
+				LayerTransforms->CreateConstIterator());
+
+			OutputMap = OutputBoneTransforms;
+		}
+
+		void FLayer_VectorAttribute::TransformBoundValueMap(const FBoundValueMap* Base, const FBoundValueMap* Layer, const FBoundValueMap* PerValueWeights, float DefaultLayerWeight, FBoundValueMap*& OutputMap, FReallocFun OutputAllocator)
+		{
+			if (OutputMap == nullptr)
+			{
+				const FBoundValueMap::FConstructArgs Args(Base->GetTypedSet(), Base->GetValueType(), OutputAllocator);
+				OutputMap = MakeBoundValueMap<FVectorAnimationAttribute>(Args);
+			}
+
+			// Our output should have the same shape as our base
+			check(Base->GetTypedSet() == OutputMap->GetTypedSet());
+			check(Base->Num() == OutputMap->Num());
+
+			// Our per-value weights should have the same shape as our base
+			check(PerValueWeights == nullptr || Base->GetTypedSet() == PerValueWeights->GetTypedSet());
+			check(PerValueWeights == nullptr || Base->Num() == PerValueWeights->Num());
+
+			const FAttributeTypedSetPtr& BaseTypedSet = Base->GetTypedSet();
+			const FAttributeTypedSetPtr& LayerTypedSet = Layer->GetTypedSet();
+			const FAttributeTypedSet* PerValueWeightTypedSet = PerValueWeights != nullptr ? PerValueWeights->GetTypedSet().GetReference() : nullptr;
+
+			const TBoundValueMap<FVectorAnimationAttribute>* BaseVectors = CastChecked<FVectorAnimationAttribute>(Base);
+			const TBoundValueMap<FVectorAnimationAttribute>* LayerVectors = CastChecked<FVectorAnimationAttribute>(Layer);
+			const TBoundValueMap<FFloatAnimationAttribute>* PerValueWeightFloats = PerValueWeights != nullptr ? CastChecked<FFloatAnimationAttribute>(PerValueWeights) : nullptr;
+			TBoundValueMap<FVectorAnimationAttribute>* OutputVectors = CastChecked<FVectorAnimationAttribute>(OutputMap);
+
+			// Use an outer join over our base, layer, and weights
+			// If a value is in the base but not the layer, use the base value
+			// If a value is not in the base but in the layer, ignore it
+			// If a value is in both, blend it
+			//
+			// Efficiency of an outer join wouldn't be great and would be difficult to port so ISPC
+			// Instead we iterate over the base set and use the binding index to get layer values and weights
+
+			const int32 NumBaseBones = BaseTypedSet->Num();
+			for (FAttributeSetIndex BaseSetIndex(0); BaseSetIndex < NumBaseBones; ++BaseSetIndex)
+			{
+				const FAttributeBindingIndex BaseBindingIndex = BaseTypedSet->GetBindingIndex(BaseSetIndex);
+				check(BaseBindingIndex.IsValid());
+
+				FVector Result = (*BaseVectors)[BaseSetIndex].Value;
+
+				const FAttributeSetIndex LayerSetIndex = LayerTypedSet->GetIndex(BaseBindingIndex);
+				if (LayerSetIndex.IsValid())
+				{
+					float ValueWeight = DefaultLayerWeight;
+					if (PerValueWeightTypedSet != nullptr)
+					{
+						const FAttributeSetIndex PerWeightSetIndex = PerValueWeightTypedSet->GetIndex(BaseBindingIndex);
+						if (PerWeightSetIndex.IsValid())
+						{
+							ValueWeight = (*PerValueWeightFloats)[PerWeightSetIndex].Value;
+							ValueWeight = FMath::Clamp(ValueWeight, 0.0f, 1.0f);
+						}
+					}
+
+					Result *= (1.0f - ValueWeight);
+					Result += (*LayerVectors)[LayerSetIndex].Value * ValueWeight;
+				}
+
+				(*OutputVectors)[BaseSetIndex].Value = Result;
+			}
+		}
+
+		void FLayer_VectorAttribute::TransformUnboundValueMap(const FUnboundValueMap* Base, const FUnboundValueMap* Layer, float LayerWeight, FUnboundValueMap*& OutputMap, FReallocFun OutputAllocator)
+		{
+			check(OutputMap == nullptr);
+
+			const TUnboundValueMap<FVectorAnimationAttribute>* BaseVectors = CastChecked<FVectorAnimationAttribute>(Base);
+			const TUnboundValueMap<FVectorAnimationAttribute>* LayerVectors = CastChecked<FVectorAnimationAttribute>(Layer);
+
+			TUnboundValueMap<FVectorAnimationAttribute>* OutputVectors = MakeUnboundValueMap<FVectorAnimationAttribute>(OutputAllocator);
+
+			OuterJoinBy(
+				FUnboundValueMapJoinOp(),
+				[LayerWeight, OutputVectors](FName Name, const FVectorAnimationAttribute* BaseValue, const FVectorAnimationAttribute* LayerValue)
+				{
+					// Mismatched pairs are ignored
+					if (BaseValue != nullptr && LayerValue != nullptr)
+					{
+						FVectorAnimationAttribute Output = *BaseValue;
+						Output.Value *= (1.0f - LayerWeight);
+						Output.Value += LayerValue->Value * LayerWeight;
+
+						OutputVectors->Append(Name, Output);
+					}
+					else if (BaseValue != nullptr)
+					{
+						OutputVectors->Append(Name, *BaseValue);
+					}
+				},
+				BaseVectors->CreateConstIterator(),
+				LayerVectors->CreateConstIterator());
+
+			OutputMap = OutputVectors;
+		}
+
+		void FLayer_QuaternionAttribute::TransformBoundValueMap(const FBoundValueMap* Base, const FBoundValueMap* Layer, const FBoundValueMap* PerValueWeights, float DefaultLayerWeight, FBoundValueMap*& OutputMap, FReallocFun OutputAllocator)
+		{
+			if (OutputMap == nullptr)
+			{
+				const FBoundValueMap::FConstructArgs Args(Base->GetTypedSet(), Base->GetValueType(), OutputAllocator);
+				OutputMap = MakeBoundValueMap<FQuaternionAnimationAttribute>(Args);
+			}
+
+			// Our output should have the same shape as our base
+			check(Base->GetTypedSet() == OutputMap->GetTypedSet());
+			check(Base->Num() == OutputMap->Num());
+
+			// Our per-value weights should have the same shape as our base
+			check(PerValueWeights == nullptr || Base->GetTypedSet() == PerValueWeights->GetTypedSet());
+			check(PerValueWeights == nullptr || Base->Num() == PerValueWeights->Num());
+
+			const FAttributeTypedSetPtr& BaseTypedSet = Base->GetTypedSet();
+			const FAttributeTypedSetPtr& LayerTypedSet = Layer->GetTypedSet();
+			const FAttributeTypedSet* PerValueWeightTypedSet = PerValueWeights != nullptr ? PerValueWeights->GetTypedSet().GetReference() : nullptr;
+
+			const TBoundValueMap<FQuaternionAnimationAttribute>* BaseQuaternions = CastChecked<FQuaternionAnimationAttribute>(Base);
+			const TBoundValueMap<FQuaternionAnimationAttribute>* LayerQuaternions = CastChecked<FQuaternionAnimationAttribute>(Layer);
+			const TBoundValueMap<FFloatAnimationAttribute>* PerValueWeightFloats = PerValueWeights != nullptr ? CastChecked<FFloatAnimationAttribute>(PerValueWeights) : nullptr;
+			TBoundValueMap<FQuaternionAnimationAttribute>* OutputQuaternions = CastChecked<FQuaternionAnimationAttribute>(OutputMap);
+
+			// Use an outer join over our base, layer, and weights
+			// If a value is in the base but not the layer, use the base value
+			// If a value is not in the base but in the layer, ignore it
+			// If a value is in both, blend it
+			//
+			// Efficiency of an outer join wouldn't be great and would be difficult to port so ISPC
+			// Instead we iterate over the base set and use the binding index to get layer values and weights
+
+			const int32 NumBaseBones = BaseTypedSet->Num();
+			for (FAttributeSetIndex BaseSetIndex(0); BaseSetIndex < NumBaseBones; ++BaseSetIndex)
+			{
+				const FAttributeBindingIndex BaseBindingIndex = BaseTypedSet->GetBindingIndex(BaseSetIndex);
+				check(BaseBindingIndex.IsValid());
+
+				FQuat Result = (*BaseQuaternions)[BaseSetIndex].Value;
+
+				const FAttributeSetIndex LayerSetIndex = LayerTypedSet->GetIndex(BaseBindingIndex);
+				if (LayerSetIndex.IsValid())
+				{
+					float ValueWeight = DefaultLayerWeight;
+					if (PerValueWeightTypedSet != nullptr)
+					{
+						const FAttributeSetIndex PerWeightSetIndex = PerValueWeightTypedSet->GetIndex(BaseBindingIndex);
+						if (PerWeightSetIndex.IsValid())
+						{
+							ValueWeight = (*PerValueWeightFloats)[PerWeightSetIndex].Value;
+							ValueWeight = FMath::Clamp(ValueWeight, 0.0f, 1.0f);
+						}
+					}
+
+					Result = FQuat::FastLerp(Result, (*LayerQuaternions)[LayerSetIndex].Value, ValueWeight);
+				}
+
+				(*OutputQuaternions)[BaseSetIndex].Value = Result;
+			}
+		}
+
+		void FLayer_QuaternionAttribute::TransformUnboundValueMap(const FUnboundValueMap* Base, const FUnboundValueMap* Layer, float LayerWeight, FUnboundValueMap*& OutputMap, FReallocFun OutputAllocator)
+		{
+			check(OutputMap == nullptr);
+
+			const TUnboundValueMap<FQuaternionAnimationAttribute>* BaseQuaternions = CastChecked<FQuaternionAnimationAttribute>(Base);
+			const TUnboundValueMap<FQuaternionAnimationAttribute>* LayerQuaternions = CastChecked<FQuaternionAnimationAttribute>(Layer);
+
+			TUnboundValueMap<FQuaternionAnimationAttribute>* OutputQuaternions = MakeUnboundValueMap<FQuaternionAnimationAttribute>(OutputAllocator);
+
+			OuterJoinBy(
+				FUnboundValueMapJoinOp(),
+				[LayerWeight, OutputQuaternions](FName Name, const FQuaternionAnimationAttribute* BaseValue, const FQuaternionAnimationAttribute* LayerValue)
+				{
+					// Mismatched pairs are ignored
+					if (BaseValue != nullptr && LayerValue != nullptr)
+					{
+						FQuaternionAnimationAttribute Output = *BaseValue;
+						Output.Value = FQuat::FastLerp(Output.Value, LayerValue->Value, LayerWeight);
+
+						OutputQuaternions->Append(Name, Output);
+					}
+				},
+				BaseQuaternions->CreateConstIterator(),
+				LayerQuaternions->CreateConstIterator());
+
+			OutputMap = OutputQuaternions;
+		}
 	}
 }
diff --git a/Engine/Plugins/Experimental/UAF/UAF/Source/UAF/Private/UAF/ValueRuntime/Transformers/Overwrite.cpp b/Engine/Plugins/Experimental/UAF/UAF/Source/UAF/Private/UAF/ValueRuntime/Transformers/Overwrite.cpp
--- a/Engine/Plugins/Experimental/UAF/UAF/Source/UAF/Private/UAF/ValueRuntime/Transformers/Overwrite.cpp
+++ b/Engine/Plugins/Experimental/UAF/UAF/Source/UAF/Private/UAF/ValueRuntime/Transformers/Overwrite.cpp
@@ -16,7 +16,7 @@ namespace UE::UAF::Transformers
 			// Our input is our output and it has full weight, nothing to do
 			return;
 		}
-		
+
 		// Allocator is irrelevant since it will remain empty
 		FValueBundleHeap EmptyCollection(Input.GetNamedSet());
 		FOverwrite::Apply(TransformerMap, Input, EmptyCollection, Weight, Output);
@@ -357,5 +357,322 @@ namespace UE::UAF::Transformers
 				}
 			}
 		}
+
+		void FOverwrite_IntegerAttribute::TransformBoundValueMap(const FBoundValueMap* InputMap, const FBoundValueMap* PerValueWeightMap, float DefaultWeight, FBoundValueMap*& OutputMap, FReallocFun OutputAllocator)
+		{
+			if (OutputMap == nullptr)
+			{
+				const FBoundValueMap::FConstructArgs Args(InputMap->GetTypedSet(), InputMap->GetValueType(), OutputAllocator);
+				OutputMap = MakeBoundValueMap<FIntegerAnimationAttribute>(Args);
+			}
+
+			check(InputMap->GetTypedSet() == OutputMap->GetTypedSet());
+			check(InputMap->Num() == OutputMap->Num());
+
+			const TBoundValueMap<FIntegerAnimationAttribute>* InputIntegers = CastChecked<FIntegerAnimationAttribute>(InputMap);
+			TBoundValueMap<FIntegerAnimationAttribute>* OutputIntegers = CastChecked<FIntegerAnimationAttribute>(OutputMap);
+
+			const FIntegerAnimationAttribute* IntegersPtr = InputIntegers->GetData();
+			FIntegerAnimationAttribute* OutputIntegersPtr = OutputIntegers->GetData();
+			const FIntegerAnimationAttribute* IntegersPtrEnd = IntegersPtr + InputIntegers->Num();
+
+			if (const TBoundValueMap<FFloatAnimationAttribute>* PerValueWeights = Cast<FFloatAnimationAttribute>(PerValueWeightMap))
+			{
+				check(InputMap->GetTypedSet() == PerValueWeights->GetTypedSet());
+				check(InputMap->Num() == PerValueWeights->Num());
+
+				const FFloatAnimationAttribute* PerValueWeightsPtr = PerValueWeights->GetData();
+
+				while (IntegersPtr < IntegersPtrEnd)
+				{
+					OutputIntegersPtr->Value = static_cast<int32>(IntegersPtr->Value * PerValueWeightsPtr->Value);
+
+					IntegersPtr++;
+					PerValueWeightsPtr++;
+					OutputIntegersPtr++;
+				}
+			}
+			else
+			{
+				while (IntegersPtr < IntegersPtrEnd)
+				{
+					OutputIntegersPtr->Value = static_cast<int32>(IntegersPtr->Value * DefaultWeight);
+
+					IntegersPtr++;
+					OutputIntegersPtr++;
+				}
+			}
+		}
+
+		void FOverwrite_IntegerAttribute::TransformUnboundValueMap(const FUnboundValueMap* InputMap, float Weight, FUnboundValueMap*& OutputMap, FReallocFun OutputAllocator)
+		{
+			const bool bAllocateOutput = OutputMap == nullptr;
+			if (bAllocateOutput)
+			{
+				OutputMap = MakeUnboundValueMap<FIntegerAnimationAttribute>(OutputAllocator);
+			}
+
+			const TUnboundValueMap<FIntegerAnimationAttribute>* InputIntegers = CastChecked<FIntegerAnimationAttribute>(InputMap);
+			TUnboundValueMap<FIntegerAnimationAttribute>* OutputIntegers = CastChecked<FIntegerAnimationAttribute>(OutputMap);
+
+			if (bAllocateOutput)
+			{
+				// We allocated a new output, append to it
+				// TODO: Memcpy for weight = 1.0?
+				for (auto It = InputIntegers->CreateConstIterator(); It; ++It)
+				{
+					FIntegerAnimationAttribute Output{ static_cast<int32>(It.GetValue()->Value * Weight) };
+					OutputIntegers->Append(It.GetName(), Output);
+				}
+			}
+			else
+			{
+				// We overwrite in place
+				check(InputIntegers == OutputIntegers);
+
+				for (auto It = OutputIntegers->CreateIterator(); It; ++It)
+				{
+					It.GetValue()->Value = static_cast<int32>(It.GetValue()->Value * Weight);
+				}
+			}
+		}
+
+		void FOverwrite_TransformAttribute::TransformBoundValueMap(const FBoundValueMap* InputMap, const FBoundValueMap* PerValueWeightMap, float DefaultWeight, FBoundValueMap*& OutputMap, FReallocFun OutputAllocator)
+		{
+			if (OutputMap == nullptr)
+			{
+				const FBoundValueMap::FConstructArgs Args(InputMap->GetTypedSet(), InputMap->GetValueType(), OutputAllocator);
+				OutputMap = MakeBoundValueMap<FTransformAnimationAttribute>(Args);
+			}
+
+			check(InputMap->GetTypedSet() == OutputMap->GetTypedSet());
+			check(InputMap->Num() == OutputMap->Num());
+
+			const TBoundValueMap<FTransformAnimationAttribute>* InputTransforms = CastChecked<FTransformAnimationAttribute>(InputMap);
+			TBoundValueMap<FTransformAnimationAttribute>* OutputTransforms = CastChecked<FTransformAnimationAttribute>(OutputMap);
+
+			const FTransformAnimationAttribute* InputTransformPtr = InputTransforms->GetData();
+			FTransformAnimationAttribute* OutputTransformPtr = OutputTransforms->GetData();
+			const FTransformAnimationAttribute* InputTransformPtrEnd = InputTransformPtr + InputTransforms->Num();
+
+			if (const TBoundValueMap<FFloatAnimationAttribute>* PerValueWeights = Cast<FFloatAnimationAttribute>(PerValueWeightMap))
+			{
+				check(InputMap->GetTypedSet() == PerValueWeights->GetTypedSet());
+				check(InputMap->Num() == PerValueWeights->Num());
+
+				const FFloatAnimationAttribute* PerValueWeightsPtr = PerValueWeights->GetData();
+
+				while (InputTransformPtr < InputTransformPtrEnd)
+				{
+					OutputTransformPtr->Value = InputTransformPtr->Value * ScalarRegister(PerValueWeightsPtr->Value);
+
+					InputTransformPtr++;
+					PerValueWeightsPtr++;
+					OutputTransformPtr++;
+				}
+			}
+			else
+			{
+				const ScalarRegister VDefaultWeight(DefaultWeight);
+
+				while (InputTransformPtr < InputTransformPtrEnd)
+				{
+					OutputTransformPtr->Value = InputTransformPtr->Value * VDefaultWeight;
+
+					InputTransformPtr++;
+					OutputTransformPtr++;
+				}
+			}
+		}
+
+		void FOverwrite_TransformAttribute::TransformUnboundValueMap(const FUnboundValueMap* InputMap, float Weight, FUnboundValueMap*& OutputMap, FReallocFun OutputAllocator)
+		{
+			const bool bAllocateOutput = OutputMap == nullptr;
+			if (bAllocateOutput)
+			{
+				OutputMap = MakeUnboundValueMap<FTransformAnimationAttribute>(OutputAllocator);
+			}
+
+			const TUnboundValueMap<FTransformAnimationAttribute>* InputTransforms = CastChecked<FTransformAnimationAttribute>(InputMap);
+			TUnboundValueMap<FTransformAnimationAttribute>* OutputTransforms = CastChecked<FTransformAnimationAttribute>(OutputMap);
+
+			const ScalarRegister VWeight(Weight);
+
+			if (bAllocateOutput)
+			{
+				// We allocated a new output, append to it
+				for (auto It = InputTransforms->CreateConstIterator(); It; ++It)
+				{
+					FBoneTransformAnimationAttribute Output{ It.GetValue()->Value * VWeight };
+					OutputTransforms->Append(It.GetName(), Output);
+				}
+			}
+			else
+			{
+				// We overwrite in place
+				check(InputTransforms == OutputTransforms);
+
+				for (auto It = OutputTransforms->CreateIterator(); It; ++It)
+				{
+					It.GetValue()->Value *= VWeight;
+				}
+			}
+		}
+
+		void FOverwrite_VectorAttribute::TransformBoundValueMap(const FBoundValueMap* InputMap, const FBoundValueMap* PerValueWeightMap, float DefaultWeight, FBoundValueMap*& OutputMap, FReallocFun OutputAllocator)
+		{
+			if (OutputMap == nullptr)
+			{
+				const FBoundValueMap::FConstructArgs Args(InputMap->GetTypedSet(), InputMap->GetValueType(), OutputAllocator);
+				OutputMap = MakeBoundValueMap<FTransformAnimationAttribute>(Args);
+			}
+
+			check(InputMap->GetTypedSet() == OutputMap->GetTypedSet());
+			check(InputMap->Num() == OutputMap->Num());
+
+			const TBoundValueMap<FVectorAnimationAttribute>* InputVectors = CastChecked<FVectorAnimationAttribute>(InputMap);
+			TBoundValueMap<FVectorAnimationAttribute>* OutputVectors = CastChecked<FVectorAnimationAttribute>(OutputMap);
+			
+			const FVectorAnimationAttribute* InputVectorPtr = InputVectors->GetData();
+			FVectorAnimationAttribute* OutputVectorPtr = OutputVectors->GetData();
+			const FVectorAnimationAttribute* InputVectorPtrEnd = InputVectorPtr + InputVectors->Num();
+
+			if (const TBoundValueMap<FFloatAnimationAttribute>* PerValueWeights = Cast<FFloatAnimationAttribute>(PerValueWeightMap))
+			{
+				check(InputMap->GetTypedSet() == PerValueWeights->GetTypedSet());
+				check(InputMap->Num() == PerValueWeights->Num());
+
+				const FFloatAnimationAttribute* PerValueWeightsPtr = PerValueWeights->GetData();
+
+				while (InputVectorPtr < InputVectorPtrEnd)
+				{
+					OutputVectorPtr->Value = InputVectorPtr->Value * PerValueWeightsPtr->Value;
+
+					InputVectorPtr++;
+					PerValueWeightsPtr++;
+					OutputVectorPtr++;
+				}
+			}
+			else
+			{
+				while (InputVectorPtr < InputVectorPtrEnd)
+				{
+					OutputVectorPtr->Value = InputVectorPtr->Value * DefaultWeight;
+
+					InputVectorPtr++;
+					OutputVectorPtr++;
+				}
+			}
+		}
+
+		void FOverwrite_VectorAttribute::TransformUnboundValueMap(const FUnboundValueMap* InputMap, float Weight, FUnboundValueMap*& OutputMap, FReallocFun OutputAllocator)
+		{
+			const bool bAllocateOutput = OutputMap == nullptr;
+			if (bAllocateOutput)
+			{
+				OutputMap = MakeUnboundValueMap<FVectorAnimationAttribute>(OutputAllocator);
+			}
+
+			const TUnboundValueMap<FVectorAnimationAttribute>* InputVectors = CastChecked<FVectorAnimationAttribute>(InputMap);
+			TUnboundValueMap<FVectorAnimationAttribute>* OutputVectors = CastChecked<FVectorAnimationAttribute>(OutputMap);
+
+			if (bAllocateOutput)
+			{
+				// We allocated a new output, append to it
+				for (auto It = InputVectors->CreateConstIterator(); It; ++It)
+				{
+					FVectorAnimationAttribute Output{ It.GetValue()->Value * Weight };
+					OutputVectors->Append(It.GetName(), Output);
+				}
+			}
+			else
+			{
+				// We overwrite in place
+				check(InputVectors == OutputVectors);
+
+				for (auto It = OutputVectors->CreateIterator(); It; ++It)
+				{
+					It.GetValue()->Value *= Weight;
+				}
+			}
+		}
+
+		void FOverwrite_QuaternionAttribute::TransformBoundValueMap(const FBoundValueMap* InputMap, const FBoundValueMap* PerValueWeightMap, float DefaultWeight, FBoundValueMap*& OutputMap, FReallocFun OutputAllocator)
+		{
+			if (OutputMap == nullptr)
+			{
+				const FBoundValueMap::FConstructArgs Args(InputMap->GetTypedSet(), InputMap->GetValueType(), OutputAllocator);
+				OutputMap = MakeBoundValueMap<FBoneTransformAnimationAttribute>(Args);
+			}
+
+			check(InputMap->GetTypedSet() == OutputMap->GetTypedSet());
+			check(InputMap->Num() == OutputMap->Num());
+
+			const TBoundValueMap<FQuaternionAnimationAttribute>* InputQuaternions = CastChecked<FQuaternionAnimationAttribute>(InputMap);
+			TBoundValueMap<FQuaternionAnimationAttribute>* OutputQuaternions = CastChecked<FQuaternionAnimationAttribute>(OutputMap);
+
+			const FQuaternionAnimationAttribute* InputQuaternionPtr = InputQuaternions->GetData();
+			FQuaternionAnimationAttribute* OutputQuaternionPtr = OutputQuaternions->GetData();
+			const FQuaternionAnimationAttribute* InputQuaternionPtrEnd = InputQuaternionPtr + InputQuaternions->Num();
+
+			if (const TBoundValueMap<FFloatAnimationAttribute>* PerValueWeights = Cast<FFloatAnimationAttribute>(PerValueWeightMap))
+			{
+				check(InputMap->GetTypedSet() == PerValueWeights->GetTypedSet());
+				check(InputMap->Num() == PerValueWeights->Num());
+
+				const FFloatAnimationAttribute* PerValueWeightsPtr = PerValueWeights->GetData();
+
+				while (InputQuaternionPtr < InputQuaternionPtrEnd)
+				{
+					OutputQuaternionPtr->Value = InputQuaternionPtr->Value * PerValueWeightsPtr->Value;
+
+					InputQuaternionPtr++;
+					PerValueWeightsPtr++;
+					OutputQuaternionPtr++;
+				}
+			}
+			else
+			{
+				while (InputQuaternionPtr < InputQuaternionPtrEnd)
+				{
+					OutputQuaternionPtr->Value = InputQuaternionPtr->Value * DefaultWeight;
+
+					InputQuaternionPtr++;
+					OutputQuaternionPtr++;
+				}
+			}
+		}
+
+		void FOverwrite_QuaternionAttribute::TransformUnboundValueMap(const FUnboundValueMap* InputMap, float Weight, FUnboundValueMap*& OutputMap, FReallocFun OutputAllocator)
+		{
+			const bool bAllocateOutput = OutputMap == nullptr;
+			if (bAllocateOutput)
+			{
+				OutputMap = MakeUnboundValueMap<FQuaternionAnimationAttribute>(OutputAllocator);
+			}
+
+			const TUnboundValueMap<FQuaternionAnimationAttribute>* InputQuaternions = CastChecked<FQuaternionAnimationAttribute>(InputMap);
+			TUnboundValueMap<FQuaternionAnimationAttribute>* OutputQuaternions = CastChecked<FQuaternionAnimationAttribute>(OutputMap);
+
+			if (bAllocateOutput)
+			{
+				// We allocated a new output, append to it
+				for (auto It = InputQuaternions->CreateConstIterator(); It; ++It)
+				{
+					FQuaternionAnimationAttribute Output{ It.GetValue()->Value * Weight };
+					OutputQuaternions->Append(It.GetName(), Output);
+				}
+			}
+			else
+			{
+				// We overwrite in place
+				check(InputQuaternions == OutputQuaternions);
+
+				for (auto It = OutputQuaternions->CreateIterator(); It; ++It)
+				{
+					It.GetValue()->Value *= Weight;
+				}
+			}
+		}
 	}
 }
diff --git a/Engine/Plugins/Experimental/UAF/UAF/Source/UAF/Private/UAF/ValueRuntime/Transformers/Sanitize.cpp b/Engine/Plugins/Experimental/UAF/UAF/Source/UAF/Private/UAF/ValueRuntime/Transformers/Sanitize.cpp
--- a/Engine/Plugins/Experimental/UAF/UAF/Source/UAF/Private/UAF/ValueRuntime/Transformers/Sanitize.cpp
+++ b/Engine/Plugins/Experimental/UAF/UAF/Source/UAF/Private/UAF/ValueRuntime/Transformers/Sanitize.cpp
@@ -43,24 +43,26 @@ namespace UE::UAF::Transformers
 	namespace Private
 	{
 #if UE_UAF_VALIDATE_MAPPED_VALUES
+		
 		// Don't inline this function to keep the stack usage down
-		static FORCENOINLINE void IsBoneTransformMapValid(const TBoundValueMap<FBoneTransformAnimationAttribute>* BoneTransforms)
+		template <typename TransformAttributeType>
+		static FORCENOINLINE void IsTransformMapValid(const TBoundValueMap<TransformAttributeType>* Transforms)
 		{
-			const FBoneTransformAnimationAttribute* TransformPtr = BoneTransforms->GetData();
-			const FBoneTransformAnimationAttribute* EndTransformPtr = TransformPtr + BoneTransforms->Num();
+			const TransformAttributeType* TransformPtr = Transforms->GetData();
+			const TransformAttributeType* EndTransformPtr = TransformPtr + Transforms->Num();
 
 			while (TransformPtr < EndTransformPtr)
 			{
 				if (TransformPtr->Value.ContainsNaN())
 				{
-					ensureMsgf(false, TEXT("Bone (%s) contains NaN with Value:[%s]"),
-						*BoneTransforms->GetTypedSet()->GetName(FAttributeSetIndex(TransformPtr - BoneTransforms->GetData())).ToString(),
+					ensureMsgf(false, TEXT("Transform (%s) contains NaN with Value:[%s]"),
+						*Transforms->GetTypedSet()->GetName(FAttributeSetIndex(TransformPtr - Transforms->GetData())).ToString(),
 						*TransformPtr->Value.ToString());
 				}
 				else if (!TransformPtr->Value.IsRotationNormalized())
 				{
-					ensureMsgf(false, TEXT("Bone (%s) rotation not normalized with Rotation:[%s]"),
-						*BoneTransforms->GetTypedSet()->GetName(FAttributeSetIndex(TransformPtr - BoneTransforms->GetData())).ToString(),
+					ensureMsgf(false, TEXT("Transform (%s) rotation not normalized with Rotation:[%s]"),
+						*Transforms->GetTypedSet()->GetName(FAttributeSetIndex(TransformPtr - Transforms->GetData())).ToString(),
 						*TransformPtr->Value.GetRotation().ToString());
 				}
 
@@ -69,9 +71,10 @@ namespace UE::UAF::Transformers
 		}
 
 		// Don't inline this function to keep the stack usage down
-		static FORCENOINLINE void IsBoneTransformMapValid(const TUnboundValueMap<FBoneTransformAnimationAttribute>* BoneTransforms)
+		template <typename TransformAttributeType>
+		static FORCENOINLINE void IsTransformMapValid(const TUnboundValueMap<TransformAttributeType>* Transforms)
 		{
-			for (auto It = BoneTransforms->CreateConstIterator(); It; ++It)
+			for (auto It = Transforms->CreateConstIterator(); It; ++It)
 			{
 				if (It.GetValue()->Value.ContainsNaN())
 				{
@@ -104,7 +107,7 @@ namespace UE::UAF::Transformers
 			}
 
 #if UE_UAF_VALIDATE_MAPPED_VALUES
-			IsBoneTransformMapValid(BoneTransforms);
+			IsTransformMapValid(BoneTransforms);
 #endif
 		}
 
@@ -118,7 +121,7 @@ namespace UE::UAF::Transformers
 			}
 
 #if UE_UAF_VALIDATE_MAPPED_VALUES
-			IsBoneTransformMapValid(BoneTransforms);
+			IsTransformMapValid(BoneTransforms);
 #endif
 		}
 
@@ -152,5 +155,63 @@ namespace UE::UAF::Transformers
 			}
 		}
 #endif
+
+		void FSanitize_TransformAttribute::TransformBoundValueMap(FBoundValueMap* Map)
+		{
+			TBoundValueMap<FTransformAnimationAttribute>* Transforms = CastChecked<FTransformAnimationAttribute>(Map);
+
+			FTransformAnimationAttribute* TransformPtr = Transforms->GetData();
+			const FTransformAnimationAttribute* TransformPtrEnd = TransformPtr + Transforms->Num();
+
+			while (TransformPtr < TransformPtrEnd)
+			{
+				TransformPtr->Value.NormalizeRotation();
+
+				TransformPtr++;
+			}
+
+#if UE_UAF_VALIDATE_MAPPED_VALUES
+			IsTransformMapValid(Transforms);
+#endif
+		}
+
+		void FSanitize_TransformAttribute::TransformUnboundValueMap(FUnboundValueMap* Map)
+		{
+			TUnboundValueMap<FTransformAnimationAttribute>* Transforms = CastChecked<FTransformAnimationAttribute>(Map);
+
+			for (auto It = Transforms->CreateIterator(); It; ++It)
+			{
+				It.GetValue()->Value.NormalizeRotation();
+			}
+
+#if UE_UAF_VALIDATE_MAPPED_VALUES
+			IsTransformMapValid(Transforms);
+#endif
+		}
+
+		void FSanitize_QuaternionAttribute::TransformBoundValueMap(FBoundValueMap* Map)
+		{
+			TBoundValueMap<FQuaternionAnimationAttribute>* Quaternions = CastChecked<FQuaternionAnimationAttribute>(Map);
+
+			FQuaternionAnimationAttribute* QuaternionPtr = Quaternions->GetData();
+			const FQuaternionAnimationAttribute* QuaternionPtrEnd = QuaternionPtr + Quaternions->Num();
+
+			while (QuaternionPtr < QuaternionPtrEnd)
+			{
+				QuaternionPtr->Value.Normalize();
+
+				QuaternionPtr++;
+			}
+		}
+
+		void FSanitize_QuaternionAttribute::TransformUnboundValueMap(FUnboundValueMap* Map)
+		{
+			TUnboundValueMap<FQuaternionAnimationAttribute>* Quaternions = CastChecked<FQuaternionAnimationAttribute>(Map);
+
+			for (auto It = Quaternions->CreateIterator(); It; ++It)
+			{
+				It.GetValue()->Value.Normalize();
+			}
+		}
 	}
 }
diff --git a/Engine/Plugins/Experimental/UAF/UAF/Source/UAF/Public/UAF/ValueRuntime/Transformers/Accumulate.h b/Engine/Plugins/Experimental/UAF/UAF/Source/UAF/Public/UAF/ValueRuntime/Transformers/Accumulate.h
--- a/Engine/Plugins/Experimental/UAF/UAF/Source/UAF/Public/UAF/ValueRuntime/Transformers/Accumulate.h
+++ b/Engine/Plugins/Experimental/UAF/UAF/Source/UAF/Public/UAF/ValueRuntime/Transformers/Accumulate.h
@@ -58,6 +58,42 @@ namespace UE::UAF::Transformers
 			static void TransformBoundValueMap(const FBoundValueMap* MapA, const FBoundValueMap* MapB, const FBoundValueMap* PerValueWeightMapB, float DefaultWeightB, FBoundValueMap*& OutputMap, FReallocFun OutputAllocator);
 			static void TransformUnboundValueMap(const FUnboundValueMap* MapA, const FUnboundValueMap* MapB, float WeightB, FUnboundValueMap*& OutputMap, FReallocFun OutputAllocator);
 		};
+
+		struct FAccumulate_IntegerAttribute
+		{
+			using TransformerType = FAccumulate;
+			using ValueType = FIntegerAnimationAttribute;
+
+			static void TransformBoundValueMap(const FBoundValueMap* MapA, const FBoundValueMap* MapB, const FBoundValueMap* PerValueWeightMapB, float DefaultWeightB, FBoundValueMap*& OutputMap, FReallocFun OutputAllocator);
+			static void TransformUnboundValueMap(const FUnboundValueMap* MapA, const FUnboundValueMap* MapB, float WeightB, FUnboundValueMap*& OutputMap, FReallocFun OutputAllocator);
+		};
+
+		struct FAccumulate_TransformAttribute
+		{
+			using TransformerType = FAccumulate;
+			using ValueType = FTransformAnimationAttribute;
+
+			static void TransformBoundValueMap(const FBoundValueMap* MapA, const FBoundValueMap* MapB, const FBoundValueMap* PerValueWeightMapB, float DefaultWeightB, FBoundValueMap*& OutputMap, FReallocFun OutputAllocator);
+			static void TransformUnboundValueMap(const FUnboundValueMap* MapA, const FUnboundValueMap* MapB, float WeightB, FUnboundValueMap*& OutputMap, FReallocFun OutputAllocator);
+		};
+
+		struct FAccumulate_VectorAttribute
+		{
+			using TransformerType = FAccumulate;
+			using ValueType = FVectorAnimationAttribute;
+
+			static void TransformBoundValueMap(const FBoundValueMap* MapA, const FBoundValueMap* MapB, const FBoundValueMap* PerValueWeightMapB, float DefaultWeightB, FBoundValueMap*& OutputMap, FReallocFun OutputAllocator);
+			static void TransformUnboundValueMap(const FUnboundValueMap* MapA, const FUnboundValueMap* MapB, float WeightB, FUnboundValueMap*& OutputMap, FReallocFun OutputAllocator);
+		};
+
+		struct FAccumulate_QuaternionAttribute
+		{
+			using TransformerType = FAccumulate;
+			using ValueType = FQuaternionAnimationAttribute;
+
+			static void TransformBoundValueMap(const FBoundValueMap* MapA, const FBoundValueMap* MapB, const FBoundValueMap* PerValueWeightMapB, float DefaultWeightB, FBoundValueMap*& OutputMap, FReallocFun OutputAllocator);
+			static void TransformUnboundValueMap(const FUnboundValueMap* MapA, const FUnboundValueMap* MapB, float WeightB, FUnboundValueMap*& OutputMap, FReallocFun OutputAllocator);
+		};
 	}
 }
 
diff --git a/Engine/Plugins/Experimental/UAF/UAF/Source/UAF/Public/UAF/ValueRuntime/Transformers/AdditiveSpace.h b/Engine/Plugins/Experimental/UAF/UAF/Source/UAF/Public/UAF/ValueRuntime/Transformers/AdditiveSpace.h
--- a/Engine/Plugins/Experimental/UAF/UAF/Source/UAF/Public/UAF/ValueRuntime/Transformers/AdditiveSpace.h
+++ b/Engine/Plugins/Experimental/UAF/UAF/Source/UAF/Public/UAF/ValueRuntime/Transformers/AdditiveSpace.h
@@ -87,6 +87,42 @@ namespace UE::UAF::Transformers
 			static void TransformUnboundValueMap(const FUnboundValueMap* Base, const FUnboundValueMap* Input, FUnboundValueMap*& Output, FReallocFun OutputAllocator);
 		};
 
+		struct FMakeAdditiveSpace_IntegerAttribute
+		{
+			using TransformerType = FMakeAdditiveSpace;
+			using ValueType = FIntegerAnimationAttribute;
+
+			static void TransformBoundValueMap(const FBoundValueMap* Base, const FBoundValueMap* Input, FBoundValueMap*& Output, FReallocFun OutputAllocator);
+			static void TransformUnboundValueMap(const FUnboundValueMap* Base, const FUnboundValueMap* Input, FUnboundValueMap*& Output, FReallocFun OutputAllocator);
+		};
+
+		struct FMakeAdditiveSpace_TransformAttribute
+		{
+			using TransformerType = FMakeAdditiveSpace;
+			using ValueType = FTransformAnimationAttribute;
+
+			static void TransformBoundValueMap(const FBoundValueMap* Base, const FBoundValueMap* Input, FBoundValueMap*& Output, FReallocFun OutputAllocator);
+			static void TransformUnboundValueMap(const FUnboundValueMap* Base, const FUnboundValueMap* Input, FUnboundValueMap*& Output, FReallocFun OutputAllocator);
+		};
+
+		struct FMakeAdditiveSpace_VectorAttribute
+		{
+			using TransformerType = FMakeAdditiveSpace;
+			using ValueType = FVectorAnimationAttribute;
+
+			static void TransformBoundValueMap(const FBoundValueMap* Base, const FBoundValueMap* Input, FBoundValueMap*& Output, FReallocFun OutputAllocator);
+			static void TransformUnboundValueMap(const FUnboundValueMap* Base, const FUnboundValueMap* Input, FUnboundValueMap*& Output, FReallocFun OutputAllocator);
+		};
+
+		struct FMakeAdditiveSpace_QuaternionAttribute
+		{
+			using TransformerType = FMakeAdditiveSpace;
+			using ValueType = FQuaternionAnimationAttribute;
+
+			static void TransformBoundValueMap(const FBoundValueMap* Base, const FBoundValueMap* Input, FBoundValueMap*& Output, FReallocFun OutputAllocator);
+			static void TransformUnboundValueMap(const FUnboundValueMap* Base, const FUnboundValueMap* Input, FUnboundValueMap*& Output, FReallocFun OutputAllocator);
+		};
+
 		struct FApplyAdditiveSpace_BoneTransformAttribute
 		{
 			using TransformerType = FApplyAdditiveSpace;
@@ -104,6 +140,42 @@ namespace UE::UAF::Transformers
 			static void TransformBoundValueMap(const FBoundValueMap* Base, const FBoundValueMap* Additive, const FBoundValueMap* PerValueWeightMap, float DefaultWeight, FBoundValueMap*& Output, FReallocFun OutputAllocator);
 			static void TransformUnboundValueMap(const FUnboundValueMap* Base, const FUnboundValueMap* Additive, float AdditiveWeight, FUnboundValueMap*& Output, FReallocFun OutputAllocator);
 		};
+
+		struct FApplyAdditiveSpace_IntegerAttribute
+		{
+			using TransformerType = FApplyAdditiveSpace;
+			using ValueType = FIntegerAnimationAttribute;
+
+			static void TransformBoundValueMap(const FBoundValueMap* Base, const FBoundValueMap* Additive, const FBoundValueMap* PerValueWeightMap, float DefaultWeight, FBoundValueMap*& Output, FReallocFun OutputAllocator);
+			static void TransformUnboundValueMap(const FUnboundValueMap* Base, const FUnboundValueMap* Additive, float AdditiveWeight, FUnboundValueMap*& Output, FReallocFun OutputAllocator);
+		};
+
+		struct FApplyAdditiveSpace_TransformAttribute
+		{
+			using TransformerType = FApplyAdditiveSpace;
+			using ValueType = FTransformAnimationAttribute;
+
+			static void TransformBoundValueMap(const FBoundValueMap* Base, const FBoundValueMap* Additive, const FBoundValueMap* PerValueWeightMap, float DefaultWeight, FBoundValueMap*& Output, FReallocFun OutputAllocator);
+			static void TransformUnboundValueMap(const FUnboundValueMap* Base, const FUnboundValueMap* Additive, float AdditiveWeight, FUnboundValueMap*& Output, FReallocFun OutputAllocator);
+		};
+
+		struct FApplyAdditiveSpace_VectorAttribute
+		{
+			using TransformerType = FApplyAdditiveSpace;
+			using ValueType = FVectorAnimationAttribute;
+
+			static void TransformBoundValueMap(const FBoundValueMap* Base, const FBoundValueMap* Additive, const FBoundValueMap* PerValueWeightMap, float DefaultWeight, FBoundValueMap*& Output, FReallocFun OutputAllocator);
+			static void TransformUnboundValueMap(const FUnboundValueMap* Base, const FUnboundValueMap* Additive, float AdditiveWeight, FUnboundValueMap*& Output, FReallocFun OutputAllocator);
+		};
+
+		struct FApplyAdditiveSpace_QuaternionAttribute
+		{
+			using TransformerType = FApplyAdditiveSpace;
+			using ValueType = FQuaternionAnimationAttribute;
+
+			static void TransformBoundValueMap(const FBoundValueMap* Base, const FBoundValueMap* Additive, const FBoundValueMap* PerValueWeightMap, float DefaultWeight, FBoundValueMap*& Output, FReallocFun OutputAllocator);
+			static void TransformUnboundValueMap(const FUnboundValueMap* Base, const FUnboundValueMap* Additive, float AdditiveWeight, FUnboundValueMap*& Output, FReallocFun OutputAllocator);
+		};
 	}
 }
 
diff --git a/Engine/Plugins/Experimental/UAF/UAF/Source/UAF/Public/UAF/ValueRuntime/Transformers/Interpolate.h b/Engine/Plugins/Experimental/UAF/UAF/Source/UAF/Public/UAF/ValueRuntime/Transformers/Interpolate.h
--- a/Engine/Plugins/Experimental/UAF/UAF/Source/UAF/Public/UAF/ValueRuntime/Transformers/Interpolate.h
+++ b/Engine/Plugins/Experimental/UAF/UAF/Source/UAF/Public/UAF/ValueRuntime/Transformers/Interpolate.h
@@ -59,6 +59,42 @@ namespace UE::UAF::Transformers
 			static void TransformBoundValueMap(const FBoundValueMap* MapA, const FBoundValueMap* MapB, const FBoundValueMap* PerValueWeightMapB, float DefaultWeightB, FBoundValueMap*& OutputMap, FReallocFun OutputAllocator);
 			static void TransformUnboundValueMap(const FUnboundValueMap* MapA, const FUnboundValueMap* MapB, float WeightB, FUnboundValueMap*& OutputMap, FReallocFun OutputAllocator);
 		};
+
+		struct FInterpolate_IntegerAttribute
+		{
+			using TransformerType = FInterpolate;
+			using ValueType = FIntegerAnimationAttribute;
+
+			static void TransformBoundValueMap(const FBoundValueMap* MapA, const FBoundValueMap* MapB, const FBoundValueMap* PerValueWeightMapB, float DefaultWeightB, FBoundValueMap*& OutputMap, FReallocFun OutputAllocator);
+			static void TransformUnboundValueMap(const FUnboundValueMap* MapA, const FUnboundValueMap* MapB, float WeightB, FUnboundValueMap*& OutputMap, FReallocFun OutputAllocator);
+		};
+
+		struct FInterpolate_TransformAttribute
+		{
+			using TransformerType = FInterpolate;
+			using ValueType = FTransformAnimationAttribute;
+
+			static void TransformBoundValueMap(const FBoundValueMap* MapA, const FBoundValueMap* MapB, const FBoundValueMap* PerValueWeightMapB, float DefaultWeightB, FBoundValueMap*& OutputMap, FReallocFun OutputAllocator);
+			static void TransformUnboundValueMap(const FUnboundValueMap* MapA, const FUnboundValueMap* MapB, float WeightB, FUnboundValueMap*& OutputMap, FReallocFun OutputAllocator);
+		};
+
+		struct FInterpolate_VectorAttribute
+		{
+			using TransformerType = FInterpolate;
+			using ValueType = FVectorAnimationAttribute;
+
+			static void TransformBoundValueMap(const FBoundValueMap* MapA, const FBoundValueMap* MapB, const FBoundValueMap* PerValueWeightMapB, float DefaultWeightB, FBoundValueMap*& OutputMap, FReallocFun OutputAllocator);
+			static void TransformUnboundValueMap(const FUnboundValueMap* MapA, const FUnboundValueMap* MapB, float WeightB, FUnboundValueMap*& OutputMap, FReallocFun OutputAllocator);
+		};
+
+		struct FInterpolate_QuaternionAttribute
+		{
+			using TransformerType = FInterpolate;
+			using ValueType = FQuaternionAnimationAttribute;
+
+			static void TransformBoundValueMap(const FBoundValueMap* MapA, const FBoundValueMap* MapB, const FBoundValueMap* PerValueWeightMapB, float DefaultWeightB, FBoundValueMap*& OutputMap, FReallocFun OutputAllocator);
+			static void TransformUnboundValueMap(const FUnboundValueMap* MapA, const FUnboundValueMap* MapB, float WeightB, FUnboundValueMap*& OutputMap, FReallocFun OutputAllocator);
+		};
 	}
 }
 
diff --git a/Engine/Plugins/Experimental/UAF/UAF/Source/UAF/Public/UAF/ValueRuntime/Transformers/Layer.h b/Engine/Plugins/Experimental/UAF/UAF/Source/UAF/Public/UAF/ValueRuntime/Transformers/Layer.h
--- a/Engine/Plugins/Experimental/UAF/UAF/Source/UAF/Public/UAF/ValueRuntime/Transformers/Layer.h
+++ b/Engine/Plugins/Experimental/UAF/UAF/Source/UAF/Public/UAF/ValueRuntime/Transformers/Layer.h
@@ -70,6 +70,42 @@ namespace UE::UAF::Transformers
 			static void TransformBoundValueMap(const FBoundValueMap* Base, const FBoundValueMap* Layer, const FBoundValueMap* PerValueWeights, float DefaultLayerWeight, FBoundValueMap*& OutputMap, FReallocFun OutputAllocator);
 			static void TransformUnboundValueMap(const FUnboundValueMap* Base, const FUnboundValueMap* Layer, float LayerWeight, FUnboundValueMap*& OutputMap, FReallocFun OutputAllocator);
 		};
+
+		struct FLayer_IntegerAttribute
+		{
+			using TransformerType = FLayer;
+			using ValueType = FIntegerAnimationAttribute;
+
+			static void TransformBoundValueMap(const FBoundValueMap* Base, const FBoundValueMap* Layer, const FBoundValueMap* PerValueWeights, float DefaultLayerWeight, FBoundValueMap*& OutputMap, FReallocFun OutputAllocator);
+			static void TransformUnboundValueMap(const FUnboundValueMap* Base, const FUnboundValueMap* Layer, float LayerWeight, FUnboundValueMap*& OutputMap, FReallocFun OutputAllocator);
+		};
+
+		struct FLayer_TransformAttribute
+		{
+			using TransformerType = FLayer;
+			using ValueType = FTransformAnimationAttribute;
+
+			static void TransformBoundValueMap(const FBoundValueMap* Base, const FBoundValueMap* Layer, const FBoundValueMap* PerValueWeights, float DefaultLayerWeight, FBoundValueMap*& OutputMap, FReallocFun OutputAllocator);
+			static void TransformUnboundValueMap(const FUnboundValueMap* Base, const FUnboundValueMap* Layer, float LayerWeight, FUnboundValueMap*& OutputMap, FReallocFun OutputAllocator);
+		};
+
+		struct FLayer_VectorAttribute
+		{
+			using TransformerType = FLayer;
+			using ValueType = FVectorAnimationAttribute;
+
+			static void TransformBoundValueMap(const FBoundValueMap* Base, const FBoundValueMap* Layer, const FBoundValueMap* PerValueWeights, float DefaultLayerWeight, FBoundValueMap*& OutputMap, FReallocFun OutputAllocator);
+			static void TransformUnboundValueMap(const FUnboundValueMap* Base, const FUnboundValueMap* Layer, float LayerWeight, FUnboundValueMap*& OutputMap, FReallocFun OutputAllocator);
+		};
+
+		struct FLayer_QuaternionAttribute
+		{
+			using TransformerType = FLayer;
+			using ValueType = FQuaternionAnimationAttribute;
+
+			static void TransformBoundValueMap(const FBoundValueMap* Base, const FBoundValueMap* Layer, const FBoundValueMap* PerValueWeights, float DefaultLayerWeight, FBoundValueMap*& OutputMap, FReallocFun OutputAllocator);
+			static void TransformUnboundValueMap(const FUnboundValueMap* Base, const FUnboundValueMap* Layer, float LayerWeight, FUnboundValueMap*& OutputMap, FReallocFun OutputAllocator);
+		};
 	}
 }
 
diff --git a/Engine/Plugins/Experimental/UAF/UAF/Source/UAF/Public/UAF/ValueRuntime/Transformers/Overwrite.h b/Engine/Plugins/Experimental/UAF/UAF/Source/UAF/Public/UAF/ValueRuntime/Transformers/Overwrite.h
--- a/Engine/Plugins/Experimental/UAF/UAF/Source/UAF/Public/UAF/ValueRuntime/Transformers/Overwrite.h
+++ b/Engine/Plugins/Experimental/UAF/UAF/Source/UAF/Public/UAF/ValueRuntime/Transformers/Overwrite.h
@@ -59,6 +59,42 @@ namespace UE::UAF::Transformers
 			static void TransformBoundValueMap(const FBoundValueMap* InputMap, const FBoundValueMap* PerValueWeightMap, float DefaultWeight, FBoundValueMap*& OutputMap, FReallocFun OutputAllocator);
 			static void TransformUnboundValueMap(const FUnboundValueMap* InputMap, float Weight, FUnboundValueMap*& OutputMap, FReallocFun OutputAllocator);
 		};
+
+		struct FOverwrite_IntegerAttribute
+		{
+			using TransformerType = FOverwrite;
+			using ValueType = FIntegerAnimationAttribute;
+
+			static void TransformBoundValueMap(const FBoundValueMap* InputMap, const FBoundValueMap* PerValueWeightMap, float DefaultWeight, FBoundValueMap*& OutputMap, FReallocFun OutputAllocator);
+			static void TransformUnboundValueMap(const FUnboundValueMap* InputMap, float Weight, FUnboundValueMap*& OutputMap, FReallocFun OutputAllocator);
+		};
+
+		struct FOverwrite_TransformAttribute
+		{
+			using TransformerType = FOverwrite;
+			using ValueType = FTransformAnimationAttribute;
+
+			static void TransformBoundValueMap(const FBoundValueMap* InputMap, const FBoundValueMap* PerValueWeightMap, float DefaultWeight, FBoundValueMap*& OutputMap, FReallocFun OutputAllocator);
+			static void TransformUnboundValueMap(const FUnboundValueMap* InputMap, float Weight, FUnboundValueMap*& OutputMap, FReallocFun OutputAllocator);
+		};
+
+		struct FOverwrite_VectorAttribute
+		{
+			using TransformerType = FOverwrite;
+			using ValueType = FVectorAnimationAttribute;
+
+			static void TransformBoundValueMap(const FBoundValueMap* InputMap, const FBoundValueMap* PerValueWeightMap, float DefaultWeight, FBoundValueMap*& OutputMap, FReallocFun OutputAllocator);
+			static void TransformUnboundValueMap(const FUnboundValueMap* InputMap, float Weight, FUnboundValueMap*& OutputMap, FReallocFun OutputAllocator);
+		};
+
+		struct FOverwrite_QuaternionAttribute
+		{
+			using TransformerType = FOverwrite;
+			using ValueType = FQuaternionAnimationAttribute;
+
+			static void TransformBoundValueMap(const FBoundValueMap* InputMap, const FBoundValueMap* PerValueWeightMap, float DefaultWeight, FBoundValueMap*& OutputMap, FReallocFun OutputAllocator);
+			static void TransformUnboundValueMap(const FUnboundValueMap* InputMap, float Weight, FUnboundValueMap*& OutputMap, FReallocFun OutputAllocator);
+		};
 	}
 }
 
diff --git a/Engine/Plugins/Experimental/UAF/UAF/Source/UAF/Public/UAF/ValueRuntime/Transformers/Sanitize.h b/Engine/Plugins/Experimental/UAF/UAF/Source/UAF/Public/UAF/ValueRuntime/Transformers/Sanitize.h
--- a/Engine/Plugins/Experimental/UAF/UAF/Source/UAF/Public/UAF/ValueRuntime/Transformers/Sanitize.h
+++ b/Engine/Plugins/Experimental/UAF/UAF/Source/UAF/Public/UAF/ValueRuntime/Transformers/Sanitize.h
@@ -54,6 +54,24 @@ namespace UE::UAF::Transformers
 			static void TransformUnboundValueMap(FUnboundValueMap* Map);
 		};
 #endif
+
+		struct FSanitize_TransformAttribute
+		{
+			using TransformerType = FSanitize;
+			using ValueType = FTransformAnimationAttribute;
+
+			static void TransformBoundValueMap(FBoundValueMap* Map);
+			static void TransformUnboundValueMap(FUnboundValueMap* Map);
+		};
+
+		struct FSanitize_QuaternionAttribute
+		{
+			using TransformerType = FSanitize;
+			using ValueType = FQuaternionAnimationAttribute;
+
+			static void TransformBoundValueMap(FBoundValueMap* Map);
+			static void TransformUnboundValueMap(FUnboundValueMap* Map);
+		};
 	}
 }
 
