Suppress spurious GUID lookup warning when SoftVariableReference already carries type info

When constructing FAnimNextVariableReference from a FAnimNextSoftVariableReference
that has its TOptional<FAnimNextParamType> set (e.g. during rename or type-change
operations), the GUID lookup for the old variable name is expected to fail. Skip
the warning in this case since the caller already has the type context it needs.

Also updated RenameVariable to use FromNameAndType for the FindReference so the
type is carried through the soft reference during rename operations.

Fixes warnings introduced by CL 51494494 that fire during variable rename/type-change.

#Synced-CL 52098686

diff --git a/Engine/Plugins/Experimental/UAF/UAF/Source/UAF/Private/Variables/AnimNextVariableReference.cpp b/Engine/Plugins/Experimental/UAF/UAF/Source/UAF/Private/Variables/AnimNextVariableReference.cpp
--- a/Engine/Plugins/Experimental/UAF/UAF/Source/UAF/Private/Variables/AnimNextVariableReference.cpp
+++ b/Engine/Plugins/Experimental/UAF/UAF/Source/UAF/Private/Variables/AnimNextVariableReference.cpp
@@ -76,7 +76,7 @@ FAnimNextVariableReference::FAnimNextVariableReference(const FAnimNextSoftVariab
 	if (Object != nullptr)
 	{
 		CachedGuid = UE::UAF::UncookedOnly::IAnimNextUncookedOnlyModule::Get().GetVariableGuidByName(Name, Object);
-		if (!CachedGuid.IsValid())
+		if (!CachedGuid.IsValid() && !InSoftReference.HasType())
 		{
 			UE_LOG(LogAnimation, Warning, TEXT("Unable to retrieve Variable Guid while attempting to construct FAnimNextVariableReference for %s in %s, variable might not exist on object."), *Name.ToString(), *Object->GetPathName());
 		}
diff --git a/Engine/Plugins/Experimental/UAF/UAF/Source/UAFUncookedOnly/Private/UncookedOnlyUtils.cpp b/Engine/Plugins/Experimental/UAF/UAF/Source/UAFUncookedOnly/Private/UncookedOnlyUtils.cpp
--- a/Engine/Plugins/Experimental/UAF/UAF/Source/UAFUncookedOnly/Private/UncookedOnlyUtils.cpp
+++ b/Engine/Plugins/Experimental/UAF/UAF/Source/UAFUncookedOnly/Private/UncookedOnlyUtils.cpp
@@ -1974,7 +1974,7 @@ void FUtils::RenameVariable(UAnimNextVariableEntry* VariableEntry, const FName N
 
 		FCompilationScope CompilerResults(LOCTEXT("ModifiedAssets_RenameVariable", "Modified Assets Rename Variable"), OuterRigVMAsset);
 		
-		const FAnimNextSoftVariableReference FindReference = FAnimNextSoftVariableReference::FromName(VariableEntry->GetEntryName(), OuterRigVMAsset);
+		const FAnimNextSoftVariableReference FindReference = FAnimNextSoftVariableReference::FromNameAndType(VariableEntry->GetEntryName(), OuterRigVMAsset, VariableEntry->GetType());
 		
 		// Rename variable first
 		VariableEntry->SetEntryName(NewName);
