Expose FUtils::RenameVariable to Python scripting via UUAFRigVMAssetLibrary

#misc AnimNext >> UAF
#Synced-CL 52098686

diff --git a/Engine/Plugins/Experimental/UAF/UAF/Source/UAFUncookedOnly/Internal/AnimNextRigVMAssetEditorData.h b/Engine/Plugins/Experimental/UAF/UAF/Source/UAFUncookedOnly/Internal/AnimNextRigVMAssetEditorData.h
--- a/Engine/Plugins/Experimental/UAF/UAF/Source/UAFUncookedOnly/Internal/AnimNextRigVMAssetEditorData.h
+++ b/Engine/Plugins/Experimental/UAF/UAF/Source/UAFUncookedOnly/Internal/AnimNextRigVMAssetEditorData.h
@@ -128,49 +128,53 @@ public:
 	UFUNCTION(BlueprintCallable, Category = "UAF|Entries", meta=(ScriptMethod))
 	static UAFUNCOOKEDONLY_API UUAFRigVMAssetEntry* FindEntry(UUAFRigVMAsset* InAsset, FName InName);
 
-	/** Removes an entry from an AnimNext asset */
+	/** Removes an entry from an UAF asset */
 	UFUNCTION(BlueprintCallable, Category = "UAF|Entries", meta=(ScriptMethod))
 	static UAFUNCOOKEDONLY_API bool RemoveEntry(UUAFRigVMAsset* InAsset, UUAFRigVMAssetEntry* InEntry, bool bSetupUndoRedo = true, bool bPrintPythonCommand = true);
 
-	/** Removes multiple entries from an AnimNext asset */
+	/** Removes multiple entries from an UAF asset */
 	UFUNCTION(BlueprintCallable, Category = "UAF|Entries", meta=(ScriptMethod))
 	static UAFUNCOOKEDONLY_API bool RemoveEntries(UUAFRigVMAsset* InAsset, const TArray<UUAFRigVMAssetEntry*>& InEntries, bool bSetupUndoRedo = true, bool bPrintPythonCommand = true);
 
-	/** Removes all entries from an AnimNext asset */
+	/** Removes all entries from an UAF asset */
 	UFUNCTION(BlueprintCallable, Category = "UAF|Entries", meta=(ScriptMethod))
 	static UAFUNCOOKEDONLY_API bool RemoveAllEntries(UUAFRigVMAsset* InAsset, bool bSetupUndoRedo = true, bool bPrintPythonCommand = true);
 
-	/** Adds a parameter to an AnimNext asset */
+	/** Adds a parameter to an UAF asset */
 	UFUNCTION(BlueprintCallable, Category = "UAF|Entries", meta=(ScriptMethod))
 	static UAFUNCOOKEDONLY_API UAnimNextVariableEntry* AddVariable(UUAFRigVMAsset* InAsset, FName InName, EPropertyBagPropertyType InValueType, EPropertyBagContainerType InContainerType = EPropertyBagContainerType::None, const UObject* InValueTypeObject = nullptr, const FString& InDefaultValue = TEXT(""), bool bSetupUndoRedo = true, bool bPrintPythonCommand = true);
 
-	/** Adds an event graph to an AnimNext asset */
+	/** Adds an event graph to an UAF asset */
 	UFUNCTION(BlueprintCallable, Category = "UAF|Entries", meta=(ScriptMethod))
 	static UAFUNCOOKEDONLY_API UAnimNextEventGraphEntry* AddEventGraph(UUAFRigVMAsset* InAsset, FName InName, UScriptStruct* InEventStruct, bool bSetupUndoRedo = true, bool bPrintPythonCommand = true);
 
-	/** Adds shared variables to an AnimNext asset */
+	/** Adds shared variables to an UAF asset */
 	UFUNCTION(BlueprintCallable, Category = "UAF|Entries", meta=(ScriptMethod))
 	static UAFUNCOOKEDONLY_API UUAFSharedVariablesEntry* AddSharedVariables(UUAFRigVMAsset* InAsset, UUAFSharedVariables* InSharedVariables, bool bSetupUndoRedo = true, bool bPrintPythonCommand = true);
 
-	/** Adds shared variables struct to an AnimNext asset */
+	/** Adds shared variables struct to an UAF asset */
 	UFUNCTION(BlueprintCallable, Category = "UAF|Entries", meta=(ScriptMethod))
 	static UAFUNCOOKEDONLY_API UUAFSharedVariablesEntry* AddSharedVariablesStruct(UUAFRigVMAsset* InAsset, UScriptStruct* InStruct, bool bSetupUndoRedo = true, bool bPrintPythonCommand = true);
 	
-	/** Adds a function to an AnimNext asset */
+	/** Adds a function to an UAF asset */
 	UFUNCTION(BlueprintCallable, Category = "UAF|Entries", meta=(ScriptMethod))
 	static UAFUNCOOKEDONLY_API URigVMLibraryNode* AddFunction(UUAFRigVMAsset* InAsset, FName InFunctionName, bool bInMutable, bool bSetupUndoRedo = true, bool bPrintPythonCommand = true);
 
-	/** Adds a category to an AnimNext asset */
+	/** Adds a category to an UAF asset */
 	UFUNCTION(BlueprintCallable, Category = "UAF|Entries", meta=(ScriptMethod))
 	static UAFUNCOOKEDONLY_API bool AddCategory(UUAFRigVMAsset* InAsset, const FString& CategoryName, bool bSetupUndoRedo = true, bool bPrintPythonCommand = true);
 
-	/** Renames a category within an AnimNext asset */
+	/** Renames a category within an UAF asset */
 	UFUNCTION(BlueprintCallable, Category = "UAF|Entries", meta=(ScriptMethod))
 	static UAFUNCOOKEDONLY_API bool RenameCategory(UUAFRigVMAsset* InAsset, const FString& CategoryName, const FString& NewCategoryName, bool bSetupUndoRedo = true, bool bPrintPythonCommand = true);
 	
-	/** Removes a category from within an AnimNext asset */
+	/** Removes a category from within an UAF asset */
 	UFUNCTION(BlueprintCallable, Category = "UAF|Entries", meta=(ScriptMethod))
 	static UAFUNCOOKEDONLY_API bool RemoveCategory(UUAFRigVMAsset* InAsset, const FString& CategoryName, bool bSetupUndoRedo = true, bool bPrintPythonCommand = true);
+
+	/** Renames a variable in an UAF asset, updating all references across the project */
+	UFUNCTION(BlueprintCallable, Category = "UAF|Entries", meta=(ScriptMethod))
+	static UAFUNCOOKEDONLY_API bool RenameVariable(UUAFRigVMAsset* InAsset, FName InVariableName, FName InNewName, bool bSetupUndoRedo = true, bool bPrintPythonCommand = true);
 };
 
 /* Base class for all AnimNext editor data objects that use RigVM */
diff --git a/Engine/Plugins/Experimental/UAF/UAF/Source/UAFUncookedOnly/Private/AnimNextRigVMAssetEditorData.cpp b/Engine/Plugins/Experimental/UAF/UAF/Source/UAFUncookedOnly/Private/AnimNextRigVMAssetEditorData.cpp
--- a/Engine/Plugins/Experimental/UAF/UAF/Source/UAFUncookedOnly/Private/AnimNextRigVMAssetEditorData.cpp
+++ b/Engine/Plugins/Experimental/UAF/UAF/Source/UAFUncookedOnly/Private/AnimNextRigVMAssetEditorData.cpp
@@ -2773,6 +2773,19 @@ bool UUAFRigVMAssetLibrary::RemoveCategory(UUAFRigVMAsset* InAsset, const FStrin
 	return UE::UAF::UncookedOnly::FUtils::GetEditorData(InAsset)->RemoveCategory(CategoryName, bSetupUndoRedo, bPrintPythonCommand);
 }
 
+bool UUAFRigVMAssetLibrary::RenameVariable(UUAFRigVMAsset* InAsset, FName InVariableName, FName InNewName, bool bSetupUndoRedo, bool bPrintPythonCommand)
+{
+	UUAFRigVMAssetEditorData* EditorData = UE::UAF::UncookedOnly::FUtils::GetEditorData(InAsset);
+	if (UAnimNextVariableEntry* VariableEntry = Cast<UAnimNextVariableEntry>(EditorData->FindEntry(InVariableName)))
+	{
+		UE::UAF::UncookedOnly::FUtils::RenameVariable(VariableEntry, InNewName, bSetupUndoRedo, bPrintPythonCommand);
+		return true;
+	}
+
+	UUAFRigVMAssetEditorData::ReportError(*FString::Printf(TEXT("UUAFRigVMAssetLibrary::RenameVariable: Variable '%s' not found."), *InVariableName.ToString()));
+	return false;
+}
+
 URigVMLibraryNode* UUAFRigVMAssetEditorData::AddFunction(FName InFunctionName, bool bInMutable, bool bInSetupUndoRedo, bool bPrintPythonCommand)
 {
 	URigVMController* Controller = RigVMClient.GetOrCreateController(GetLocalFunctionLibrary());
