UAF: Shared FRemapPoseDataPool and CopyFrom helpers (pre-requisites for lazy cross-skeleton remapping)

Extracts remap pose data caching from individual nodes into a shared, lock-efficient pool. Groundwork for the lazy cross-skeleton pose remapping change.

FAnimNextGraphLODPose::CopyFrom / CopyPose: Static helper + convenience overloads that transparently remap bones and attributes when source and target skeletons differ.

The goal is to eliminate per-node remap data caching and centralize it in a shared pool that is efficient when using MT, no lock contention on the read path, minimal locking on writes.

#Synced-CL 51755844

diff --git a/Engine/Plugins/Experimental/UAF/UAF/Source/UAF/Internal/Graph/AnimNext_LODPose.h b/Engine/Plugins/Experimental/UAF/UAF/Source/UAF/Internal/Graph/AnimNext_LODPose.h
--- a/Engine/Plugins/Experimental/UAF/UAF/Source/UAF/Internal/Graph/AnimNext_LODPose.h
+++ b/Engine/Plugins/Experimental/UAF/UAF/Source/UAF/Internal/Graph/AnimNext_LODPose.h
@@ -10,6 +10,8 @@
 
 #include "AnimNext_LODPose.generated.h"
 
+#define UE_API UAF_API
+
 USTRUCT(BlueprintType, meta = (DisplayName = "ReferencePose"))
 struct FAnimNextGraphReferencePose
 {
@@ -39,4 +41,20 @@ struct FAnimNextGraphLODPose
 	// Attributes
 	// Note that attribute bone indices are LOD bone indices matching the LOD pose
 	UE::Anim::FHeapAttributeContainer Attributes;
+
+	/**
+	 * Copy source pose into target pose, automatically remapping bones and attributes if skeletons differ.
+	 * Curves are always copied directly (no skeleton dependency).
+	 */
+	static UE_API void CopyPose(
+		const UE::UAF::FLODPoseHeap& InSourcePose, const FBlendedHeapCurve& InSourceCurves, const UE::Anim::FHeapAttributeContainer& InSourceAttributes,
+		UE::UAF::FLODPoseHeap& OutTargetPose, FBlendedHeapCurve& OutTargetCurves, UE::Anim::FHeapAttributeContainer& OutTargetAttributes);
+
+	/** Copy from individual source components, automatically remapping bones and attributes via FRemapPoseDataPool if skeletons differ. */
+	UE_API void CopyFrom(const UE::UAF::FLODPoseHeap& InPose, const FBlendedHeapCurve& InCurves, const UE::Anim::FHeapAttributeContainer& InAttributes);
+
+	/** Copy from source, automatically remapping bones and attributes via FRemapPoseDataPool if skeletons differ. */
+	UE_API void CopyFrom(const FAnimNextGraphLODPose& Source);
 };
+
+#undef UE_API
diff --git a/Engine/Plugins/Experimental/UAF/UAF/Source/UAF/Internal/Graph/UAFSystemOutputComponent.h b/Engine/Plugins/Experimental/UAF/UAF/Source/UAF/Internal/Graph/UAFSystemOutputComponent.h
--- a/Engine/Plugins/Experimental/UAF/UAF/Source/UAF/Internal/Graph/UAFSystemOutputComponent.h
+++ b/Engine/Plugins/Experimental/UAF/UAF/Source/UAF/Internal/Graph/UAFSystemOutputComponent.h
@@ -3,7 +3,6 @@
 #pragma once
 
 #include "AnimNext_LODPose.h"
-#include "RemapPoseData.h"
 #include "Module/UAFModuleInstanceComponent.h"
 
 #include "UAFSystemOutputComponent.generated.h"
@@ -122,10 +121,6 @@ private:
 	UPROPERTY(Transient)
 	FAnimNextGraphReferencePose RefPose;
 
-	// Remapping data used to support on-demand remaps
-	UPROPERTY(Transient)
-	FRemapPoseData RemapPoseData;
-
 	// Adapter to use for output
 	UE::UAF::ISystemOutputAdapter* SystemOutputAdapter = nullptr;
 
diff --git a/Engine/Plugins/Experimental/UAF/UAF/Source/UAF/Private/DataRegistry.cpp b/Engine/Plugins/Experimental/UAF/UAF/Source/UAF/Private/DataRegistry.cpp
--- a/Engine/Plugins/Experimental/UAF/UAF/Source/UAF/Private/DataRegistry.cpp
+++ b/Engine/Plugins/Experimental/UAF/UAF/Source/UAF/Private/DataRegistry.cpp
@@ -3,6 +3,7 @@
 #include "DataRegistry.h"
 
 #include "Misc/CoreDelegates.h"
+#include "RemapPoseDataPool.h"
 #include "Misc/ScopeRWLock.h"
 #include "DataRegistryTypes.h"
 #include "GenerationTools.h"
@@ -92,6 +93,9 @@ FDataRegistry* FDataRegistry::Get()
 {
 	SCOPED_NAMED_EVENT(UAF_DataRegistry_PostGC, FColor::Blue);
 
+	// Clean up stale entries in the remap pose data pool (game thread instance)
+	FRemapPoseDataPool::Get().GarbageCollect();
+
 	// Compact the registry on GC
 	if (GAnimationDataRegistry)
 	{
diff --git a/Engine/Plugins/Experimental/UAF/UAF/Source/UAF/Private/Graph/AnimNext_LODPose.cpp b/Engine/Plugins/Experimental/UAF/UAF/Source/UAF/Private/Graph/AnimNext_LODPose.cpp
new file mode 100644
--- /dev/null
+++ b/Engine/Plugins/Experimental/UAF/UAF/Source/UAF/Private/Graph/AnimNext_LODPose.cpp
@@ -0,0 +1,37 @@
+// Copyright Epic Games, Inc. All Rights Reserved.
+
+#include "Graph/AnimNext_LODPose.h"
+#include "RemapPoseDataPool.h"
+#include "ReferencePose.h"
+
+void FAnimNextGraphLODPose::CopyPose(
+	const UE::UAF::FLODPoseHeap& InSourcePose, const FBlendedHeapCurve& InSourceCurves, const UE::Anim::FHeapAttributeContainer& InSourceAttributes,
+	UE::UAF::FLODPoseHeap& OutTargetPose, FBlendedHeapCurve& OutTargetCurves, UE::Anim::FHeapAttributeContainer& OutTargetAttributes)
+{
+	const UE::UAF::FReferencePose* SourceRef = InSourcePose.RefPose;
+	const UE::UAF::FReferencePose* TargetRef = OutTargetPose.RefPose;
+
+	if (SourceRef != nullptr && TargetRef != nullptr && SourceRef != TargetRef)
+	{
+		const FRemapPoseData& RemapData = UE::UAF::FRemapPoseDataPool::Get().GetRemapData(*SourceRef, *TargetRef);
+		RemapData.RemapPose(InSourcePose, OutTargetPose);
+		RemapData.RemapAttributes(InSourcePose, InSourceAttributes, OutTargetPose, OutTargetAttributes);
+	}
+	else
+	{
+		OutTargetPose.CopyFrom(InSourcePose);
+		OutTargetAttributes.CopyFrom(InSourceAttributes);
+	}
+
+	OutTargetCurves.CopyFrom(InSourceCurves);
+}
+
+void FAnimNextGraphLODPose::CopyFrom(const UE::UAF::FLODPoseHeap& InPose, const FBlendedHeapCurve& InCurves, const UE::Anim::FHeapAttributeContainer& InAttributes)
+{
+	CopyPose(InPose, InCurves, InAttributes, LODPose, Curves, Attributes);
+}
+
+void FAnimNextGraphLODPose::CopyFrom(const FAnimNextGraphLODPose& Source)
+{
+	CopyFrom(Source.LODPose, Source.Curves, Source.Attributes);
+}
diff --git a/Engine/Plugins/Experimental/UAF/UAF/Source/UAF/Private/Graph/UAFSystemOutputComponent.cpp b/Engine/Plugins/Experimental/UAF/UAF/Source/UAF/Private/Graph/UAFSystemOutputComponent.cpp
--- a/Engine/Plugins/Experimental/UAF/UAF/Source/UAF/Private/Graph/UAFSystemOutputComponent.cpp
+++ b/Engine/Plugins/Experimental/UAF/UAF/Source/UAF/Private/Graph/UAFSystemOutputComponent.cpp
@@ -152,19 +152,13 @@ void FUAFSystemOutputComponent::RemapLODPoseToOutput(const FAnimNextGraphLODPose
 	{
 		if (InputRefPose != OutputRefPose)
 		{
-			// Different reference pose, so we must lazily remap here
-			// Pre-create the mapping and cache it to avoid runtime lookups.
-			if (RemapPoseData.ShouldReinit(*InputRefPose, *OutputRefPose))
-			{
-				RemapPoseData.Reinit(*InputRefPose, *OutputRefPose);
-			}
-
-			InOutputPose.Curves.CopyFrom(InInputPose.Curves);
-			RemapPoseData.RemapPose(InInputPose.LODPose, InOutputPose.LODPose);
-			RemapPoseData.RemapAttributes(InInputPose.LODPose, InInputPose.Attributes, InOutputPose.LODPose, InOutputPose.Attributes);
+			// Cross-skeleton: prepare the output with the target reference pose, then copy/remap via shared pool
+			InOutputPose.LODPose.PrepareForLOD(*OutputRefPose, OutputRefPose->GetSourceLODLevel(), true, InInputPose.LODPose.IsAdditive());
+			InOutputPose.CopyFrom(InInputPose);
 		}
 		else
 		{
+			// Same skeleton: direct copy, no remap needed
 			InOutputPose = InInputPose;
 		}
 	}
diff --git a/Engine/Plugins/Experimental/UAF/UAF/Source/UAF/Private/RemapPoseData.cpp b/Engine/Plugins/Experimental/UAF/UAF/Source/UAF/Private/RemapPoseData.cpp
--- a/Engine/Plugins/Experimental/UAF/UAF/Source/UAF/Private/RemapPoseData.cpp
+++ b/Engine/Plugins/Experimental/UAF/UAF/Source/UAF/Private/RemapPoseData.cpp
@@ -204,7 +204,7 @@ void FRemapPoseData::RemapPose(const UE::UAF::FLODPoseHeap& SourcePose,
 void FRemapPoseData::RemapAttributes(const UE::UAF::FLODPose& SourceLODPose,
 	const UE::Anim::FHeapAttributeContainer& InAttributes,
 	const UE::UAF::FLODPose& TargetLODPose,
-	UE::Anim::FHeapAttributeContainer& OutAttributes)
+	UE::Anim::FHeapAttributeContainer& OutAttributes) const
 {
 	const USkeletalMesh* SourceSkeletalMesh = SourceRefPose->SkeletalMesh.Get();
 	if (!SourceSkeletalMesh)
diff --git a/Engine/Plugins/Experimental/UAF/UAF/Source/UAF/Private/RemapPoseDataPool.cpp b/Engine/Plugins/Experimental/UAF/UAF/Source/UAF/Private/RemapPoseDataPool.cpp
new file mode 100644
--- /dev/null
+++ b/Engine/Plugins/Experimental/UAF/UAF/Source/UAF/Private/RemapPoseDataPool.cpp
@@ -0,0 +1,66 @@
+// Copyright Epic Games, Inc. All Rights Reserved.
+
+#include "RemapPoseDataPool.h"
+#include "ReferencePose.h"
+
+namespace UE::UAF
+{
+
+FRemapPoseDataPool& FRemapPoseDataPool::Get()
+{
+	static FRemapPoseDataPool Instance;
+	return Instance;
+}
+
+const FRemapPoseData& FRemapPoseDataPool::GetRemapData(const FReferencePose& InSourceRefPose, const FReferencePose& InTargetRefPose)
+{
+	const FRefPosePair Key{ &InSourceRefPose, &InTargetRefPose };
+
+	// Fast path: Read lock for lookup (concurrent, no contention).
+	{
+		FRWScopeLock ReadLock(PoolLock, SLT_ReadOnly);
+		if (TUniquePtr<FRemapPoseData>* Found = Pool.Find(Key))
+		{
+			FRemapPoseData& Entry = **Found;
+			if (!Entry.ShouldReinit(InSourceRefPose, InTargetRefPose))
+			{
+				return Entry;
+			}
+			// Needs reinit, fall through to write lock.
+		}
+	}
+
+	// Slow path: Write lock for creation or reinit.
+	// If multiple threads race here with the same key, the first one creates/reinits
+	// and subsequent ones find the existing entry with ShouldReinit() == false.
+	{
+		FRWScopeLock WriteLock(PoolLock, SLT_Write);
+		TUniquePtr<FRemapPoseData>& Entry = Pool.FindOrAdd(Key);
+		if (!Entry)
+		{
+			Entry = MakeUnique<FRemapPoseData>();
+		}
+
+		if (Entry->ShouldReinit(InSourceRefPose, InTargetRefPose))
+		{
+			Entry->Reinit(InSourceRefPose, InTargetRefPose);
+		}
+
+		return *Entry;
+	}
+}
+
+void FRemapPoseDataPool::GarbageCollect()
+{
+	FRWScopeLock WriteLock(PoolLock, SLT_Write);
+	for (auto It = Pool.CreateIterator(); It; ++It)
+	{
+		const FRefPosePair& Key = It.Key();
+		if (!Key.Source->SkeletalMesh.IsValid() || !Key.Target->SkeletalMesh.IsValid())
+		{
+			It.RemoveCurrent();
+		}
+	}
+}
+
+} // namespace UE::UAF
diff --git a/Engine/Plugins/Experimental/UAF/UAF/Source/UAF/Public/RemapPoseData.h b/Engine/Plugins/Experimental/UAF/UAF/Source/UAF/Public/RemapPoseData.h
--- a/Engine/Plugins/Experimental/UAF/UAF/Source/UAF/Public/RemapPoseData.h
+++ b/Engine/Plugins/Experimental/UAF/UAF/Source/UAF/Public/RemapPoseData.h
@@ -64,7 +64,7 @@ struct FRemapPoseData
 	void RemapAttributes(const UE::UAF::FLODPose& SourceLODPose,
 		const UE::Anim::FHeapAttributeContainer& InAttributes,
 		const UE::UAF::FLODPose& TargetLODPose,
-		UE::Anim::FHeapAttributeContainer& OutAttributes);
+		UE::Anim::FHeapAttributeContainer& OutAttributes) const;
 
 	/**
 	 * Temporary calculate modelspace transform, until the LODPose is able to do that.
diff --git a/Engine/Plugins/Experimental/UAF/UAF/Source/UAF/Public/RemapPoseDataPool.h b/Engine/Plugins/Experimental/UAF/UAF/Source/UAF/Public/RemapPoseDataPool.h
new file mode 100644
--- /dev/null
+++ b/Engine/Plugins/Experimental/UAF/UAF/Source/UAF/Public/RemapPoseDataPool.h
@@ -0,0 +1,62 @@
+// Copyright Epic Games, Inc. All Rights Reserved.
+
+#pragma once
+
+#include "RemapPoseData.h"
+#include "ReferencePose.h"
+#include "Misc/ScopeRWLock.h"
+
+#define UE_API UAF_API
+
+namespace UE::UAF
+{
+
+/**
+ * Shared pool for FRemapPoseData instances, keyed by (source, target) reference pose pairs.
+ * Avoids per-node caching and provides fast lookup with minimal locking.
+ * Uses a read-write lock: concurrent reads are lock-free, writes only lock when a new pair is first encountered or needs to be reinited.
+ */
+class FRemapPoseDataPool
+{
+public:
+	/** Access the singleton pool instance. */
+	static UE_API FRemapPoseDataPool& Get();
+
+	/**
+	 * Get or create remap data for a source->target reference pose pair.
+	 * On first access for a given pair, FRemapPoseData::Reinit() is called to build the bone mapping. Subsequent calls with the same pair return the cached data.
+	 * If either reference pose's skeletal mesh has changed since the last Reinit, the data is refreshed.
+	 * Thread safety notes: GetRemapData() returns a reference to the cached FRemapPoseData. The read lock is released before the reference is used by the caller. This is safe because:
+	 *   1. GarbageCollect() only runs from FDataRegistry::HandlePostGarbageCollect, which executes on the game thread between frames and never during animation evaluation.
+	 *   2. ShouldReinit() only returns true if the skeletal mesh pointer changed, which doesn't happen mid-frame.
+	 *   3. All animation evaluation completes within a single frame before GC can run. Therefore, returned references remain valid for the duration of their use during evaluation.
+	 */
+	UE_API const FRemapPoseData& GetRemapData(const FReferencePose& InSourceRefPose, const FReferencePose& InTargetRefPose);
+
+	/** Remove pool entries that reference stale skeletal meshes. Call periodically (e.g. on GC). */
+	UE_API void GarbageCollect();
+
+private:
+	struct FRefPosePair
+	{
+		const FReferencePose* Source;
+		const FReferencePose* Target;
+
+		bool operator==(const FRefPosePair& Other) const
+		{
+			return Source == Other.Source && Target == Other.Target;
+		}
+
+		friend uint32 GetTypeHash(const FRefPosePair& Pair)
+		{
+			return HashCombine(PointerHash(Pair.Source), PointerHash(Pair.Target));
+		}
+	};
+
+	TMap<FRefPosePair, TUniquePtr<FRemapPoseData>> Pool;
+	mutable FRWLock PoolLock;
+};
+
+} // namespace UE::UAF
+
+#undef UE_API
