UAF - Allow adding arbitrary attributes to set bindings and implement temporary hard-coded defaults for each supported attribute type

#Synced-CL 52140028

diff --git a/Engine/Plugins/Experimental/UAF/UAF/Source/UAF/Private/UAF/AbstractSkeleton/AbstractSkeletonSetBinding.cpp b/Engine/Plugins/Experimental/UAF/UAF/Source/UAF/Private/UAF/AbstractSkeleton/AbstractSkeletonSetBinding.cpp
--- a/Engine/Plugins/Experimental/UAF/UAF/Source/UAF/Private/UAF/AbstractSkeleton/AbstractSkeletonSetBinding.cpp
+++ b/Engine/Plugins/Experimental/UAF/UAF/Source/UAF/Private/UAF/AbstractSkeleton/AbstractSkeletonSetBinding.cpp
@@ -232,6 +232,14 @@ bool UAbstractSkeletonSetBinding::RemoveAllFromSet(const FName SetName)
 	return false;
 }
 
+[[nodiscard]] bool UAbstractSkeletonSetBinding::ContainsAttribute(const FName AttributeName) const
+{
+	return AttributeBindings.ContainsByPredicate([AttributeName](const FAbstractSkeleton_AttributeBinding& Binding)
+	{
+		return Binding.Attribute.GetName() == AttributeName;
+	});
+}
+
 bool UAbstractSkeletonSetBinding::IsAttributeInSet(const FAnimationAttributeIdentifier Attribute) const
 {
 	return AttributeBindings.ContainsByPredicate([&](const FAbstractSkeleton_AttributeBinding& Binding)
diff --git a/Engine/Plugins/Experimental/UAF/UAF/Source/UAF/Private/UAF/Attributes/AttributeBindingData.cpp b/Engine/Plugins/Experimental/UAF/UAF/Source/UAF/Private/UAF/Attributes/AttributeBindingData.cpp
--- a/Engine/Plugins/Experimental/UAF/UAF/Source/UAF/Private/UAF/Attributes/AttributeBindingData.cpp
+++ b/Engine/Plugins/Experimental/UAF/UAF/Source/UAF/Private/UAF/Attributes/AttributeBindingData.cpp
@@ -347,45 +347,110 @@ namespace UE::UAF
 
 				for (auto TypedSetIt = NamedSet->CreateTypedSetIterator(); TypedSetIt; ++TypedSetIt)
 				{
-					if (TypedSetIt->GetType() != FBoneTransformAnimationAttribute::StaticStruct())
+					// TODO: This should be replaced with the default values defined per-attribute in the set binding
+					// instead of hard-coding them here
+
+					if (TypedSetIt->GetType() == FTransformAnimationAttribute::StaticStruct())
 					{
-						// Only bones have a default value for now
-						continue;
+						FAttributeMappingKey MappingKey = FAttributeMappingKey::MakeFrom(TypedSetIt->GetType()).To<FTransformAnimationAttribute>();
+						TBoundValueMap<FTransformAnimationAttribute>* DefaultValues = DefaultAttributeValues->GetBoundValueMaps().Add<FTransformAnimationAttribute>(MappingKey);
+
+						const int32 NumValues = DefaultValues->Num();
+
+						for (FAttributeSetIndex ValueIndex(0); ValueIndex < NumValues; ++ValueIndex)
+						{
+							(*DefaultValues)[ValueIndex] = FTransformAnimationAttribute{ FTransform::Identity };
+						}
 					}
+					else if (TypedSetIt->GetType() == FFloatAnimationAttribute::StaticStruct())
+					{
+						FAttributeMappingKey MappingKey = FAttributeMappingKey::MakeFrom(TypedSetIt->GetType()).To<FFloatAnimationAttribute>();
+						TBoundValueMap<FFloatAnimationAttribute>* DefaultValues = DefaultAttributeValues->GetBoundValueMaps().Add<FFloatAnimationAttribute>(MappingKey);
 
-					FAttributeMappingKey MappingKey = FAttributeMappingKey::MakeFrom(TypedSetIt->GetType()).To<FBoneTransformAnimationAttribute>();
-					TBoundValueMap<FBoneTransformAnimationAttribute>* DefaultValues = DefaultAttributeValues->GetBoundValueMaps().Add<FBoneTransformAnimationAttribute>(MappingKey);
+						const int32 NumValues = DefaultValues->Num();
 
-					const int32 NumValues = DefaultValues->Num();
+						for (FAttributeSetIndex ValueIndex(0); ValueIndex < NumValues; ++ValueIndex)
+						{
+							(*DefaultValues)[ValueIndex] = FFloatAnimationAttribute{ 0.0f };
+						}
+					}
+					else if (TypedSetIt->GetType() == FIntegerAnimationAttribute::StaticStruct())
+					{
+						FAttributeMappingKey MappingKey = FAttributeMappingKey::MakeFrom(TypedSetIt->GetType()).To<FIntegerAnimationAttribute>();
+						TBoundValueMap<FIntegerAnimationAttribute>* DefaultValues = DefaultAttributeValues->GetBoundValueMaps().Add<FIntegerAnimationAttribute>(MappingKey);
+
+						const int32 NumValues = DefaultValues->Num();
+
+						for (FAttributeSetIndex ValueIndex(0); ValueIndex < NumValues; ++ValueIndex)
+						{
+							(*DefaultValues)[ValueIndex] = FIntegerAnimationAttribute{ 0 };
+						}
+					}
+					else if (TypedSetIt->GetType() == FVectorAnimationAttribute::StaticStruct())
+					{
+						FAttributeMappingKey MappingKey = FAttributeMappingKey::MakeFrom(TypedSetIt->GetType()).To<FVectorAnimationAttribute>();
+						TBoundValueMap<FVectorAnimationAttribute>* DefaultValues = DefaultAttributeValues->GetBoundValueMaps().Add<FVectorAnimationAttribute>(MappingKey);
 
-					// Set every value within our map
-					if (SkeletalMeshRefSkeleton)
+						const int32 NumValues = DefaultValues->Num();
+
+						for (FAttributeSetIndex ValueIndex(0); ValueIndex < NumValues; ++ValueIndex)
+						{
+							(*DefaultValues)[ValueIndex] = FVectorAnimationAttribute{ FVector::ZeroVector };
+						}
+					}
+					else if (TypedSetIt->GetType() == FQuaternionAnimationAttribute::StaticStruct())
 					{
+						FAttributeMappingKey MappingKey = FAttributeMappingKey::MakeFrom(TypedSetIt->GetType()).To<FQuaternionAnimationAttribute>();
+						TBoundValueMap<FQuaternionAnimationAttribute>* DefaultValues = DefaultAttributeValues->GetBoundValueMaps().Add<FQuaternionAnimationAttribute>(MappingKey);
+
+						const int32 NumValues = DefaultValues->Num();
+
 						for (FAttributeSetIndex ValueIndex(0); ValueIndex < NumValues; ++ValueIndex)
 						{
-							FAttributeBindingIndex BindingIndex = TypedSetIt->GetBindingIndex(ValueIndex);
-							FSkeletonPoseBoneIndex SkeletonBoneIndex = BindingIndex.GetSkeletonBoneIndex();
+							(*DefaultValues)[ValueIndex] = FQuaternionAnimationAttribute{ FQuat::Identity };
+						}
+					}
+					else if (TypedSetIt->GetType() == FBoneTransformAnimationAttribute::StaticStruct())
+					{
+						FAttributeMappingKey MappingKey = FAttributeMappingKey::MakeFrom(TypedSetIt->GetType()).To<FBoneTransformAnimationAttribute>();
+						TBoundValueMap<FBoneTransformAnimationAttribute>* DefaultValues = DefaultAttributeValues->GetBoundValueMaps().Add<FBoneTransformAnimationAttribute>(MappingKey);
+
+						const int32 NumValues = DefaultValues->Num();
 
-							FName BoneName = SkeletonRefSkeleton.GetBoneName(SkeletonBoneIndex.GetInt());
-							int32 MeshBoneIndex = SkeletalMeshRefSkeleton->FindBoneIndex(BoneName);
+						// Set every value within our map
+						if (SkeletalMeshRefSkeleton)
+						{
+							for (FAttributeSetIndex ValueIndex(0); ValueIndex < NumValues; ++ValueIndex)
+							{
+								FAttributeBindingIndex BindingIndex = TypedSetIt->GetBindingIndex(ValueIndex);
+								FSkeletonPoseBoneIndex SkeletonBoneIndex = BindingIndex.GetSkeletonBoneIndex();
 
-							SkeletonToMeshIndicesMap.Add(SkeletonBoneIndex, FMeshPoseBoneIndex(MeshBoneIndex));
+								FName BoneName = SkeletonRefSkeleton.GetBoneName(SkeletonBoneIndex.GetInt());
+								int32 MeshBoneIndex = SkeletalMeshRefSkeleton->FindBoneIndex(BoneName);
 
-							if (MeshBoneIndex != INDEX_NONE)
+								SkeletonToMeshIndicesMap.Add(SkeletonBoneIndex, FMeshPoseBoneIndex(MeshBoneIndex));
+
+								if (MeshBoneIndex != INDEX_NONE)
+								{
+									(*DefaultValues)[ValueIndex] = FBoneTransformAnimationAttribute{ RefPose[MeshBoneIndex] };
+								}
+							}
+						}
+						else
+						{
+							for (FAttributeSetIndex ValueIndex(0); ValueIndex < NumValues; ++ValueIndex)
 							{
-								(*DefaultValues)[ValueIndex] = FBoneTransformAnimationAttribute{ RefPose[MeshBoneIndex] };
+								FAttributeBindingIndex BindingIndex = TypedSetIt->GetBindingIndex(ValueIndex);
+								FSkeletonPoseBoneIndex SkeletonBoneIndex = BindingIndex.GetSkeletonBoneIndex();
+
+								(*DefaultValues)[ValueIndex] = FBoneTransformAnimationAttribute{ RefPose[SkeletonBoneIndex.GetInt()] };
 							}
 						}
+
 					}
 					else
 					{
-						for (FAttributeSetIndex ValueIndex(0); ValueIndex < NumValues; ++ValueIndex)
-						{
-							FAttributeBindingIndex BindingIndex = TypedSetIt->GetBindingIndex(ValueIndex);
-							FSkeletonPoseBoneIndex SkeletonBoneIndex = BindingIndex.GetSkeletonBoneIndex();
-
-							(*DefaultValues)[ValueIndex] = FBoneTransformAnimationAttribute{ RefPose[SkeletonBoneIndex.GetInt()] };
-						}
+						check(false); // Using a bound attribute without a default value being defined, accessing these attributes won't work
 					}
 
 					// Our set map has a reference to a typed set
diff --git a/Engine/Plugins/Experimental/UAF/UAF/Source/UAF/Public/UAF/AbstractSkeleton/AbstractSkeletonSetBinding.h b/Engine/Plugins/Experimental/UAF/UAF/Source/UAF/Public/UAF/AbstractSkeleton/AbstractSkeletonSetBinding.h
--- a/Engine/Plugins/Experimental/UAF/UAF/Source/UAF/Public/UAF/AbstractSkeleton/AbstractSkeletonSetBinding.h
+++ b/Engine/Plugins/Experimental/UAF/UAF/Source/UAF/Public/UAF/AbstractSkeleton/AbstractSkeletonSetBinding.h
@@ -78,11 +78,20 @@ public:
 
 	UE_API bool AddAttributeToSet(const FAnimationAttributeIdentifier Attribute, const FName SetName);
 
+	// Note this function does two different things depending on the type of attribute provided.
+	// If the attribute is bound then it is unbound from its set. This does not remove it from the
+	// set binding and instead just assigns it to the NANE_None set.
+	// If the attribute is unbound then it is completely removed from the set binding.
+	// TODO: This function needs to be changed to avoid confusion
 	UE_API bool RemoveAttributeFromSet(const FAnimationAttributeIdentifier Attribute);
 
 	// Unbinds all attributes bound to a set. Returns whether or not any attributes were removed.
 	UE_API bool RemoveAllFromSet(const FName SetName);
 	
+	// Returns if an attribute with the given name is currently stored in this set binding.
+	// Note that an attribute existing does not mean it is bound to a set.
+	[[nodiscard]] UE_API bool ContainsAttribute(const FName AttributeName) const;
+	
 	// Returns true if the provided attribute is in ANY set
 	UE_API bool IsAttributeInSet(const FAnimationAttributeIdentifier Attribute) const;
 
diff --git a/Engine/Plugins/Experimental/UAF/UAF/Source/UAFEditor/Private/UAF/AbstractSkeleton/Sets/SAttributesList.cpp b/Engine/Plugins/Experimental/UAF/UAF/Source/UAFEditor/Private/UAF/AbstractSkeleton/Sets/SAttributesList.cpp
--- a/Engine/Plugins/Experimental/UAF/UAF/Source/UAFEditor/Private/UAF/AbstractSkeleton/Sets/SAttributesList.cpp
+++ b/Engine/Plugins/Experimental/UAF/UAF/Source/UAFEditor/Private/UAF/AbstractSkeleton/Sets/SAttributesList.cpp
@@ -7,19 +7,22 @@
 #include "Framework/Application/SlateApplication.h"
 #include "Framework/MultiBox/MultiBoxBuilder.h"
 #include "PersonaModule.h"
+#include "PropertyCustomizationHelpers.h"
 #include "SPositiveActionButton.h"
 #include "ScopedTransaction.h"
 #include "UAF/AbstractSkeleton/Sets/SetDragDrop.h"
 #include "UAF/AbstractSkeleton/AbstractSkeletonSetBinding.h"
 #include "ToolMenus.h"
 #include "Widgets/Input/SSearchBox.h"
+#include "Widgets/Text/SInlineEditableTextBlock.h"
 
 #define LOCTEXT_NAMESPACE "UE::UAF::Editor::SAttributesList"
 
 using namespace UE::UAF::Editor;
 
 static const FName ColumnId_Set("Set");
-static const FName ColumnId_Attribute("Attribute");
+static const FName ColumnId_AttributeName("AttributeName");
+static const FName ColumnId_AttributeType("AttributeType");
 
 static const FName MenuId_AddAttribute("SAttributesList.AddAttribute");
 
@@ -93,9 +96,12 @@ void SAttributesList::Construct(const FArguments& InArgs, TWeakObjectPtr<UAbstra
 						+ SHeaderRow::Column(ColumnId_Set)
 						.DefaultLabel(LOCTEXT("SetColumnLabel", "Set"))
 						.FillWidth(0.3f)
-						+ SHeaderRow::Column(ColumnId_Attribute)
-						.DefaultLabel(LOCTEXT("AttributeColumnLabel", "Attribute"))
-						.FillWidth(0.7f)
+						+ SHeaderRow::Column(ColumnId_AttributeName)
+						.DefaultLabel(LOCTEXT("AttributeNameColumnLabel", "Attribute Name"))
+						.FillWidth(0.35f)
+						+ SHeaderRow::Column(ColumnId_AttributeType)
+						.DefaultLabel(LOCTEXT("AttributeTypeColumnLabel", "Attribute Type"))
+						.FillWidth(0.35f)
 					)
 			]
 	];
@@ -141,16 +147,21 @@ void SAttributesList::RepopulateListData()
 class SAttributesListRow : public SMultiColumnTableRow<SAttributesList::FListItemPtr>
 {
 public:
+	using FOnAttributeTypeChanged = TDelegate<void(TNonNullPtr<const UScriptStruct>)>;
+	
 	SLATE_BEGIN_ARGS(SAttributesListRow) {}
 		SLATE_ARGUMENT(SAttributesList::FListItemPtr, Item)
 		SLATE_EVENT(FOnCanAcceptDrop, OnCanAcceptDrop)
 		SLATE_EVENT(FOnAcceptDrop, OnAcceptDrop)
 		SLATE_EVENT(FOnDragDetected, OnDragDetected)
+		SLATE_EVENT(FOnAttributeTypeChanged, OnAttributeTypeChanged)
 	SLATE_END_ARGS()
 
-	void Construct(const FArguments& InArgs, const TSharedRef<STableViewBase>& OwnerTableView)
+	void Construct(const FArguments& InArgs, const TSharedRef<STableViewBase>& OwnerTableView, TNonNullPtr<UAbstractSkeletonSetBinding> InSetBinding)
 	{
 		Item = InArgs._Item;
+		OnAttributeTypeChanged = InArgs._OnAttributeTypeChanged;
+		SetBinding = InSetBinding;
 
 		const FSuperRowType::FArguments Args = FSuperRowType::FArguments()
 			.OnDragDetected(InArgs._OnDragDetected)
@@ -197,7 +208,7 @@ public:
 						.Text(FText::FromName(Item->SetName))
 				];
 		}
-		else if (ColumnName == ColumnId_Attribute)
+		else if (ColumnName == ColumnId_AttributeName)
 		{
 			return SNew(SHorizontalBox)
 				+ SHorizontalBox::Slot()
@@ -207,15 +218,80 @@ public:
 				[
 					SNew(SImage)
 						.Image(FAppStyle::Get().GetBrush("AnimGraph.Attribute.Attributes.Icon"))
-						.ColorAndOpacity((Item->SetName != NAME_None) ? FLinearColor(FColor(31, 228, 75)) : FSlateColor::UseForeground())
 				]
 				+ SHorizontalBox::Slot()
 				.AutoWidth()
 				.VAlign(VAlign_Center)
 				.Padding(2, 2)
 				[
-					SNew(STextBlock)
-						.Text(FText::FromName(Item->Attribute.GetName()))
+					SAssignNew(Item->EditableTextBlock, SInlineEditableTextBlock)
+					.Text_Lambda([this]()
+						{
+							return FText::FromName(Item->Attribute.GetName());
+						})
+					.OnVerifyTextChanged_Lambda([this](const FText& Label, FText& ErrorMessage)
+						{
+							FName NewAttributeName = FName(Label.ToString());
+							
+							if (NewAttributeName == Item->Attribute.GetName())
+							{
+								return true; // Name unchanged
+							}
+
+							if (NewAttributeName.IsNone())
+							{
+								ErrorMessage = LOCTEXT("EmptyAttributeNameErrorMessage", "An attribute name must be non-empty");
+								return false;
+							}
+							
+							if (SetBinding->ContainsAttribute(NewAttributeName))
+							{
+								ErrorMessage = LOCTEXT("RenameClashErrorMessage", "An attribute with that name already exists");
+								return false;
+							}
+							
+							return true;
+						})
+					.OnTextCommitted_Lambda([this](const FText& InText, ETextCommit::Type CommitInfo)
+						{
+							FName NewAttributeName = FName(InText.ToString());
+
+							// The double call to RemoveAttributeFromSet here is intentional.
+							// If the attribute is bound then the first removal unbinds it.
+							// The second removal will remove the unbound attribute.
+
+							// TODO: This API is poor and needs changing to be clearer.
+							
+							if (!Item->SetName.IsNone())
+							{
+								SetBinding->RemoveAttributeFromSet(Item->Attribute);
+							}
+							SetBinding->RemoveAttributeFromSet(Item->Attribute);
+
+							FAnimationAttributeIdentifier NewIdentifier(NewAttributeName, INDEX_NONE, NAME_None, Item->Attribute.GetType());
+				
+							SetBinding->AddAttributeToSet(NewIdentifier, Item->SetName);
+						})
+				];
+		}
+		else if (ColumnName == ColumnId_AttributeType)
+		{
+			return SNew(SHorizontalBox)
+				+ SHorizontalBox::Slot()
+				.AutoWidth()
+				.VAlign(VAlign_Center)
+				.Padding(2, 2)
+				[
+					SNew(SStructPropertyEntryBox)
+						.AllowNone(false)
+						.SelectedStruct_Lambda([this]()
+							{
+								return Item->Attribute.GetType();
+							})
+						.OnSetStruct_Lambda([this](const UScriptStruct* InNewAttributeType)
+							{
+								OnAttributeTypeChanged.ExecuteIfBound(InNewAttributeType);
+							})
 				];
 		}
 		else 
@@ -226,6 +302,8 @@ public:
 
 private:
 	SAttributesList::FListItemPtr Item;
+	UAbstractSkeletonSetBinding* SetBinding = nullptr;
+	FOnAttributeTypeChanged OnAttributeTypeChanged;
 };
 
 SAttributesList::FListItem::FListItem(const FName InSetName, const FAnimationAttributeIdentifier& InAttribute)
@@ -236,7 +314,7 @@ SAttributesList::FListItem::FListItem(const FName InSetName, const FAnimationAtt
 
 TSharedRef<ITableRow> SAttributesList::ListView_OnGenerateRow(FListItemPtr Item, const TSharedRef<STableViewBase>& OwnerTable)
 {
-	return SNew(SAttributesListRow, OwnerTable)
+	return SNew(SAttributesListRow, OwnerTable, SetBinding.Get())
 		.Item(Item)
 		.OnCanAcceptDrop_Lambda([](const FDragDropEvent& DragDropEvent, EItemDropZone DropZone, FListItemPtr TargetTreeItem) -> TOptional<EItemDropZone>
 			{
@@ -287,6 +365,26 @@ TSharedRef<ITableRow> SAttributesList::ListView_OnGenerateRow(FListItemPtr Item,
 				}
 
 				return FReply::Unhandled();
+			})
+		.OnAttributeTypeChanged_Lambda([this, Item](TNonNullPtr<const UScriptStruct> InNewAttributeType)
+			{
+				// The double call to RemoveAttributeFromSet here is intentional.
+				// If the attribute is bound then the first removal unbinds it.
+				// The second removal will remove the unbound attribute.
+
+				// TODO: This API is poor and needs changing to be clearer.
+				
+				if (!Item->SetName.IsNone())
+				{
+					SetBinding->RemoveAttributeFromSet(Item->Attribute);
+				}
+				SetBinding->RemoveAttributeFromSet(Item->Attribute);
+
+				// TODO: Will be removing FAnimationAttributeIdentifier here in favour of just (AttributeName, AttributeType)
+				// So this const cast will go away
+				FAnimationAttributeIdentifier NewIdentifier(Item->Attribute.GetName(), INDEX_NONE, NAME_None, const_cast<UScriptStruct*>(InNewAttributeType.Get()));
+				
+				SetBinding->AddAttributeToSet(NewIdentifier, Item->SetName);
 			});
 }
 
@@ -301,6 +399,18 @@ TSharedPtr<SWidget> SAttributesList::ListView_OnContextMenuOpening()
 	FMenuBuilder MenuBuilder(true, nullptr);
 	FListItemPtr Selected = Selection[0]; // Single selection only
 
+	MenuBuilder.AddMenuEntry(
+			LOCTEXT("RenameAttribute_Label", "Rename"),
+			LOCTEXT("RenameAttribute_Tooltip", "Renames the selected attribute"),
+			FSlateIcon(FAppStyle::GetAppStyleSetName(), "Icons.Rename"),
+			FUIAction(
+				FExecuteAction::CreateLambda([this, Selected]()
+					{
+						Selected->EditableTextBlock.Pin()->EnterEditingMode();
+					}
+				)
+			));
+	
 	if (Selected->SetName.IsNone())
 	{
 		MenuBuilder.AddMenuEntry(
@@ -375,6 +485,34 @@ void SAttributesList::RegisterMenus()
 	{
 		FToolMenuSection& AttributeWidgetSection = AddAttributeMenu->AddSection("WidgetSection");
 
+		AttributeWidgetSection.AddDynamicEntry(
+			"NewAttribute",
+			FNewToolMenuSectionDelegate::CreateLambda([](FToolMenuSection& Section)
+			{
+				if (const UAttributesListMenuContext* Context = Section.FindContext<UAttributesListMenuContext>())
+				{
+					TSharedPtr<SAttributesList> SetBindingWidget = Context->SetBindingWidget.Pin();
+					ensure(SetBindingWidget);
+			
+					const TWeakObjectPtr<UAbstractSkeletonSetBinding> SetBinding = SetBindingWidget->SetBinding;
+			
+					if (SetBinding.IsValid())
+					{	
+						FToolMenuEntry Entry = FToolMenuEntry::InitMenuEntry(
+							"AddNewAttribute",
+							LOCTEXT("AddNewAttribute_Label", "Add New Attribute"),
+							LOCTEXT("AddNewAttribute_Tooltip", "Adds a new attribute"),
+							FSlateIcon(),
+							FUIAction(FExecuteAction::CreateRaw(SetBindingWidget.Get(), &SAttributesList::HandleAddNewAttribute)),
+							EUserInterfaceActionType::Button
+						);
+						Section.AddEntry(Entry);
+					}
+				}
+			}));
+		
+
+		
 		AttributeWidgetSection.AddDynamicEntry(
 			"AttributePicker",
 			FNewToolMenuSectionDelegate::CreateLambda([](FToolMenuSection& Section)
@@ -465,6 +603,23 @@ void SAttributesList::RegisterMenus()
 	}
 }
 
+void SAttributesList::HandleAddNewAttribute()
+{
+	const FScopedTransaction Transaction(LOCTEXT("AddNewAttribute", "Add New Attribute"));
+
+	int32_t Suffix = 0;
+	FName NewAttributeName = "NewAttribute0";
+	while (SetBinding->ContainsAttribute(NewAttributeName))
+	{
+		++Suffix;
+		NewAttributeName = FName(FString::Format(TEXT("NewAttribute{0}"), { Suffix }));
+	}
+	
+	SetBinding->AddAttributeToSet(FAnimationAttributeIdentifier(NewAttributeName, INDEX_NONE, NAME_None, FFloatAnimationAttribute::StaticStruct()), NAME_None);
+
+	FSlateApplication::Get().DismissAllMenus();
+}
+
 void SAttributesList::HandleBindingsChanged()
 {
 	bListDirty = true;
diff --git a/Engine/Plugins/Experimental/UAF/UAF/Source/UAFEditor/Private/UAF/AbstractSkeleton/Sets/SAttributesList.h b/Engine/Plugins/Experimental/UAF/UAF/Source/UAFEditor/Private/UAF/AbstractSkeleton/Sets/SAttributesList.h
--- a/Engine/Plugins/Experimental/UAF/UAF/Source/UAFEditor/Private/UAF/AbstractSkeleton/Sets/SAttributesList.h
+++ b/Engine/Plugins/Experimental/UAF/UAF/Source/UAFEditor/Private/UAF/AbstractSkeleton/Sets/SAttributesList.h
@@ -11,6 +11,7 @@
 #include "SAttributesList.generated.h"
 
 class UAbstractSkeletonSetBinding;
+class SInlineEditableTextBlock;
 
 namespace UE::UAF::Editor
 {
@@ -32,6 +33,7 @@ namespace UE::UAF::Editor
 
 			FAnimationAttributeIdentifier Attribute;
 			FName SetName;
+			TWeakPtr<SInlineEditableTextBlock> EditableTextBlock;
 		};
 
 		using FListItemPtr = TSharedPtr<FListItem>;
@@ -48,6 +50,8 @@ namespace UE::UAF::Editor
 
 		void RegisterMenus();
 
+		void HandleAddNewAttribute();
+
 		// Called when the bindings inside the set binding asset we're editing have been modified
 		void HandleBindingsChanged();
 
