UAF: Strongly-typed bindable value system with variable and sub-property binding

Introduces FBindableValueBase and 12 concrete USTRUCT types (FBindableBool, FBindableFloat, FBindableDouble, FBindableInt32, FBindableInt64, FBindableByte, FBindableName, FBindableVector, FBindableObject, FBindableQuat, FBindableEnum, FBindableStruct) in the UAF module.

Applied to anim nodes: FUAFBlendByBoolNode uses FBindableBool, FUAFScalePlayRate uses FBindableFloat (PlayRate).

#Synced-CL 52150898

diff --git a/Engine/Plugins/Experimental/UAF/UAF/Source/UAF/Private/BindableValue/UAFBindableTypes.cpp b/Engine/Plugins/Experimental/UAF/UAF/Source/UAF/Private/BindableValue/UAFBindableTypes.cpp
new file mode 100644
--- /dev/null
+++ b/Engine/Plugins/Experimental/UAF/UAF/Source/UAF/Private/BindableValue/UAFBindableTypes.cpp
@@ -0,0 +1,710 @@
+// Copyright Epic Games, Inc. All Rights Reserved.
+
+#include "BindableValue/UAFBindableTypes.h"
+
+#include "UAFAssetInstance.h"
+#include "Param/ParamType.h"
+#include "PropertyBindingDataView.h"
+#include "PropertyBindingPath.h"
+#include "StructUtils/PropertyBag.h"
+#include "UObject/UnrealType.h"
+
+#include UE_INLINE_GENERATED_CPP_BY_NAME(UAFBindableTypes)
+
+//------------------------------------------------------------------------------
+// Private helpers
+//------------------------------------------------------------------------------
+
+namespace UE::UAF::Private
+{
+	/**
+	 * Resolves a SubProperty binding into OutMem (slow path).
+	 * TargetProp is used only for float<->double numeric coercion; pass nullptr to skip.
+	 */
+	static bool ResolveSubPropertyToMem(
+		const FUAFPropertyBinding& Binding,
+		const FProperty*           TargetProp,
+		void*                      OutMem,
+		FUAFAssetInstance&         Instance)
+	{
+		const FProperty*       SourceProp      = Binding.SourceVariable.ResolveProperty();
+		const FStructProperty* SourceStructProp = SourceProp ? CastField<FStructProperty>(SourceProp) : nullptr;
+		if (!SourceStructProp || !SourceStructProp->Struct)
+		{
+			return false;
+		}
+
+		const int32 StructSize = SourceStructProp->Struct->GetStructureSize();
+		TArray<uint8, TInlineAllocator<128>> StructBuf;
+		StructBuf.SetNumUninitialized(StructSize);
+		SourceStructProp->Struct->InitializeStruct(StructBuf.GetData());
+
+		const FAnimNextParamType StructParamType(
+			EPropertyBagPropertyType::Struct,
+			EPropertyBagContainerType::None,
+			SourceStructProp->Struct);
+
+		const EPropertyBagResult GetResult = Instance.GetVariable(
+			Binding.SourceVariable,
+			StructParamType,
+			TArrayView<uint8>(StructBuf.GetData(), StructSize));
+
+		if (GetResult != EPropertyBagResult::Success)
+		{
+			SourceStructProp->Struct->DestroyStruct(StructBuf.GetData());
+			return false;
+		}
+
+		TArray<FPropertyBindingPathIndirection, TInlineAllocator<16>> Indirections;
+		const FPropertyBindingDataView DataView(SourceStructProp->Struct, StructBuf.GetData());
+
+		if (!Binding.SubPropertyPath.ResolveIndirectionsWithValue(DataView, Indirections) || Indirections.IsEmpty())
+		{
+			SourceStructProp->Struct->DestroyStruct(StructBuf.GetData());
+			return false;
+		}
+
+		const FPropertyBindingPathIndirection& Leaf     = Indirections.Last();
+		const FProperty*                       LeafProp = Leaf.GetProperty();
+		if (!LeafProp || !Leaf.GetPropertyAddress())
+		{
+			SourceStructProp->Struct->DestroyStruct(StructBuf.GetData());
+			return false;
+		}
+
+		const FNumericProperty* NumLeaf   = CastField<FNumericProperty>(LeafProp);
+		const FNumericProperty* NumTarget = TargetProp ? CastField<FNumericProperty>(TargetProp) : nullptr;
+		if (NumLeaf && NumTarget && NumLeaf->IsFloatingPoint() && NumTarget->IsFloatingPoint())
+		{
+			NumTarget->SetFloatingPointPropertyValue(
+				OutMem,
+				NumLeaf->GetFloatingPointPropertyValue(Leaf.GetPropertyAddress()));
+		}
+		else
+		{
+			LeafProp->CopySingleValue(OutMem, Leaf.GetPropertyAddress());
+		}
+
+		SourceStructProp->Struct->DestroyStruct(StructBuf.GetData());
+		return true;
+	}
+
+	static bool ResolveSubStructPropertyToMem(
+		const FUAFPropertyBinding& B,
+		const UScriptStruct*       ExpectedStruct,
+		void*                      OutMem,
+		FUAFAssetInstance&         Instance)
+	{
+		const FProperty*       SourceProp      = B.SourceVariable.ResolveProperty();
+		const FStructProperty* SourceStructProp = SourceProp ? CastField<FStructProperty>(SourceProp) : nullptr;
+		if (!SourceStructProp || !SourceStructProp->Struct)
+		{
+			return false;
+		}
+
+		const int32 StructSize = SourceStructProp->Struct->GetStructureSize();
+		TArray<uint8> StructBuf;
+		StructBuf.SetNumUninitialized(StructSize);
+		SourceStructProp->Struct->InitializeStruct(StructBuf.GetData());
+
+		const FAnimNextParamType SourceParamType(
+			EPropertyBagPropertyType::Struct, EPropertyBagContainerType::None,
+			SourceStructProp->Struct);
+
+		const EPropertyBagResult GetResult = Instance.GetVariable(
+			B.SourceVariable, SourceParamType,
+			TArrayView<uint8>(StructBuf.GetData(), StructSize));
+
+		if (GetResult != EPropertyBagResult::Success)
+		{
+			SourceStructProp->Struct->DestroyStruct(StructBuf.GetData());
+			return false;
+		}
+
+		TArray<FPropertyBindingPathIndirection, TInlineAllocator<16>> Indirections;
+		const FPropertyBindingDataView DataView(SourceStructProp->Struct, StructBuf.GetData());
+
+		if (!B.SubPropertyPath.ResolveIndirectionsWithValue(DataView, Indirections) || Indirections.IsEmpty())
+		{
+			SourceStructProp->Struct->DestroyStruct(StructBuf.GetData());
+			return false;
+		}
+
+		const FPropertyBindingPathIndirection& Leaf          = Indirections.Last();
+		const FStructProperty*                 LeafStructProp = CastField<FStructProperty>(Leaf.GetProperty());
+		if (!LeafStructProp || !LeafStructProp->Struct || !Leaf.GetPropertyAddress())
+		{
+			SourceStructProp->Struct->DestroyStruct(StructBuf.GetData());
+			return false;
+		}
+
+		if (!ensureMsgf(LeafStructProp->Struct == ExpectedStruct,
+			TEXT("FBindableStruct SubProperty type mismatch: expected %s, leaf is %s"),
+			*ExpectedStruct->GetName(), *LeafStructProp->Struct->GetName()))
+		{
+			SourceStructProp->Struct->DestroyStruct(StructBuf.GetData());
+			return false;
+		}
+
+		LeafStructProp->CopySingleValue(OutMem, Leaf.GetPropertyAddress());
+		SourceStructProp->Struct->DestroyStruct(StructBuf.GetData());
+		return true;
+	}
+
+	static bool ResolveVariableToMem(
+		const FAnimNextVariableReference& SourceVariable,
+		const FAnimNextParamType&          TargetParamType,
+		void*                              OutMem,
+		int32                              OutMemSize,
+		FUAFAssetInstance&                 Instance)
+	{
+		const EPropertyBagResult Result = Instance.GetVariable(
+			SourceVariable,
+			TargetParamType,
+			TArrayView<uint8>(static_cast<uint8*>(OutMem), OutMemSize));
+		return Result == EPropertyBagResult::Success;
+	}
+
+	/**
+	 * Shared resolution logic for all simple FBindableXxx types.
+	 * Handles Variable and SubProperty binding modes, falling back to ConstantValue.
+	 * Scalar types (bool, float, etc.) pass TargetProp for float<->double coercion;
+	 * struct types (FVector, FQuat) pass nullptr.
+	 */
+	template<typename T, typename BindableType>
+	static T ResolveBindableValue(const BindableType& Bindable, FUAFAssetInstance* Instance)
+	{
+		const auto ConstantValue = Bindable.GetConstantValue();
+
+		if (!Bindable.HasBinding() || !Instance)
+		{
+			return ConstantValue;
+		}
+		const FUAFPropertyBinding* Binding = Bindable.GetBinding();
+		if (!Binding || Binding->SourceVariable.IsNone())
+		{
+			return ConstantValue;
+		}
+
+		if (Binding->SourceType == EUAFBindingSourceType::Variable)
+		{
+			T Out = ConstantValue;
+			Instance->GetVariable(Binding->SourceVariable, Out);
+			return Out;
+		}
+		if (Binding->SourceType == EUAFBindingSourceType::SubProperty)
+		{
+			const FProperty* TargetProp = nullptr;
+			if constexpr (TIsArithmetic<T>::Value)
+			{
+				static const FProperty* CachedProp = BindableType::StaticStruct()->FindPropertyByName(GET_MEMBER_NAME_CHECKED(BindableType, ConstantValue));
+				TargetProp = CachedProp;
+			}
+			T Out = ConstantValue;
+			ResolveSubPropertyToMem(*Binding, TargetProp, &Out, *Instance);
+			return Out;
+		}
+		return ConstantValue;
+	}
+
+} // namespace UE::UAF::Private
+
+//------------------------------------------------------------------------------
+// Simple bindable type GetValue implementations
+//------------------------------------------------------------------------------
+
+bool FBindableBool::GetValue(FUAFAssetInstance* Instance) const
+{
+	return UE::UAF::Private::ResolveBindableValue<bool>(*this, Instance);
+}
+
+float FBindableFloat::GetValue(FUAFAssetInstance* Instance) const
+{
+	return UE::UAF::Private::ResolveBindableValue<float>(*this, Instance);
+}
+
+double FBindableDouble::GetValue(FUAFAssetInstance* Instance) const
+{
+	return UE::UAF::Private::ResolveBindableValue<double>(*this, Instance);
+}
+
+int32 FBindableInt32::GetValue(FUAFAssetInstance* Instance) const
+{
+	return UE::UAF::Private::ResolveBindableValue<int32>(*this, Instance);
+}
+
+int64 FBindableInt64::GetValue(FUAFAssetInstance* Instance) const
+{
+	return UE::UAF::Private::ResolveBindableValue<int64>(*this, Instance);
+}
+
+uint8 FBindableByte::GetValue(FUAFAssetInstance* Instance) const
+{
+	return UE::UAF::Private::ResolveBindableValue<uint8>(*this, Instance);
+}
+
+FName FBindableName::GetValue(FUAFAssetInstance* Instance) const
+{
+	return UE::UAF::Private::ResolveBindableValue<FName>(*this, Instance);
+}
+
+FVector FBindableVector::GetValue(FUAFAssetInstance* Instance) const
+{
+	return UE::UAF::Private::ResolveBindableValue<FVector>(*this, Instance);
+}
+
+FQuat FBindableQuat::GetValue(FUAFAssetInstance* Instance) const
+{
+	return UE::UAF::Private::ResolveBindableValue<FQuat>(*this, Instance);
+}
+
+//------------------------------------------------------------------------------
+// FBindableEnum
+//------------------------------------------------------------------------------
+
+int32 FBindableEnum::GetValue(FUAFAssetInstance* Instance) const
+{
+	if (!Binding.IsValid() || !Instance)
+	{
+		return ConstantValue;
+	}
+	const FUAFPropertyBinding& B = *Binding;
+	if (B.SourceVariable.IsNone())
+	{
+		return ConstantValue;
+	}
+
+	if (B.SourceType == EUAFBindingSourceType::Variable)
+	{
+		const FAnimNextParamType ParamType = GetEnumClass()
+			? FAnimNextParamType(EPropertyBagPropertyType::Enum, EPropertyBagContainerType::None, GetEnumClass())
+			: FAnimNextParamType(EPropertyBagPropertyType::Byte);
+
+		// Determine the source enum's actual underlying byte size.
+		// FByteProperty (uint8 enum) has ElementSize 1; FEnumProperty wraps a numeric property.
+		int32 ReadSize = sizeof(int32);
+		const FProperty* SourceProp = B.SourceVariable.ResolveProperty();
+		if (const FEnumProperty* EnumProp = SourceProp ? CastField<FEnumProperty>(SourceProp) : nullptr)
+		{
+			if (!ensureMsgf(EnumProp->GetElementSize() <= sizeof(int32),
+				TEXT("FBindableEnum: enum '%s' has underlying size %u which exceeds int32; binding ignored"),
+				*EnumProp->GetName(), static_cast<uint32>(EnumProp->GetElementSize())))
+			{
+				return ConstantValue;
+			}
+			ReadSize = static_cast<int32>(EnumProp->GetElementSize());
+		}
+		else if (const FByteProperty* ByteProp = SourceProp ? CastField<FByteProperty>(SourceProp) : nullptr)
+		{
+			ReadSize = static_cast<int32>(ByteProp->GetElementSize());
+		}
+
+		int32 Out = ConstantValue;
+		UE::UAF::Private::ResolveVariableToMem(B.SourceVariable, ParamType, &Out, ReadSize, *Instance);
+		return Out;
+	}
+	return ConstantValue;
+}
+
+//------------------------------------------------------------------------------
+// FBindableStruct
+//------------------------------------------------------------------------------
+
+void FBindableStruct::GetValueToMem(
+	FUAFAssetInstance*   Instance,
+	const UScriptStruct* ExpectedStruct,
+	void*                OutMem) const
+{
+	// No binding or null instance: copy ConstantValue
+	if (!Binding.IsValid() || !Instance)
+	{
+		if (ensureMsgf(!ConstantValue.IsValid() || ConstantValue.GetScriptStruct() == ExpectedStruct,
+			TEXT("FBindableStruct::GetValue type mismatch: expected %s, ConstantValue is %s"),
+			*ExpectedStruct->GetName(),
+			ConstantValue.IsValid() ? *ConstantValue.GetScriptStruct()->GetName() : TEXT("(empty)")))
+		{
+			if (ConstantValue.IsValid())
+			{
+				ExpectedStruct->CopyScriptStruct(OutMem, ConstantValue.GetMemory());
+			}
+			else
+			{
+				ExpectedStruct->InitializeStruct(OutMem);
+			}
+		}
+		return;
+	}
+
+	const FUAFPropertyBinding& B = *Binding;
+	if (B.SourceVariable.IsNone())
+	{
+		if (ConstantValue.IsValid())
+		{
+			ExpectedStruct->CopyScriptStruct(OutMem, ConstantValue.GetMemory());
+		}
+		else
+		{
+			ExpectedStruct->InitializeStruct(OutMem);
+		}
+		return;
+	}
+
+	// Variable binding
+	if (B.SourceType == EUAFBindingSourceType::Variable)
+	{
+		const FProperty*       SourceProp      = B.SourceVariable.ResolveProperty();
+		const FStructProperty* SourceStructProp = SourceProp ? CastField<FStructProperty>(SourceProp) : nullptr;
+		if (!SourceStructProp || !SourceStructProp->Struct)
+		{
+			return;
+		}
+
+		ensureMsgf(SourceStructProp->Struct == ExpectedStruct,
+			TEXT("FBindableStruct::GetValue type mismatch: expected %s, source variable is %s"),
+			*ExpectedStruct->GetName(), *SourceStructProp->Struct->GetName());
+
+		const FAnimNextParamType ParamType(
+			EPropertyBagPropertyType::Struct, EPropertyBagContainerType::None,
+			SourceStructProp->Struct);
+
+		Instance->GetVariable(
+			B.SourceVariable, ParamType,
+			TArrayView<uint8>(static_cast<uint8*>(OutMem), ExpectedStruct->GetStructureSize()));
+		return;
+	}
+
+	// SubProperty binding
+	if (B.SourceType == EUAFBindingSourceType::SubProperty)
+	{
+		UE::UAF::Private::ResolveSubStructPropertyToMem(B, ExpectedStruct, OutMem, *Instance);
+	}
+}
+
+//------------------------------------------------------------------------------
+// FBindableObject
+//------------------------------------------------------------------------------
+
+UObject* FBindableObject::GetValue(FUAFAssetInstance* Instance) const
+{
+	if (!HasBinding() || !Instance)
+	{
+		return ConstantValue;
+	}
+	const FUAFPropertyBinding* OptBinding = GetBinding();
+	if (!OptBinding || OptBinding->SourceVariable.IsNone())
+	{
+		return ConstantValue;
+	}
+
+	if (OptBinding->SourceType == EUAFBindingSourceType::Variable)
+	{
+		const UClass* ResolvedClass = ObjectClass ? ObjectClass.Get() : UObject::StaticClass();
+		const FAnimNextParamType ParamType(
+			EPropertyBagPropertyType::Object,
+			EPropertyBagContainerType::None,
+			ResolvedClass);
+		TObjectPtr<UObject> Out = ConstantValue;
+		UE::UAF::Private::ResolveVariableToMem(
+			OptBinding->SourceVariable, ParamType,
+			&Out, sizeof(TObjectPtr<UObject>), *Instance);
+		return Out;
+	}
+	if (OptBinding->SourceType == EUAFBindingSourceType::SubProperty)
+	{
+		TObjectPtr<UObject> Out = ConstantValue;
+		UE::UAF::Private::ResolveSubPropertyToMem(
+			*OptBinding, nullptr, &Out, *Instance);
+		return Out;
+	}
+	return ConstantValue;
+}
+
+//------------------------------------------------------------------------------
+// FBindableValueBase — custom serialization
+//------------------------------------------------------------------------------
+
+bool FBindableValueBase::Serialize(FArchive& Ar)
+{
+	bool bHasBinding = Binding.IsValid();
+	Ar << bHasBinding;
+
+	if (bHasBinding)
+	{
+		if (Ar.IsLoading())
+		{
+			Binding = MakeUnique<FUAFPropertyBinding>();
+		}
+		FUAFPropertyBinding::StaticStruct()->SerializeItem(Ar, Binding.Get(), nullptr);
+	}
+	else if (Ar.IsLoading())
+	{
+		Binding.Reset();
+	}
+
+	return false; // returning false tells UE to continue with default property serialization for remaining members
+}
+
+//------------------------------------------------------------------------------
+// FBindableValueBase — text export/import helpers
+//------------------------------------------------------------------------------
+
+static const TCHAR* BindingPropertyName = TEXT("__Binding");
+
+void FBindableValueBase::ExportBindingText(FString& ValueStr, const FUAFPropertyBinding* InBinding, int32 PortFlags)
+{
+	if (!InBinding)
+	{
+		return;
+	}
+
+	FString BindingText;
+	FUAFPropertyBinding::StaticStruct()->ExportText(BindingText, InBinding, nullptr, nullptr, PortFlags, nullptr, true);
+	// Insert __Binding=(...) before the closing )
+	if (ValueStr.Len() > 0 && ValueStr[ValueStr.Len() - 1] == TCHAR(')'))
+	{
+		// Only add comma separator if there are existing properties between the parens (more than just "()")
+		const bool bHasExistingContent = ValueStr.Len() > 2;
+		ValueStr.InsertAt(ValueStr.Len() - 1, FString::Printf(TEXT("%s%s=%s"),
+			bHasExistingContent ? TEXT(",") : TEXT(""),
+			BindingPropertyName, *BindingText));
+	}
+}
+
+/** Creates a copy of THIS struct's text (up to matching ')') with __Binding=(...) stripped out.
+ *  The Buffer may contain text for sibling properties — only the first (...) is copied. */
+static FString RemoveBindingFromText(const TCHAR* Buffer)
+{
+	// Find the end of this struct's text (matching closing paren)
+	int32 StructLen = 0;
+	if (*Buffer == TCHAR('('))
+	{
+		int32 Depth = 0;
+		while (Buffer[StructLen])
+		{
+			if (Buffer[StructLen] == TCHAR('('))
+			{
+				Depth++;
+			}
+			else if (Buffer[StructLen] == TCHAR(')'))
+			{
+				if (--Depth <= 0)
+				{
+					StructLen++;
+					break;
+				}
+			}
+			StructLen++;
+		}
+	}
+	else
+	{
+		StructLen = FCString::Strlen(Buffer);
+	}
+
+	// Copy only this struct's text
+	FString Text(StructLen, Buffer);
+	int32 BindingIdx = Text.Find(BindingPropertyName);
+	if (BindingIdx == INDEX_NONE)
+	{
+		return Text;
+	}
+
+	// Include leading comma in removal
+	int32 Start = BindingIdx;
+	if (Start > 0 && Text[Start - 1] == TCHAR(','))
+	{
+		Start--;
+	}
+
+	// Find end by matching nested parens after __Binding=
+	int32 End = BindingIdx + FCString::Strlen(BindingPropertyName);
+	if (End < Text.Len() && Text[End] == TCHAR('='))
+	{
+		End++;
+		int32 Depth = 0;
+		while (End < Text.Len())
+		{
+			const TCHAR Ch = Text[End];
+			if (Ch == TCHAR('('))
+			{
+				Depth++;
+			}
+			else if (Ch == TCHAR(')'))
+			{
+				if (--Depth <= 0)
+				{
+					End++;
+					break;
+				}
+			}
+			End++;
+		}
+	}
+	Text.RemoveAt(Start, End - Start);
+	return Text;
+}
+
+/** Advances Buffer past the outermost (...) struct text. */
+static void AdvancePastStructText(const TCHAR*& Buffer)
+{
+	int32 Depth = 0;
+	while (*Buffer)
+	{
+		if (*Buffer == TCHAR('('))
+		{
+			Depth++;
+		}
+		else if (*Buffer == TCHAR(')'))
+		{
+			if (--Depth <= 0)
+			{
+				Buffer++;
+				return;
+			}
+		}
+		Buffer++;
+	}
+}
+
+bool FBindableValueBase::ImportBindingText(const TCHAR*& Buffer, TUniquePtr<FUAFPropertyBinding>& OutBinding, int32 PortFlags)
+{
+	// Find the end of THIS struct's text (the matching ')' for the opening '(').
+	// The Buffer may contain text for subsequent sibling properties in a parent struct,
+	// so we must NOT search beyond our own closing paren — otherwise we'd pick up
+	// a __Binding from a later sibling.
+	const TCHAR* SearchEnd = Buffer;
+	if (*SearchEnd == TCHAR('('))
+	{
+		int32 Depth = 0;
+		while (*SearchEnd)
+		{
+			if (*SearchEnd == TCHAR('('))
+			{
+				Depth++;
+			}
+			else if (*SearchEnd == TCHAR(')'))
+			{
+				if (--Depth <= 0)
+				{
+					SearchEnd++; // include the closing ')'
+					break;
+				}
+			}
+			SearchEnd++;
+		}
+	}
+
+	// Search for __Binding only within this struct's text
+	const TCHAR* BindingStart = nullptr;
+	const int32 BindingNameLen = FCString::Strlen(BindingPropertyName);
+	for (const TCHAR* Scan = Buffer; Scan + BindingNameLen <= SearchEnd; Scan++)
+	{
+		if (FCString::Strncmp(Scan, BindingPropertyName, BindingNameLen) == 0)
+		{
+			BindingStart = Scan;
+			break;
+		}
+	}
+
+	if (BindingStart)
+	{
+		// Advance past "__Binding="
+		const TCHAR* ValueStart = BindingStart + BindingNameLen;
+		if (*ValueStart == TCHAR('='))
+		{
+			ValueStart++;
+			OutBinding = MakeUnique<FUAFPropertyBinding>();
+			FUAFPropertyBinding::StaticStruct()->ImportText(ValueStart,
+				OutBinding.Get(),
+				nullptr,
+				PortFlags,
+				GLog,
+				TEXT("FUAFPropertyBinding"),
+				/*bAllowNativeOverride=*/false);
+		}
+	}
+	else
+	{
+		OutBinding.Reset();
+	}
+	return true;
+}
+
+bool FBindableValueBase::ExportTextItem(FString& ValueStr, const FBindableValueBase& DefaultValue,
+	UObject* Parent, int32 PortFlags, UObject* ExportRootScope) const
+{
+	// Always compare against a proper type default rather than the passed-in DefaultValue.
+	// This avoids the self-comparison issue when callers pass Data as both value and default
+	// (e.g. FRigVMMemoryStorageStruct::GetDataAsString), which would cause all UPROPERTYs to
+	// appear unchanged and be omitted from the exported text.
+	const FBindableValueBase TypeDefault;
+	FBindableValueBase::StaticStruct()->ExportText(ValueStr,
+		this,
+		&TypeDefault,
+		Parent,
+		PortFlags,
+		ExportRootScope,
+		/*bAllowNativeOverride=*/false);
+	ExportBindingText(ValueStr, Binding.Get(), PortFlags);
+	return true;
+}
+
+bool FBindableValueBase::ImportTextItem(const TCHAR*& Buffer, int32 PortFlags,
+	UObject* Parent, FOutputDevice* ErrorText)
+{
+	ImportBindingText(Buffer, Binding, PortFlags);
+	// Import from a clean copy without __Binding to avoid "Unknown property" warnings
+	// (the RigVM pin validation treats any import warning as an error and rejects the value).
+	FString CleanText = RemoveBindingFromText(Buffer);
+	const TCHAR* CleanPtr = *CleanText;
+	const TCHAR* Result = FBindableValueBase::StaticStruct()->ImportText(CleanPtr, this, Parent, PortFlags, ErrorText,
+		TEXT("FBindableValueBase"), /*bAllowNativeOverride=*/false);
+	if (Result)
+	{
+		AdvancePastStructText(Buffer);
+	}
+	return Result != nullptr;
+}
+
+//------------------------------------------------------------------------------
+// Per-type ExportTextItem/ImportTextItem implementations
+//------------------------------------------------------------------------------
+
+#define IMPLEMENT_BINDABLE_TEXT_SERIALIZATION(StructType) \
+	bool StructType::ExportTextItem(FString& ValueStr, const StructType& DefaultValue, \
+		UObject* Parent, int32 PortFlags, UObject* ExportRootScope) const \
+	{ \
+		const StructType TypeDefault; \
+		StructType::StaticStruct()->ExportText(ValueStr, this, &TypeDefault, Parent, PortFlags, ExportRootScope, /*bAllowNativeOverride=*/false); \
+		ExportBindingText(ValueStr, Binding.Get(), PortFlags); \
+		return true; \
+	} \
+	bool StructType::ImportTextItem(const TCHAR*& Buffer, int32 PortFlags, \
+		UObject* Parent, FOutputDevice* ErrorText) \
+	{ \
+		ImportBindingText(Buffer, Binding, PortFlags); \
+		FString CleanText = RemoveBindingFromText(Buffer); \
+		const TCHAR* CleanPtr = *CleanText; \
+		const TCHAR* Result = StructType::StaticStruct()->ImportText(CleanPtr, this, Parent, PortFlags, ErrorText, \
+			TEXT(#StructType), /*bAllowNativeOverride=*/false); \
+		if (Result) \
+		{ \
+			AdvancePastStructText(Buffer); \
+		} \
+		return Result != nullptr; \
+	}
+
+IMPLEMENT_BINDABLE_TEXT_SERIALIZATION(FBindableBool)
+IMPLEMENT_BINDABLE_TEXT_SERIALIZATION(FBindableFloat)
+IMPLEMENT_BINDABLE_TEXT_SERIALIZATION(FBindableDouble)
+IMPLEMENT_BINDABLE_TEXT_SERIALIZATION(FBindableInt32)
+IMPLEMENT_BINDABLE_TEXT_SERIALIZATION(FBindableInt64)
+IMPLEMENT_BINDABLE_TEXT_SERIALIZATION(FBindableByte)
+IMPLEMENT_BINDABLE_TEXT_SERIALIZATION(FBindableName)
+IMPLEMENT_BINDABLE_TEXT_SERIALIZATION(FBindableVector)
+IMPLEMENT_BINDABLE_TEXT_SERIALIZATION(FBindableQuat)
+IMPLEMENT_BINDABLE_TEXT_SERIALIZATION(FBindableEnum)
+IMPLEMENT_BINDABLE_TEXT_SERIALIZATION(FBindableStruct)
+IMPLEMENT_BINDABLE_TEXT_SERIALIZATION(FBindableObject)
+
+#undef IMPLEMENT_BINDABLE_TEXT_SERIALIZATION
diff --git a/Engine/Plugins/Experimental/UAF/UAF/Source/UAF/Public/BindableValue/UAFBindableTypes.h b/Engine/Plugins/Experimental/UAF/UAF/Source/UAF/Public/BindableValue/UAFBindableTypes.h
new file mode 100644
--- /dev/null
+++ b/Engine/Plugins/Experimental/UAF/UAF/Source/UAF/Public/BindableValue/UAFBindableTypes.h
@@ -0,0 +1,677 @@
+// Copyright Epic Games, Inc. All Rights Reserved.
+
+#pragma once
+
+#include "CoreMinimal.h"
+#include "StructUtils/InstancedStruct.h"
+#include "BindableValue/UAFPropertyBinding.h"
+
+#include "UAFBindableTypes.generated.h"
+
+struct FUAFAssetInstance;
+
+/**
+ * Base struct for all strongly-typed bindable value wrappers.
+ *
+ * Each concrete subclass (FBindableBool, FBindableFloat, etc.) co-locates a constant
+ * fallback value with an optional runtime Binding. Call HasBinding() to check whether
+ * a binding is set.
+ *
+ * To read the resolved value at runtime, call GetValue(Instance):
+ *   - If Instance is null or no binding is set, returns ConstantValue (the fallback).
+ *   - Otherwise resolves the binding against the given FUAFAssetInstance, falling back
+ *     to ConstantValue on failure.
+ *
+ */
+USTRUCT(MinimalAPI)
+struct FBindableValueBase
+{
+	GENERATED_BODY()
+
+	FBindableValueBase() = default;
+	~FBindableValueBase() = default;
+
+	FBindableValueBase(const FBindableValueBase& Other)
+		: Binding(Other.Binding ? MakeUnique<FUAFPropertyBinding>(*Other.Binding) : nullptr)
+	{}
+
+	FBindableValueBase& operator=(const FBindableValueBase& Other)
+	{
+		if (this != &Other)
+		{
+			Binding = Other.Binding ? MakeUnique<FUAFPropertyBinding>(*Other.Binding) : nullptr;
+		}
+		return *this;
+	}
+
+	FBindableValueBase(FBindableValueBase&&) = default;
+	FBindableValueBase& operator=(FBindableValueBase&&) = default;
+
+	bool HasBinding() const { return Binding.IsValid(); }
+	const FUAFPropertyBinding* GetBinding() const
+	{
+		return Binding.Get();
+	}
+	
+	void SetBinding(const FUAFPropertyBinding& InBinding)
+	{
+		if (Binding.IsValid())
+		{
+			// Overwrite existing binding
+			*Binding.Get() = InBinding;
+		}
+		else
+		{
+			Binding = MakeUnique<FUAFPropertyBinding>(InBinding);
+		}
+	}
+
+	void SetBinding(FUAFPropertyBinding&& InBinding)
+	{
+		if (Binding.IsValid())
+		{
+			// Overwrite existing binding
+			*Binding.Get() = MoveTemp(InBinding);
+		}
+		else
+		{
+			Binding = MakeUnique<FUAFPropertyBinding>(MoveTemp(InBinding));
+		}
+	}
+
+	void ClearBinding()
+	{
+		Binding.Reset();
+	}
+
+	/**
+	 * Custom serialization is required because the Binding field is stored as a TUniquePtr<FUAFPropertyBinding>,
+	 * which is not supported by UPROPERTY serialization. Three serialization paths are needed:
+	 *   - Serialize(FArchive&): binary path (WithSerializer trait)
+	 *   - ExportTextItem/ImportTextItem: text path (WithExportTextItem/WithImportTextItem traits),
+	 *     used by RigVM text round-trip and property copy/paste
+	 * The binding is exported as a __Binding=(...) pseudo-property appended to the struct's text.
+	 * Each concrete subclass must declare its own TStructOpsTypeTraits (UE does not inherit traits).
+	 */
+	UAF_API bool Serialize(FArchive& Ar);
+	UAF_API bool ExportTextItem(FString& ValueStr, const FBindableValueBase& DefaultValue,
+		UObject* Parent, int32 PortFlags, UObject* ExportRootScope) const;
+	UAF_API bool ImportTextItem(const TCHAR*& Buffer, int32 PortFlags,
+		UObject* Parent, FOutputDevice* ErrorText);
+
+protected:
+	/** Shared helpers for text serialization — called by concrete subclasses via IMPLEMENT_BINDABLE_TEXT_SERIALIZATION. */
+	UAF_API static void ExportBindingText(FString& ValueStr, const FUAFPropertyBinding* InBinding, int32 PortFlags);
+	UAF_API static bool ImportBindingText(const TCHAR*& Buffer, TUniquePtr<FUAFPropertyBinding>& OutBinding, int32 PortFlags);
+
+	/** Optional binding. Null = no binding; non-null = GetValue() resolves from a UAF variable at runtime. */
+	TUniquePtr<FUAFPropertyBinding> Binding;
+};
+
+// Forward-declare the shared resolution helper so the friend declaration in the macro below is valid.
+namespace UE::UAF::Private
+{
+	template<typename T, typename BindableType>
+	static T ResolveBindableValue(const BindableType& Bindable, FUAFAssetInstance* Instance);
+}
+
+/**
+ * Declares ExportTextItem/ImportTextItem on a concrete FBindableXxx type.
+ * Also grants friend access to the shared ResolveBindableValue helper (uses GET_MEMBER_NAME_CHECKED on ConstantValue).
+ * Place inside the struct body. Implementation is provided by IMPLEMENT_BINDABLE_TEXT_SERIALIZATION in the .cpp.
+ */
+#define DECLARE_BINDABLE_TEXT_SERIALIZATION(StructType) \
+	template<typename T, typename BT> \
+	friend T UE::UAF::Private::ResolveBindableValue(const BT&, FUAFAssetInstance*); \
+	bool ExportTextItem(FString& ValueStr, const StructType& DefaultValue, \
+		UObject* Parent, int32 PortFlags, UObject* ExportRootScope) const; \
+	bool ImportTextItem(const TCHAR*& Buffer, int32 PortFlags, \
+		UObject* Parent, FOutputDevice* ErrorText);
+
+template<>
+struct TStructOpsTypeTraits<FBindableValueBase> : public TStructOpsTypeTraitsBase2<FBindableValueBase>
+{
+	enum
+	{
+		WithSerializer = true,
+		WithCopy = true,
+		WithExportTextItem = true,
+		WithImportTextItem = true,
+	};
+};
+
+//------------------------------------------------------------------------------
+// FBindableBool
+//------------------------------------------------------------------------------
+
+USTRUCT(BlueprintType, MinimalAPI)
+struct FBindableBool : public FBindableValueBase
+{
+	GENERATED_BODY()
+	DECLARE_BINDABLE_TEXT_SERIALIZATION(FBindableBool)
+
+	FBindableBool() = default;
+	FBindableBool(bool InValue) : ConstantValue(InValue) {}
+
+	/** Resolve the binding against Instance (may be null), falling back to ConstantValue. */
+	UAF_API bool GetValue(FUAFAssetInstance* Instance) const;
+
+	bool GetConstantValue() const { return ConstantValue; }
+	void SetConstantValue(bool InValue) { ConstantValue = InValue; }
+
+private:
+	UPROPERTY(EditAnywhere, Category = "Value")
+	bool ConstantValue = false;
+};
+
+template <>
+struct TStructOpsTypeTraits<FBindableBool> : public TStructOpsTypeTraitsBase2<FBindableBool>
+{
+	enum
+	{
+		WithSerializer = true,
+		WithCopy = true,
+		WithExportTextItem = true,
+		WithImportTextItem = true
+	};
+};
+
+//------------------------------------------------------------------------------
+// FBindableFloat
+//------------------------------------------------------------------------------
+
+USTRUCT(BlueprintType, MinimalAPI)
+struct FBindableFloat : public FBindableValueBase
+{
+	GENERATED_BODY()
+	DECLARE_BINDABLE_TEXT_SERIALIZATION(FBindableFloat)
+
+	FBindableFloat() = default;
+	FBindableFloat(float InValue) : ConstantValue(InValue) {}
+
+	UAF_API float GetValue(FUAFAssetInstance* Instance) const;
+
+	float GetConstantValue() const { return ConstantValue; }
+	void  SetConstantValue(float InValue) { ConstantValue = InValue; }
+
+private:
+	UPROPERTY(EditAnywhere, Category = "Value")
+	float ConstantValue = 0.f;
+};
+
+template<>
+struct TStructOpsTypeTraits<FBindableFloat> : public TStructOpsTypeTraitsBase2<FBindableFloat>
+{
+	enum
+	{
+		WithSerializer = true,
+		WithCopy = true,
+		WithExportTextItem = true,
+		WithImportTextItem = true,
+	};
+};
+
+//------------------------------------------------------------------------------
+// FBindableDouble
+//------------------------------------------------------------------------------
+
+USTRUCT(BlueprintType, MinimalAPI)
+struct FBindableDouble : public FBindableValueBase
+{
+	GENERATED_BODY()
+	DECLARE_BINDABLE_TEXT_SERIALIZATION(FBindableDouble)
+
+	FBindableDouble() = default;
+	FBindableDouble(double InValue) : ConstantValue(InValue) {}
+
+	UAF_API double GetValue(FUAFAssetInstance* Instance) const;
+
+	double GetConstantValue() const { return ConstantValue; }
+	void   SetConstantValue(double InValue) { ConstantValue = InValue; }
+
+private:
+	UPROPERTY(EditAnywhere, Category = "Value")
+	double ConstantValue = 0.0;
+};
+
+template<>
+struct TStructOpsTypeTraits<FBindableDouble> : public TStructOpsTypeTraitsBase2<FBindableDouble>
+{
+	enum
+	{
+		WithSerializer = true,
+		WithCopy = true,
+		WithExportTextItem = true,
+		WithImportTextItem = true,
+	};
+};
+
+//------------------------------------------------------------------------------
+// FBindableInt32
+//------------------------------------------------------------------------------
+
+USTRUCT(BlueprintType, MinimalAPI)
+struct FBindableInt32 : public FBindableValueBase
+{
+	GENERATED_BODY()
+	DECLARE_BINDABLE_TEXT_SERIALIZATION(FBindableInt32)
+
+	FBindableInt32() = default;
+	FBindableInt32(int32 InValue) : ConstantValue(InValue) {}
+
+	UAF_API int32 GetValue(FUAFAssetInstance* Instance) const;
+
+	int32 GetConstantValue() const { return ConstantValue; }
+	void  SetConstantValue(int32 InValue) { ConstantValue = InValue; }
+
+private:
+	UPROPERTY(EditAnywhere, Category = "Value")
+	int32 ConstantValue = 0;
+};
+
+template<>
+struct TStructOpsTypeTraits<FBindableInt32> : public TStructOpsTypeTraitsBase2<FBindableInt32>
+{
+	enum
+	{
+		WithSerializer = true,
+		WithCopy = true,
+		WithExportTextItem = true,
+		WithImportTextItem = true,
+	};
+};
+
+//------------------------------------------------------------------------------
+// FBindableInt64
+//------------------------------------------------------------------------------
+
+USTRUCT(BlueprintType, MinimalAPI)
+struct FBindableInt64 : public FBindableValueBase
+{
+	GENERATED_BODY()
+	DECLARE_BINDABLE_TEXT_SERIALIZATION(FBindableInt64)
+
+	FBindableInt64() = default;
+	FBindableInt64(int64 InValue) : ConstantValue(InValue) {}
+
+	UAF_API int64 GetValue(FUAFAssetInstance* Instance) const;
+
+	int64 GetConstantValue() const { return ConstantValue; }
+	void  SetConstantValue(int64 InValue) { ConstantValue = InValue; }
+
+private:
+	UPROPERTY(EditAnywhere, Category = "Value")
+	int64 ConstantValue = 0;
+};
+
+template<>
+struct TStructOpsTypeTraits<FBindableInt64> : public TStructOpsTypeTraitsBase2<FBindableInt64>
+{
+	enum
+	{
+		WithSerializer = true,
+		WithCopy = true,
+		WithExportTextItem = true,
+		WithImportTextItem = true,
+	};
+};
+
+//------------------------------------------------------------------------------
+// FBindableByte
+//------------------------------------------------------------------------------
+
+USTRUCT(BlueprintType, MinimalAPI)
+struct FBindableByte : public FBindableValueBase
+{
+	GENERATED_BODY()
+	DECLARE_BINDABLE_TEXT_SERIALIZATION(FBindableByte)
+
+	FBindableByte() = default;
+	FBindableByte(uint8 InValue) : ConstantValue(InValue) {}
+
+	UAF_API uint8 GetValue(FUAFAssetInstance* Instance) const;
+
+	uint8 GetConstantValue() const { return ConstantValue; }
+	void  SetConstantValue(uint8 InValue) { ConstantValue = InValue; }
+
+private:
+	UPROPERTY(EditAnywhere, Category = "Value")
+	uint8 ConstantValue = 0;
+};
+
+template<>
+struct TStructOpsTypeTraits<FBindableByte> : public TStructOpsTypeTraitsBase2<FBindableByte>
+{
+	enum
+	{
+		WithSerializer = true,
+		WithCopy = true,
+		WithExportTextItem = true,
+		WithImportTextItem = true,
+	};
+};
+
+//------------------------------------------------------------------------------
+// FBindableName
+//------------------------------------------------------------------------------
+
+USTRUCT(BlueprintType, MinimalAPI)
+struct FBindableName : public FBindableValueBase
+{
+	GENERATED_BODY()
+	DECLARE_BINDABLE_TEXT_SERIALIZATION(FBindableName)
+
+	FBindableName() = default;
+	FBindableName(FName InValue) : ConstantValue(InValue) {}
+
+	UAF_API FName GetValue(FUAFAssetInstance* Instance) const;
+
+	const FName& GetConstantValue() const { return ConstantValue; }
+	void         SetConstantValue(const FName& InValue) { ConstantValue = InValue; }
+
+private:
+	UPROPERTY(EditAnywhere, Category = "Value")
+	FName ConstantValue;
+};
+
+template<>
+struct TStructOpsTypeTraits<FBindableName> : public TStructOpsTypeTraitsBase2<FBindableName>
+{
+	enum
+	{
+		WithSerializer = true,
+		WithCopy = true,
+		WithExportTextItem = true,
+		WithImportTextItem = true,
+	};
+};
+
+//------------------------------------------------------------------------------
+// FBindableVector
+//------------------------------------------------------------------------------
+
+USTRUCT(BlueprintType, MinimalAPI)
+struct FBindableVector : public FBindableValueBase
+{
+	GENERATED_BODY()
+	DECLARE_BINDABLE_TEXT_SERIALIZATION(FBindableVector)
+
+	FBindableVector() = default;
+	FBindableVector(const FVector& InValue) : ConstantValue(InValue) {}
+
+	UAF_API FVector GetValue(FUAFAssetInstance* Instance) const;
+
+	const FVector& GetConstantValue() const { return ConstantValue; }
+	void           SetConstantValue(const FVector& InValue) { ConstantValue = InValue; }
+
+private:
+	UPROPERTY(EditAnywhere, Category = "Value")
+	FVector ConstantValue = FVector::ZeroVector;
+};
+
+template<>
+struct TStructOpsTypeTraits<FBindableVector> : public TStructOpsTypeTraitsBase2<FBindableVector>
+{
+	enum
+	{
+		WithSerializer = true,
+		WithCopy = true,
+		WithExportTextItem = true,
+		WithImportTextItem = true,
+	};
+};
+
+//------------------------------------------------------------------------------
+// FBindableQuat
+//------------------------------------------------------------------------------
+
+USTRUCT(BlueprintType, MinimalAPI)
+struct FBindableQuat : public FBindableValueBase
+{
+	GENERATED_BODY()
+	DECLARE_BINDABLE_TEXT_SERIALIZATION(FBindableQuat)
+
+	FBindableQuat() = default;
+	FBindableQuat(const FQuat& InValue) : ConstantValue(InValue) {}
+
+	UAF_API FQuat GetValue(FUAFAssetInstance* Instance) const;
+
+	const FQuat& GetConstantValue() const { return ConstantValue; }
+	void         SetConstantValue(const FQuat& InValue) { ConstantValue = InValue; }
+
+private:
+	UPROPERTY(EditAnywhere, Category = "Value")
+	FQuat ConstantValue = FQuat::Identity;
+};
+
+template<>
+struct TStructOpsTypeTraits<FBindableQuat> : public TStructOpsTypeTraitsBase2<FBindableQuat>
+{
+	enum
+	{
+		WithSerializer = true,
+		WithCopy = true,
+		WithExportTextItem = true,
+		WithImportTextItem = true,
+	};
+};
+
+//------------------------------------------------------------------------------
+// FBindableEnum
+//------------------------------------------------------------------------------
+
+/**
+ * Bindable enum value. Stores the raw int32 fallback alongside the UEnum class so that
+ * the editor binding picker can filter compatible enum variables.
+ *
+ * The declaring node data type should initialize EnumClass and ConstantValue in its
+ * default constructor, e.g.:
+ *   LoopMode.SetConstantValue((int32)EAnimAssetLoopMode::Auto);
+ *   LoopMode.EnumClass = StaticEnum<EAnimAssetLoopMode>();
+ *
+ * Enums whose underlying type exceeds 32 bits are not supported; a runtime ensure
+ * will fire if such an enum is bound.
+ */
+USTRUCT(BlueprintType, MinimalAPI)
+struct FBindableEnum : public FBindableValueBase
+{
+	GENERATED_BODY()
+	DECLARE_BINDABLE_TEXT_SERIALIZATION(FBindableEnum)
+
+	FBindableEnum() = default;
+	explicit FBindableEnum(UEnum* InEnumClass) : EnumClass(InEnumClass) {}
+
+	/** Constructs with EnumClass and ConstantValue both derived from the enum value. */
+	template<typename T> requires (bool(TIsUEnumClass<T>::Value))
+	explicit FBindableEnum(T InValue)
+		: EnumClass(StaticEnum<T>())
+		, ConstantValue(static_cast<int32>(InValue))
+	{
+		static_assert(sizeof(T) <= sizeof(int32),
+			"FBindableEnum only supports enums with underlying types up to 32 bits.");
+	}
+
+	UAF_API int32 GetValue(FUAFAssetInstance* Instance) const;
+
+	/** Type-safe getter — T must be a UENUM. */
+	template<typename T> requires (bool(TIsUEnumClass<T>::Value))
+	T GetValue(FUAFAssetInstance* Instance) const
+	{
+		static_assert(sizeof(T) <= sizeof(int32),
+			"FBindableEnum only supports enums with underlying types up to 32 bits.");
+		return static_cast<T>(GetValue(Instance));
+	}
+
+	int32 GetConstantValue() const { return ConstantValue; }
+	void  SetConstantValue(int32 InValue) { ConstantValue = InValue; }
+
+	/** Type-safe constant value getter — T must be a UENUM. */
+	template<typename T> requires (bool(TIsUEnumClass<T>::Value))
+	T GetConstantValue() const
+	{
+		static_assert(sizeof(T) <= sizeof(int32),
+			"FBindableEnum only supports enums with underlying types up to 32 bits.");
+		return static_cast<T>(ConstantValue);
+	}
+
+	/** Type-safe constant value setter — T must be a UENUM. */
+	template<typename T> requires (bool(TIsUEnumClass<T>::Value))
+	void SetConstantValue(T InValue)
+	{
+		static_assert(sizeof(T) <= sizeof(int32),
+			"FBindableEnum only supports enums with underlying types up to 32 bits.");
+		ConstantValue = static_cast<int32>(InValue);
+	}
+
+	UEnum* GetEnumClass() const { return EnumClass; }
+	void   SetEnumClass(UEnum* InEnumClass) { EnumClass = InEnumClass; }
+
+private:
+
+	/** Enum class for variable-picker filtering. Set by the declaring node's constructor. */
+	UPROPERTY(EditAnywhere, Category = "Value")
+	TObjectPtr<UEnum> EnumClass = nullptr;
+	
+	UPROPERTY(EditAnywhere, Category = "Value")
+	int32 ConstantValue = 0;
+};
+
+template<>
+struct TStructOpsTypeTraits<FBindableEnum> : public TStructOpsTypeTraitsBase2<FBindableEnum>
+{
+	enum
+	{
+		WithSerializer = true,
+		WithCopy = true,
+		WithExportTextItem = true,
+		WithImportTextItem = true,
+	};
+};
+
+//------------------------------------------------------------------------------
+// FBindableStruct
+//------------------------------------------------------------------------------
+
+/**
+ * Bindable struct value.
+ *
+ * Call GetValue<T>(Instance, OutValue) to resolve the binding into caller-provided storage:
+ *   - No binding (or null Instance): copies ConstantValue into OutValue,
+ *     or zero-initializes OutValue if ConstantValue is empty.
+ *   - Variable binding: reads the named variable from Instance into OutValue.
+ *   - SubProperty binding: reads the leaf sub-property from a struct variable into OutValue.
+ *
+ * T must be a USTRUCT. A runtime ensure verifies that T matches the bound or constant struct type.
+ */
+USTRUCT(BlueprintType, MinimalAPI)
+struct FBindableStruct : public FBindableValueBase
+{
+	GENERATED_BODY()
+	DECLARE_BINDABLE_TEXT_SERIALIZATION(FBindableStruct)
+
+	FBindableStruct() = default;
+	explicit FBindableStruct(UScriptStruct* InStructClass) : StructClass(InStructClass) {}
+
+	/** Constructs with StructClass and ConstantValue both derived from T. */
+	template<typename T>
+	explicit FBindableStruct(const T& InValue)
+		: StructClass(TBaseStructure<T>::Get())
+		, ConstantValue(FInstancedStruct::Make(InValue))
+	{
+	}
+
+	/**
+	 * Resolves the binding into OutValue, falling back to ConstantValue on failure.
+	 * T must be a USTRUCT. Runtime-asserts that T matches the bound or constant struct type.
+	 */
+	template<typename T>
+	void GetValue(FUAFAssetInstance* Instance, T& OutValue) const
+	{
+		GetValueToMem(Instance, TBaseStructure<T>::Get(), &OutValue);
+	}
+
+	const FInstancedStruct& GetConstantValue() const { return ConstantValue; }
+	void SetConstantValue(const FInstancedStruct& InValue) { ConstantValue = InValue; }
+	void SetConstantValue(FInstancedStruct&& InValue) { ConstantValue = MoveTemp(InValue); }
+
+	UScriptStruct* GetStructClass() const { return StructClass; }
+	void           SetStructClass(UScriptStruct* InStructClass) { StructClass = InStructClass; }
+
+private:
+	/** Non-templated implementation; called by GetValue<T>. */
+	UAF_API void GetValueToMem(FUAFAssetInstance* Instance, const UScriptStruct* ExpectedStruct, void* OutMem) const;
+
+	/** Struct class for variable-picker filtering. Set by the declaring node's constructor. */
+	UPROPERTY(EditAnywhere, Category = "Value")
+	TObjectPtr<UScriptStruct> StructClass = nullptr;
+
+	UPROPERTY(EditAnywhere, Category = "Value")
+	FInstancedStruct ConstantValue;
+};
+
+template<>
+struct TStructOpsTypeTraits<FBindableStruct> : public TStructOpsTypeTraitsBase2<FBindableStruct>
+{
+	enum
+	{
+		WithSerializer = true,
+		WithCopy = true,
+		WithExportTextItem = true,
+		WithImportTextItem = true,
+	};
+};
+
+//------------------------------------------------------------------------------
+// FBindableObject
+//------------------------------------------------------------------------------
+
+/**
+ * Bindable UObject reference. Stores a hard TObjectPtr<UObject> constant alongside
+ * an optional ObjectClass filter for the editor variable picker.
+ *
+ * The declaring node data type should set ObjectClass in its constructor to restrict
+ * which object variables are shown in the picker.
+ */
+USTRUCT(BlueprintType, MinimalAPI)
+struct FBindableObject : public FBindableValueBase
+{
+	GENERATED_BODY()
+	DECLARE_BINDABLE_TEXT_SERIALIZATION(FBindableObject)
+
+	FBindableObject() = default;
+	explicit FBindableObject(UObject* InValue) : ConstantValue(InValue) {}
+	explicit FBindableObject(UClass* InObjectClass) : ObjectClass(InObjectClass) {}
+	FBindableObject(UObject* InValue, UClass* InObjectClass)
+		: ObjectClass(InObjectClass)
+		, ConstantValue(InValue)
+	{
+	}
+
+	UAF_API UObject* GetValue(FUAFAssetInstance* Instance) const;
+
+	UObject* GetConstantValue() const { return ConstantValue; }
+	void     SetConstantValue(UObject* InValue) { ConstantValue = InValue; }
+
+	UClass* GetObjectClass() const { return ObjectClass; }
+	void    SetObjectClass(UClass* InClass) { ObjectClass = InClass; }
+
+private:
+	/** Class filter for variable-picker. Set by the declaring node's constructor. */
+	UPROPERTY(EditAnywhere, Category = "Value")
+	TObjectPtr<UClass> ObjectClass = nullptr;
+
+	UPROPERTY(EditAnywhere, Category = "Value")
+	TObjectPtr<UObject> ConstantValue = nullptr;
+};
+
+template<>
+struct TStructOpsTypeTraits<FBindableObject> : public TStructOpsTypeTraitsBase2<FBindableObject>
+{
+	enum
+	{
+		WithSerializer = true,
+		WithCopy = true,
+		WithExportTextItem = true,
+		WithImportTextItem = true,
+	};
+};
diff --git a/Engine/Plugins/Experimental/UAF/UAF/Source/UAF/Public/BindableValue/UAFPropertyBinding.h b/Engine/Plugins/Experimental/UAF/UAF/Source/UAF/Public/BindableValue/UAFPropertyBinding.h
new file mode 100644
--- /dev/null
+++ b/Engine/Plugins/Experimental/UAF/UAF/Source/UAF/Public/BindableValue/UAFPropertyBinding.h
@@ -0,0 +1,53 @@
+// Copyright Epic Games, Inc. All Rights Reserved.
+
+#pragma once
+
+#include "CoreMinimal.h"
+#include "PropertyBindingPath.h"
+#include "Variables/AnimNextVariableReference.h"
+
+#include "UAFPropertyBinding.generated.h"
+
+/**
+ * Describes how to source a bound property value.
+ */
+UENUM()
+enum class EUAFBindingSourceType : uint8
+{
+	/** Direct UAF variable (top-level, with type-match or type-promotion). */
+	Variable,
+
+	/** A dot-separated path into a struct-type UAF variable. */
+	SubProperty,
+
+	/** A 0-parameter RigVM function whose return value is used. (Reserved; not yet implemented.) */
+	Function,
+};
+
+/**
+ * Describes the source side of a property binding: where to read a value from at runtime.
+ * The target is implicit — determined by which FBindableXxx member holds this struct.
+ */
+USTRUCT()
+struct FUAFPropertyBinding
+{
+	GENERATED_BODY()
+
+	/**
+	 * The UAF variable to read from.
+	 * Used when SourceType is Variable or SubProperty.
+	 */
+	UPROPERTY(EditAnywhere, Category = "Binding")
+	FAnimNextVariableReference SourceVariable;
+
+	/**
+	 * Dot-separated path into the struct variable to the leaf property.
+	 * Used only when SourceType is SubProperty.
+	 */
+	UPROPERTY(EditAnywhere, Category = "Binding")
+	FPropertyBindingPath SubPropertyPath;
+
+	/** How to source the bound value. */
+	UPROPERTY(EditAnywhere, Category = "Binding")
+	EUAFBindingSourceType SourceType = EUAFBindingSourceType::Variable;
+};
diff --git a/Engine/Plugins/Experimental/UAF/UAF/Source/UAF/Public/UAFAssetInstance.h b/Engine/Plugins/Experimental/UAF/UAF/Source/UAF/Public/UAFAssetInstance.h
--- a/Engine/Plugins/Experimental/UAF/UAF/Source/UAF/Public/UAFAssetInstance.h
+++ b/Engine/Plugins/Experimental/UAF/UAF/Source/UAF/Public/UAFAssetInstance.h
@@ -22,6 +22,8 @@ struct FUAFScriptComponent;
 struct FAnimNextModuleInstance;
 struct FAnimNextModuleContextData;
 
+namespace UE::UAF::Tests { struct FUAFBindingTestFixture; }
+
 namespace UE::UAF
 {
 	struct FPoolHandle;
@@ -97,6 +99,17 @@ struct FUAFAssetInstance
 		return SetVariable(InVariable, FAnimNextParamType::GetType<ValueType>(), TConstArrayView<uint8>(reinterpret_cast<const uint8*>(&InNewValue), sizeof(ValueType)));
 	}
 
+	// Get a variable's value given its name, copying the value to OutResult.
+	// If the type does not match exactly then a conversion will be attempted.
+	// @param	InVariable			The variable to get the value of
+	// @param	InType				The type to read as (conversions will be attempted)
+	// @param	OutResult			Raw byte view to receive the value
+	// @return see EPropertyBagResult
+	EPropertyBagResult GetVariable(const FAnimNextVariableReference& InVariable, const FAnimNextParamType& InType, TArrayView<uint8> OutResult) const
+	{
+		return Variables.GetVariable(InVariable, InType, OutResult);
+	}
+
 	// Set a variable's value given a reference.
 	// If the type does not match exactly then a conversion will be attempted.
 	// @param	InVariable			The variable to set the value of
@@ -359,6 +372,7 @@ protected:
 	friend UE::UAF::FInstanceTaskContext;
 	friend FUAFRigVMComponent;
 	friend UE::UAF::FLayerStackTrait;
+	friend struct UE::UAF::Tests::FUAFBindingTestFixture;
 
 	FName GetDebugName() const
 	{
diff --git a/Engine/Plugins/Experimental/UAF/UAF/Source/UAF/Public/Variables/UAFInstanceVariableData.h b/Engine/Plugins/Experimental/UAF/UAF/Source/UAF/Public/Variables/UAFInstanceVariableData.h
--- a/Engine/Plugins/Experimental/UAF/UAF/Source/UAF/Public/Variables/UAFInstanceVariableData.h
+++ b/Engine/Plugins/Experimental/UAF/UAF/Source/UAF/Public/Variables/UAFInstanceVariableData.h
@@ -22,6 +22,8 @@ struct FUAFInstanceVariableContainer;
 struct FUAFRigVMComponent;
 class IRigVMRuntimeAssetInterface;
 
+namespace UE::UAF::Tests { struct FUAFBindingTestFixture; }
+
 namespace UE::UAF
 {
 template<typename TSharedData> struct TControlRigTraitBase;
@@ -51,6 +53,7 @@ private:
 	friend UE::UAF::FInstanceTaskContext;
 	friend FUAFRigVMComponent;
 	template<typename T> friend struct UE::UAF::TControlRigTraitBase;
+	friend struct UE::UAF::Tests::FUAFBindingTestFixture;
 
 	// Initialize from a host's asset and its references
 	void Initialize(const FUAFAssetInstance& InHostInstance, const TSharedPtr<const UE::UAF::FVariableOverridesCollection>& InOverrides);
diff --git a/Engine/Plugins/Experimental/UAF/UAF/Source/UAF/UAF.Build.cs b/Engine/Plugins/Experimental/UAF/UAF/Source/UAF/UAF.Build.cs
--- a/Engine/Plugins/Experimental/UAF/UAF/Source/UAF/UAF.Build.cs
+++ b/Engine/Plugins/Experimental/UAF/UAF/Source/UAF/UAF.Build.cs
@@ -16,7 +16,8 @@ namespace UnrealBuildTool.Rules
 					"RenderCore",
 					"HierarchyTableRuntime",
 					"HierarchyTableAnimationRuntime",
-					"RewindDebuggerRuntimeInterface"
+					"RewindDebuggerRuntimeInterface",
+					"PropertyBindingUtils",
 				}
 			);
 
diff --git a/Engine/Plugins/Experimental/UAF/UAF/Source/UAFEditor/Internal/Variables/VariablePickerArgs.h b/Engine/Plugins/Experimental/UAF/UAF/Source/UAFEditor/Internal/Variables/VariablePickerArgs.h
--- a/Engine/Plugins/Experimental/UAF/UAF/Source/UAFEditor/Internal/Variables/VariablePickerArgs.h
+++ b/Engine/Plugins/Experimental/UAF/UAF/Source/UAFEditor/Internal/Variables/VariablePickerArgs.h
@@ -3,11 +3,12 @@
 #pragma once
 
 #include "AnimNextExports.h"
+#include "Containers/Array.h"
 #include "Delegates/Delegate.h"
+#include "Param/ParamType.h"
 
 struct FAnimNextVariableReference;
 struct FAnimNextSoftVariableReference;
-struct FAnimNextParamType;
 
 namespace UE::UAF::Editor
 {
@@ -37,6 +38,13 @@ using FOnContextSensitivityChanged = TDelegate<void(bool)>;
 // Delegate called to get the inclusion/exclusion flags
 using FOnGetIncludeExcludeFilter = TDelegate<void(EAnimNextExportedVariableFlags& OutFlagInclusionFilter, EAnimNextExportedVariableFlags& OutFlagExclusionFilter)>;
 
+// Fired when the user picks a sub-property leaf from inside a struct variable.
+// SubPropertyPath contains the chain of FNames from the top-level variable to the leaf.
+using FOnSubPropertyVariablePicked = TDelegate<void(
+	const FAnimNextSoftVariableReference& /*TopLevelVar*/,
+	const TArray<FName>&                  /*SubPropertyPath*/,
+	const FAnimNextParamType&             /*LeafType*/)>;
+
 struct FVariablePickerArgs
 {
 	FVariablePickerArgs() = default;
@@ -50,9 +58,15 @@ struct FVariablePickerArgs
 	// Delegate called to filter variables for display to the user
 	FOnFilterVariable OnFilterVariable;
 
-	// Delegate called to filter variables by type for display to the user
+	// Delegate called to filter variables by type for display to the user (RigVM registry checks, pin constraints, etc.)
+	// Applied as an additional gate on top of CompatibleTypes when both are set.
 	FOnFilterVariableType OnFilterVariableType;
 
+	// If non-empty, the widget handles type-compat filtering and struct expansion internally.
+	// A variable passes if its type is directly compatible with any entry, or (with struct expansion)
+	// if it contains a compatible leaf. bAllowStructExpansion is auto-enabled when this is set.
+	TArray<FAnimNextParamType> CompatibleTypes;
+
 	// Delegate called to determined context sensitivity
 	FOnIsContextSensitive* OnIsContextSensitive = nullptr;
 
@@ -68,6 +82,13 @@ struct FVariablePickerArgs
 
 	// Whether the search box should be focused on widget creation
 	bool bFocusSearchWidget = true;
+
+	// Fired when the user selects a sub-property leaf from inside a struct variable.
+	FOnSubPropertyVariablePicked OnSubPropertyVariablePicked;
+
+	// If true, struct-type variables that contain at least one compatible sub-property
+	// appear in the list as expandable nodes.
+	bool bAllowStructExpansion = false;
 };
 
 }
diff --git a/Engine/Plugins/Experimental/UAF/UAF/Source/UAFEditor/Private/AnimNextEditorModule.cpp b/Engine/Plugins/Experimental/UAF/UAF/Source/UAFEditor/Private/AnimNextEditorModule.cpp
--- a/Engine/Plugins/Experimental/UAF/UAF/Source/UAFEditor/Private/AnimNextEditorModule.cpp
+++ b/Engine/Plugins/Experimental/UAF/UAF/Source/UAFEditor/Private/AnimNextEditorModule.cpp
@@ -91,6 +91,7 @@
 #include "UAF/AbstractSkeleton/Sets/Commands.h"
 #include "Variables/UAFAssetDataPropertyCustomization.h"
 #include "UAF/Wizard/AssetWizard.h"
+#include "BindableValue/UAFBindableValueCustomization.h"
 
 #define LOCTEXT_NAMESPACE "AnimNextEditorModule"
 
@@ -193,9 +194,23 @@ void FAnimNextEditorModule::StartupModule()
 		FOnGetPropertyTypeCustomizationInstance::CreateLambda([] { return MakeShared<FModuleEventPropertyCustomization>(); }),
 		ModuleEventPropertyTypeIdentifier);
 
-	PropertyModule.RegisterCustomClassLayout("AnimNextVariableEntry", 
+	PropertyModule.RegisterCustomClassLayout("AnimNextVariableEntry",
 		FOnGetDetailCustomizationInstance::CreateLambda([] { return MakeShared<FVariableCustomization>(); }));
 
+	static const FName BindableTypes[] = {
+		"BindableBool",   "BindableFloat",  "BindableDouble",
+		"BindableInt32",  "BindableInt64",  "BindableByte",
+		"BindableName",   "BindableVector", "BindableQuat",
+		"BindableEnum",   "BindableStruct", "BindableObject"
+	};
+	for (FName TypeName : BindableTypes)
+	{
+		PropertyModule.RegisterCustomPropertyTypeLayout(
+			TypeName,
+			FOnGetPropertyTypeCustomizationInstance::CreateLambda(
+				[] { return MakeShared<UE::UAF::Editor::FBindableValueCustomization>(); }));
+	}
+
 	GraphPanelPinFactory = MakeShared<FAnimNextGraphPanelPinFactoryEditor>();
 	FEdGraphUtilities::RegisterVisualPinFactory(GraphPanelPinFactory);
 
@@ -599,6 +614,17 @@ void FAnimNextEditorModule::ShutdownModule()
 		PropertyModule.UnregisterCustomPropertyTypeLayout("InstancedStruct", UAFAssetDataPropertyIdentifier);
 		PropertyModule.UnregisterCustomClassLayout("AnimNextVariableEntry");
 		PropertyModule.UnregisterCustomClassLayout("UAFRigVMAssetEditorData");
+
+		static const FName BindableTypes[] = {
+			"BindableBool",   "BindableFloat",  "BindableDouble",
+			"BindableInt32",  "BindableInt64",  "BindableByte",
+			"BindableName",   "BindableVector", "BindableQuat",
+			"BindableEnum",   "BindableStruct", "BindableObject"
+		};
+		for (FName TypeName : BindableTypes)
+		{
+			PropertyModule.UnregisterCustomPropertyTypeLayout(TypeName);
+		}
 	}
 
 	UnregisterWorkspaceDocumentTypes();
diff --git a/Engine/Plugins/Experimental/UAF/UAF/Source/UAFEditor/Private/BindableValue/UAFBindableValueCustomization.cpp b/Engine/Plugins/Experimental/UAF/UAF/Source/UAFEditor/Private/BindableValue/UAFBindableValueCustomization.cpp
new file mode 100644
--- /dev/null
+++ b/Engine/Plugins/Experimental/UAF/UAF/Source/UAFEditor/Private/BindableValue/UAFBindableValueCustomization.cpp
@@ -0,0 +1,672 @@
+// Copyright Epic Games, Inc. All Rights Reserved.
+
+#include "BindableValue/UAFBindableValueCustomization.h"
+
+#include "DetailWidgetRow.h"
+#include "PropertyHandle.h"
+#include "Framework/Application/SlateApplication.h"
+#include "Framework/Commands/UIAction.h"
+#include "Styling/AppStyle.h"
+#include "Textures/SlateIcon.h"
+#include "StructUtils/PropertyBag.h"
+#include "UObject/UnrealType.h"
+#include "IPropertyUtilities.h"
+
+#include "AnimNextRigVMAsset.h"
+#include "Variables/SVariablePickerCombo.h"
+#include "Variables/VariablePickerArgs.h"
+#include "Variables/AnimNextVariableReference.h"
+#include "Variables/AnimNextSoftVariableReference.h"
+#include "Param/ParamType.h"
+
+#include "BindableValue/UAFBindableTypes.h"
+#include "BindableValue/UAFPropertyBinding.h"
+#include "SEnumCombo.h"
+#include "IDetailChildrenBuilder.h"
+#include "InstancedStructDetails.h"
+#include "Widgets/Images/SImage.h"
+#include "Widgets/Input/SButton.h"
+#include "Widgets/Layout/SBox.h"
+#include "PropertyCustomizationHelpers.h"
+
+#define LOCTEXT_NAMESPACE "UAFBindableValueCustomization"
+
+namespace
+{
+
+/**
+ * Returns the FAnimNextParamType that corresponds to the concrete FBindableXxx subtype.
+ * For FBindableEnum, reads EnumClass from the first raw data element of InPropertyHandle.
+ */
+FAnimNextParamType GetBindableParamType(const UScriptStruct* BindableStruct, TSharedPtr<IPropertyHandle> InPropertyHandle)
+{
+	using EValueType     = EPropertyBagPropertyType;
+	using EContainerType = EPropertyBagContainerType;
+
+	if (!BindableStruct)
+	{
+		return FAnimNextParamType();
+	}
+
+	if (BindableStruct == FBindableBool::StaticStruct())
+	{
+		return FAnimNextParamType(EValueType::Bool);
+	}
+	if (BindableStruct == FBindableFloat::StaticStruct())
+	{
+		return FAnimNextParamType(EValueType::Float);
+	}
+	if (BindableStruct == FBindableDouble::StaticStruct())
+	{
+		return FAnimNextParamType(EValueType::Double);
+	}
+	if (BindableStruct == FBindableInt32::StaticStruct())
+	{
+		return FAnimNextParamType(EValueType::Int32);
+	}
+	if (BindableStruct == FBindableInt64::StaticStruct())
+	{
+		return FAnimNextParamType(EValueType::Int64);
+	}
+	if (BindableStruct == FBindableByte::StaticStruct())
+	{
+		return FAnimNextParamType(EValueType::Byte);
+	}
+	if (BindableStruct == FBindableName::StaticStruct())
+	{
+		return FAnimNextParamType(EValueType::Name);
+	}
+	if (BindableStruct == FBindableVector::StaticStruct())
+	{
+		return FAnimNextParamType(EValueType::Struct, EContainerType::None, TBaseStructure<FVector>::Get());
+	}
+	if (BindableStruct == FBindableQuat::StaticStruct())
+	{
+		return FAnimNextParamType(EValueType::Struct, EContainerType::None, TBaseStructure<FQuat>::Get());
+	}
+
+	if (BindableStruct == FBindableEnum::StaticStruct())
+	{
+		UEnum* EnumClass = nullptr;
+		if (InPropertyHandle.IsValid())
+		{
+			InPropertyHandle->EnumerateRawData(
+				[&EnumClass](void* RawData, const int32, const int32) -> bool
+				{
+					const FBindableEnum* Bindable = static_cast<const FBindableEnum*>(RawData);
+					EnumClass = Bindable->GetEnumClass();
+					return false; // first instance only
+				});
+		}
+		if (EnumClass)
+		{
+			return FAnimNextParamType(EValueType::Enum, EContainerType::None, EnumClass);
+		}
+		return FAnimNextParamType();
+	}
+
+	if (BindableStruct == FBindableStruct::StaticStruct())
+	{
+		UScriptStruct* InnerStruct = nullptr;
+		if (InPropertyHandle.IsValid())
+		{
+			InPropertyHandle->EnumerateRawData(
+				[&InnerStruct](void* RawData, const int32, const int32) -> bool
+				{
+					const FBindableStruct* Bindable = static_cast<const FBindableStruct*>(RawData);
+					// Prefer StructClass (explicit type restriction) over ConstantValue's runtime type
+					if (Bindable->GetStructClass())
+					{
+						InnerStruct = Bindable->GetStructClass();
+					}
+					else
+					{
+						InnerStruct = const_cast<UScriptStruct*>(Bindable->GetConstantValue().GetScriptStruct());
+					}
+					return false;
+				});
+		}
+		if (InnerStruct)
+		{
+			return FAnimNextParamType(EValueType::Struct, EContainerType::None, InnerStruct);
+		}
+		return FAnimNextParamType();
+	}
+
+	if (BindableStruct == FBindableObject::StaticStruct())
+	{
+		UClass* ObjClass = nullptr;
+		if (InPropertyHandle.IsValid())
+		{
+			InPropertyHandle->EnumerateRawData(
+				[&ObjClass](void* RawData, const int32, const int32) -> bool
+				{
+					ObjClass = static_cast<const FBindableObject*>(RawData)->GetObjectClass();
+					return false;
+				});
+		}
+		return FAnimNextParamType(EValueType::Object, EContainerType::None,
+			ObjClass ? ObjClass : UObject::StaticClass());
+	}
+
+	return FAnimNextParamType();
+}
+
+} // anonymous namespace
+
+namespace UE::UAF::Editor
+{
+
+void FBindableValueCustomization::CustomizeHeader(
+	TSharedRef<IPropertyHandle>      InPropertyHandle,
+	FDetailWidgetRow&                InHeaderRow,
+	IPropertyTypeCustomizationUtils& InCustomizationUtils)
+{
+	PropertyHandle = InPropertyHandle;
+
+	const FProperty*       Property    = PropertyHandle->GetProperty();
+	const FStructProperty* StructProp  = CastField<FStructProperty>(Property);
+	if (!StructProp)
+	{
+		return;
+	}
+
+	const UScriptStruct* BindableStruct = StructProp->Struct;
+
+	// For FBindableStruct: read StructClass and auto-initialize ConstantValue via the property system.
+	// We must NOT modify raw data directly — use copy+export+SetValueFromFormattedString so that
+	// undo, dirty marking, and OnPropertyChanged callbacks all fire correctly.
+	if (BindableStruct == FBindableStruct::StaticStruct())
+	{
+		bIsBindableStruct = true;
+		bool bNeedsInitialization = false;
+		PropertyHandle->EnumerateConstRawData(
+			[this, &bNeedsInitialization](const void* RawData, const int32, const int32) -> bool
+			{
+				const FBindableStruct* Bindable = static_cast<const FBindableStruct*>(RawData);
+				if (Bindable->GetStructClass())
+				{
+					ResolvedStructClass = Bindable->GetStructClass();
+					const FInstancedStruct& Current = Bindable->GetConstantValue();
+					if (!Current.IsValid() || Current.GetScriptStruct() != Bindable->GetStructClass())
+					{
+						bNeedsInitialization = true;
+					}
+				}
+				return false;
+			});
+
+		if (bNeedsInitialization)
+		{
+			FString TextValue;
+			PropertyHandle->EnumerateConstRawData(
+				[&TextValue, BindableStruct, ResolvedClass = ResolvedStructClass](const void* RawData, const int32, const int32) -> bool
+				{
+					const int32 StructSize = BindableStruct->GetStructureSize();
+
+					TArray<uint8> CopyBuf;
+					CopyBuf.SetNumUninitialized(StructSize);
+					BindableStruct->InitializeStruct(CopyBuf.GetData());
+					BindableStruct->CopyScriptStruct(CopyBuf.GetData(), RawData);
+
+					FInstancedStruct NewValue;
+					NewValue.InitializeAs(ResolvedClass);
+					static_cast<FBindableStruct*>(static_cast<void*>(CopyBuf.GetData()))->SetConstantValue(MoveTemp(NewValue));
+
+					TArray<uint8> DefaultBuf;
+					DefaultBuf.SetNumUninitialized(StructSize);
+					BindableStruct->InitializeStruct(DefaultBuf.GetData());
+
+					BindableStruct->ExportText(TextValue, CopyBuf.GetData(), DefaultBuf.GetData(), nullptr, 0, nullptr, true);
+
+					BindableStruct->DestroyStruct(DefaultBuf.GetData());
+					BindableStruct->DestroyStruct(CopyBuf.GetData());
+					return false;
+				});
+
+			PropertyHandle->SetValueFromFormattedString(TextValue);
+		}
+	}
+
+	const FAnimNextParamType TargetParamType = GetBindableParamType(BindableStruct, PropertyHandle);
+
+	// -----------------------------------------------------------------------
+	// Shared helpers (captured by value into lambdas below)
+	// -----------------------------------------------------------------------
+
+	// Returns true if the first element has an active binding.
+	auto HasBinding = [Handle = PropertyHandle]() -> bool
+	{
+		bool bResult = false;
+		Handle->EnumerateConstRawData([&bResult](const void* RawData, const int32, const int32) -> bool
+		{
+			bResult = static_cast<const FBindableValueBase*>(RawData)->HasBinding();
+			return false; // first only
+		});
+		return bResult;
+	};
+
+	// Clears the binding via copy + export + SetValueFromFormattedString.
+	// We must NOT modify raw data before SetValueFromFormattedString, because the property
+	// system compares the new text against the current in-memory state. If they match (because
+	// we already mutated raw data), no change is detected and the pin default is never updated.
+	auto ClearBinding = [Handle = PropertyHandle]()
+	{
+		const FStructProperty* StructProp = CastField<FStructProperty>(Handle->GetProperty());
+		if (!StructProp)
+		{
+			return;
+		}
+
+		FString TextValue;
+		Handle->EnumerateConstRawData(
+			[&TextValue, StructProp](const void* RawData, const int32, const int32) -> bool
+			{
+				UScriptStruct* Struct = StructProp->Struct;
+				const int32 StructSize = Struct->GetStructureSize();
+
+				TArray<uint8> CopyBuf;
+				CopyBuf.SetNumUninitialized(StructSize);
+				Struct->InitializeStruct(CopyBuf.GetData());
+				Struct->CopyScriptStruct(CopyBuf.GetData(), RawData);
+				static_cast<FBindableValueBase*>(static_cast<void*>(CopyBuf.GetData()))->ClearBinding();
+
+				TArray<uint8> DefaultBuf;
+				DefaultBuf.SetNumUninitialized(StructSize);
+				Struct->InitializeStruct(DefaultBuf.GetData());
+
+				Struct->ExportText(TextValue, CopyBuf.GetData(), DefaultBuf.GetData(), nullptr, 0, nullptr, true);
+
+				Struct->DestroyStruct(DefaultBuf.GetData());
+				Struct->DestroyStruct(CopyBuf.GetData());
+				return false;
+			});
+
+		Handle->SetValueFromFormattedString(TextValue);
+	};
+
+	// -----------------------------------------------------------------------
+	// Collect filter assets from selected objects (context sensitivity)
+	// -----------------------------------------------------------------------
+
+	{
+		TArray<TWeakObjectPtr<UObject>> Objects = InCustomizationUtils.GetPropertyUtilities()->GetSelectedObjects();
+		for (const TWeakObjectPtr<UObject>& WeakObject : Objects)
+		{
+			UObject* Object = WeakObject.Get();
+			if (!Object)
+			{
+				continue;
+			}
+
+			if (const UUAFRigVMAsset* Asset = Cast<UUAFRigVMAsset>(Object))
+			{
+				FilterAssets.AddUnique(Asset);
+			}
+			if (const UUAFRigVMAsset* OuterAsset = Object->GetTypedOuter<UUAFRigVMAsset>())
+			{
+				FilterAssets.AddUnique(OuterAsset);
+			}
+		}
+
+		TArray<TWeakObjectPtr<const UUAFRigVMAsset>> ReferencedAssets;
+		for (const TWeakObjectPtr<const UUAFRigVMAsset>& WeakAsset : FilterAssets)
+		{
+			if (const UUAFRigVMAsset* Asset = WeakAsset.Get())
+			{
+				for (const UUAFRigVMAsset* Referenced : Asset->GetReferencedVariableAssets())
+				{
+					ReferencedAssets.Add(Referenced);
+				}
+			}
+		}
+		FilterAssets.Append(ReferencedAssets);
+	}
+
+	// -----------------------------------------------------------------------
+	// Build FVariablePickerArgs
+	// -----------------------------------------------------------------------
+
+	FVariablePickerArgs PickerArgs;
+
+	if (TargetParamType.IsValid())
+	{
+		PickerArgs.CompatibleTypes = { TargetParamType };
+	}
+
+	// Helper: sets a binding by exporting from a copy and calling SetValueFromFormattedString.
+	// We must NOT modify the property handle's raw data before SetValueFromFormattedString,
+	// because the property system compares the new text against the current in-memory state.
+	// If they match (because we already mutated raw data), no change is detected and the
+	// AnimNext editor's UpdatePinDefaultValue callback never fires.
+	auto SetBindingViaText = [](TSharedPtr<IPropertyHandle> Handle, const FUAFPropertyBinding& InBinding)
+	{
+		const FStructProperty* StructProp = CastField<FStructProperty>(Handle->GetProperty());
+		if (!StructProp)
+		{
+			return;
+		}
+
+		FString TextValue;
+		Handle->EnumerateConstRawData(
+			[&TextValue, &InBinding, StructProp](const void* RawData, const int32, const int32) -> bool
+			{
+				UScriptStruct* Struct = StructProp->Struct;
+				const int32 StructSize = Struct->GetStructureSize();
+
+				// Make a copy of the current value and set the binding on the copy
+				TArray<uint8> CopyBuf;
+				CopyBuf.SetNumUninitialized(StructSize);
+				Struct->InitializeStruct(CopyBuf.GetData());
+				Struct->CopyScriptStruct(CopyBuf.GetData(), RawData);
+				static_cast<FBindableValueBase*>(static_cast<void*>(CopyBuf.GetData()))->SetBinding(InBinding);
+
+				// Export text from the copy (property handle data is untouched)
+				TArray<uint8> DefaultBuf;
+				DefaultBuf.SetNumUninitialized(StructSize);
+				Struct->InitializeStruct(DefaultBuf.GetData());
+
+				Struct->ExportText(TextValue, CopyBuf.GetData(), DefaultBuf.GetData(), nullptr, 0, nullptr, /*bAllowNativeOverride=*/true);
+
+				Struct->DestroyStruct(DefaultBuf.GetData());
+				Struct->DestroyStruct(CopyBuf.GetData());
+				return false; // first instance only
+			});
+
+		// Text differs from current state → property system detects the change → fires notification
+		Handle->SetValueFromFormattedString(TextValue);
+	};
+
+	// Write a Variable binding.
+	PickerArgs.OnVariablePicked = FOnVariablePicked::CreateLambda(
+		[this, SetBindingViaText](const FAnimNextSoftVariableReference& SoftRef, const FAnimNextParamType& /*Type*/)
+		{
+			FUAFPropertyBinding NewBinding;
+			NewBinding.SourceType     = EUAFBindingSourceType::Variable;
+			NewBinding.SourceVariable = FAnimNextVariableReference(SoftRef);
+			SetBindingViaText(PropertyHandle, NewBinding);
+		});
+
+	// Write a SubProperty binding.
+	PickerArgs.OnSubPropertyVariablePicked = FOnSubPropertyVariablePicked::CreateLambda(
+		[this, SetBindingViaText](
+			const FAnimNextSoftVariableReference& TopVarRef,
+			const TArray<FName>&                  SubPropertyPath,
+			const FAnimNextParamType& /*LeafType*/)
+		{
+			FUAFPropertyBinding NewBinding;
+			NewBinding.SourceType     = EUAFBindingSourceType::SubProperty;
+			NewBinding.SourceVariable = FAnimNextVariableReference(TopVarRef);
+			for (const FName& Seg : SubPropertyPath)
+			{
+				NewBinding.SubPropertyPath.AddPathSegment(Seg);
+			}
+			SetBindingViaText(PropertyHandle, NewBinding);
+
+			FSlateApplication::Get().DismissAllMenus();
+		});
+
+	// -----------------------------------------------------------------------
+	// Context sensitivity
+	// -----------------------------------------------------------------------
+
+	if (FilterAssets.Num() > 0)
+	{
+		OnIsContextSensitiveDelegate = FOnIsContextSensitive::CreateLambda([this]() -> bool
+		{
+			return bIsContextSensitive;
+		});
+
+		PickerArgs.OnIsContextSensitive = &OnIsContextSensitiveDelegate;
+
+		PickerArgs.OnContextSensitivityChanged = FOnContextSensitivityChanged::CreateLambda(
+			[this](bool bInIsContextSensitive)
+			{
+				bIsContextSensitive = bInIsContextSensitive;
+			});
+
+		PickerArgs.OnGetIncludeExcludeFilter = FOnGetIncludeExcludeFilter::CreateLambda(
+			[this](EAnimNextExportedVariableFlags& OutFlagInclusionFilter, EAnimNextExportedVariableFlags& OutFlagExclusionFilter)
+			{
+				if (bIsContextSensitive && FilterAssets.Num() > 0)
+				{
+					OutFlagInclusionFilter = EAnimNextExportedVariableFlags::Declared;
+					OutFlagExclusionFilter = EAnimNextExportedVariableFlags::Referenced;
+				}
+				else
+				{
+					OutFlagInclusionFilter = EAnimNextExportedVariableFlags::Declared | EAnimNextExportedVariableFlags::Public;
+					OutFlagExclusionFilter = EAnimNextExportedVariableFlags::Referenced;
+				}
+			});
+
+		PickerArgs.OnFilterVariable = FOnFilterVariable::CreateLambda(
+			[this](const FAnimNextSoftVariableReference& InVariableReference) -> EFilterVariableResult
+			{
+				if (!bIsContextSensitive || FilterAssets.Num() == 0)
+				{
+					return EFilterVariableResult::Include;
+				}
+				for (const TWeakObjectPtr<const UUAFRigVMAsset>& WeakFilterAsset : FilterAssets)
+				{
+					if (const UUAFRigVMAsset* FilterAsset = WeakFilterAsset.Get())
+					{
+						if (InVariableReference.GetSoftObjectPath() == FSoftObjectPath(FilterAsset))
+						{
+							return EFilterVariableResult::Include;
+						}
+					}
+				}
+				return EFilterVariableResult::Exclude;
+			});
+	}
+
+	// -----------------------------------------------------------------------
+	// Build row: [ConstantValue (greyed when bound)] [SVariablePickerCombo] [×]
+	// -----------------------------------------------------------------------
+
+	ConstantValueHandle = PropertyHandle->GetChildHandle(FName("ConstantValue"));
+
+	TSharedRef<SWidget> ConstantWidget = SNullWidget::NullWidget;
+	if (ConstantValueHandle.IsValid())
+	{
+		UEnum* EnumClass = nullptr;
+		if (BindableStruct == FBindableEnum::StaticStruct())
+		{
+			PropertyHandle->EnumerateRawData([&EnumClass](void* RawData, const int32, const int32) -> bool
+			{
+				EnumClass = static_cast<const FBindableEnum*>(RawData)->GetEnumClass();
+				return false;
+			});
+		}
+
+		if (EnumClass)
+		{
+			ConstantWidget = SNew(SEnumComboBox, EnumClass)
+				.CurrentValue_Lambda([Handle = ConstantValueHandle]() -> int32
+				{
+					int32 Value = 0;
+					Handle->GetValue(Value);
+					return Value;
+				})
+				.OnEnumSelectionChanged_Lambda([Handle = ConstantValueHandle](int32 NewValue, ESelectInfo::Type)
+				{
+					Handle->SetValue(NewValue);
+				});
+		}
+		else if (bIsBindableStruct && ResolvedStructClass)
+		{
+			// FBindableStruct with known type: show type name label; child properties exposed in CustomizeChildren
+			ConstantWidget = SNew(STextBlock)
+				.Text(ResolvedStructClass->GetDisplayNameText());
+		}
+		else if (BindableStruct == FBindableObject::StaticStruct())
+		{
+			UClass* ObjClass = nullptr;
+			PropertyHandle->EnumerateRawData([&ObjClass](void* RawData, const int32, const int32) -> bool
+			{
+				ObjClass = static_cast<const FBindableObject*>(RawData)->GetObjectClass();
+				return false;
+			});
+
+			ConstantWidget = SNew(SObjectPropertyEntryBox)
+				.PropertyHandle(ConstantValueHandle)
+				.AllowedClass(ObjClass ? ObjClass : UObject::StaticClass());
+		}
+		else
+		{
+			ConstantWidget = ConstantValueHandle->CreatePropertyValueWidgetWithCustomization(nullptr);
+		}
+	}
+
+	InHeaderRow
+	.NameContent()
+	[
+		PropertyHandle->CreatePropertyNameWidget()
+	]
+	.ValueContent()
+	.MinDesiredWidth(250.f)
+	[
+		SNew(SHorizontalBox)
+
+		// Constant value editor — greyed out when a binding is active.
+		+ SHorizontalBox::Slot()
+		.FillWidth(1.f)
+		.VAlign(VAlign_Center)
+		[
+			SNew(SBox)
+			.IsEnabled_Lambda([HasBinding]() -> bool
+			{
+				return !HasBinding();
+			})
+			[
+				ConstantWidget
+			]
+		]
+
+		// Variable picker combo.
+		+ SHorizontalBox::Slot()
+		.AutoWidth()
+		.VAlign(VAlign_Center)
+		.Padding(4.f, 0.f, 0.f, 0.f)
+		[
+			SNew(SVariablePickerCombo)
+			.PickerArgs(PickerArgs)
+			.OnGetVariableReference_Lambda([this]() -> FAnimNextSoftVariableReference
+			{
+				FAnimNextSoftVariableReference Result;
+				PropertyHandle->EnumerateConstRawData(
+					[&Result](const void* RawData, const int32, const int32) -> bool
+					{
+						const FBindableValueBase* Base = static_cast<const FBindableValueBase*>(RawData);
+						if (const FUAFPropertyBinding* Binding = Base->GetBinding())
+						{
+							Result = FAnimNextSoftVariableReference(Binding->SourceVariable);
+						}
+						return false; // first instance only
+					});
+				return Result;
+			})
+			.OnGetVariableType_Lambda([TargetParamType]() -> FAnimNextParamType
+			{
+				return TargetParamType;
+			})
+			// Show the full path (e.g. "TestVector.X") for sub-property bindings.
+			.VariableName_Lambda([Handle = PropertyHandle]() -> FText
+			{
+				FText Result;
+				Handle->EnumerateConstRawData(
+					[&Result](const void* RawData, const int32, const int32) -> bool
+					{
+						const FBindableValueBase* Base = static_cast<const FBindableValueBase*>(RawData);
+						if (const FUAFPropertyBinding* B = Base->GetBinding())
+						{
+							FString DisplayName = B->SourceVariable.GetName().ToString();
+							if (B->SourceType == EUAFBindingSourceType::SubProperty
+								&& B->SubPropertyPath.NumSegments() > 0)
+							{
+								DisplayName += TEXT(".") + B->SubPropertyPath.ToString();
+							}
+							Result = FText::FromString(DisplayName);
+						}
+						return false; // first instance only
+					});
+				return Result;
+			})
+		]
+
+		// Clear button — only visible when a binding is active.
+		+ SHorizontalBox::Slot()
+		.AutoWidth()
+		.VAlign(VAlign_Center)
+		.Padding(2.f, 0.f, 0.f, 0.f)
+		[
+			SNew(SButton)
+			.ButtonStyle(FAppStyle::Get(), "SimpleButton")
+			.ToolTipText(LOCTEXT("ClearBindingTooltip", "Clear Binding"))
+			.Visibility_Lambda([HasBinding]() -> EVisibility
+			{
+				return HasBinding() ? EVisibility::Visible : EVisibility::Collapsed;
+			})
+			.OnClicked_Lambda([ClearBinding]() -> FReply
+			{
+				ClearBinding();
+				return FReply::Handled();
+			})
+			[
+				SNew(SImage)
+				.Image(FAppStyle::GetBrush("Icons.X"))
+				.ColorAndOpacity(FSlateColor::UseForeground())
+			]
+		]
+	];
+
+	// Right-click "Clear Binding" context menu action on the detail row.
+	InHeaderRow.AddCustomContextMenuAction(
+		FUIAction(
+			FExecuteAction::CreateLambda([ClearBinding]() { ClearBinding(); }),
+			FCanExecuteAction::CreateLambda([HasBinding]() { return HasBinding(); })),
+		LOCTEXT("ClearBindingAction",        "Clear Binding"),
+		LOCTEXT("ClearBindingActionTooltip", "Remove the variable binding from this property"),
+		FSlateIcon(FAppStyle::GetAppStyleSetName(), "Icons.X"));
+}
+
+void FBindableValueCustomization::CustomizeChildren(
+	TSharedRef<IPropertyHandle>      InPropertyHandle,
+	IDetailChildrenBuilder&          ChildBuilder,
+	IPropertyTypeCustomizationUtils& CustomizationUtils)
+{
+	// For FBindableStruct with a known StructClass, expose the inner struct fields directly as children.
+	if (bIsBindableStruct && ResolvedStructClass && ConstantValueHandle.IsValid())
+	{
+		// Lambda to check if binding is active (greyed out when bound).
+		auto HasBinding = [Handle = PropertyHandle]() -> bool
+		{
+			bool bResult = false;
+			Handle->EnumerateConstRawData([&bResult](const void* RawData, const int32, const int32) -> bool
+			{
+				bResult = static_cast<const FBindableValueBase*>(RawData)->HasBinding();
+				return false;
+			});
+			return bResult;
+		};
+
+		TSharedRef<FInstancedStructProvider> StructProvider =
+			MakeShared<FInstancedStructProvider>(ConstantValueHandle);
+		ConstantValueHandle->RemoveChildren();
+		TArray<TSharedPtr<IPropertyHandle>> ChildProperties = ConstantValueHandle->AddChildStructure(StructProvider);
+		for (const TSharedPtr<IPropertyHandle>& ChildHandle : ChildProperties)
+		{
+			if (ChildHandle.IsValid())
+			{
+				ChildBuilder.AddProperty(ChildHandle.ToSharedRef())
+					.IsEnabled(TAttribute<bool>::CreateLambda([HasBinding]() { return !HasBinding(); }));
+			}
+		}
+	}
+	// All other types: intentionally empty — suppresses ConstantValue and Binding as separate child rows.
+}
+
+} // namespace UE::UAF::Editor
+
+#undef LOCTEXT_NAMESPACE
diff --git a/Engine/Plugins/Experimental/UAF/UAF/Source/UAFEditor/Private/BindableValue/UAFBindableValueCustomization.h b/Engine/Plugins/Experimental/UAF/UAF/Source/UAFEditor/Private/BindableValue/UAFBindableValueCustomization.h
new file mode 100644
--- /dev/null
+++ b/Engine/Plugins/Experimental/UAF/UAF/Source/UAFEditor/Private/BindableValue/UAFBindableValueCustomization.h
@@ -0,0 +1,46 @@
+// Copyright Epic Games, Inc. All Rights Reserved.
+
+#pragma once
+
+#include "IPropertyTypeCustomization.h"
+#include "Variables/VariablePickerArgs.h"
+
+class UUAFRigVMAsset;
+
+namespace UE::UAF::Editor
+{
+
+/**
+ * Property type customization for FBindableValueBase subclasses
+ * (FBindableBool, FBindableFloat, FBindableDouble, etc.).
+ *
+ * Places the constant-value widget and the SVariablePickerCombo side-by-side
+ * in the value content area. Reads/writes FBindableValueBase::Binding directly
+ * via EnumerateRawData — Binding has bare UPROPERTY() (no EditAnywhere) so it
+ * never appears in the property tree and cannot be reached via GetChildHandle.
+ */
+class FBindableValueCustomization : public IPropertyTypeCustomization
+{
+public:
+	virtual void CustomizeHeader(
+		TSharedRef<IPropertyHandle>      InPropertyHandle,
+		FDetailWidgetRow&                HeaderRow,
+		IPropertyTypeCustomizationUtils& CustomizationUtils) override;
+
+	virtual void CustomizeChildren(
+		TSharedRef<IPropertyHandle>      InPropertyHandle,
+		IDetailChildrenBuilder&          ChildBuilder,
+		IPropertyTypeCustomizationUtils& CustomizationUtils) override;
+
+private:
+	TSharedPtr<IPropertyHandle> PropertyHandle;
+	TSharedPtr<IPropertyHandle> ConstantValueHandle;
+	const UScriptStruct* ResolvedStructClass = nullptr;
+	bool bIsBindableStruct = false;
+
+	bool bIsContextSensitive = true;
+	TArray<TWeakObjectPtr<const UUAFRigVMAsset>> FilterAssets;
+	FOnIsContextSensitive OnIsContextSensitiveDelegate;
+};
+
+} // namespace UE::UAF::Editor
diff --git a/Engine/Plugins/Experimental/UAF/UAF/Source/UAFEditor/Private/Common/AnimNextAssetFindReplaceVariables.cpp b/Engine/Plugins/Experimental/UAF/UAF/Source/UAFEditor/Private/Common/AnimNextAssetFindReplaceVariables.cpp
--- a/Engine/Plugins/Experimental/UAF/UAF/Source/UAFEditor/Private/Common/AnimNextAssetFindReplaceVariables.cpp
+++ b/Engine/Plugins/Experimental/UAF/UAF/Source/UAFEditor/Private/Common/AnimNextAssetFindReplaceVariables.cpp
@@ -317,11 +317,6 @@ TSharedRef<SWidget> UAnimNextAssetFindReplaceVariables::MakeFindReplaceWidget()
 			RequestRefreshSearchResults();
 		});
 
-		SearchArgs.OnFilterVariableType  = FOnFilterVariableType::CreateLambda([this](const FAnimNextParamType& InParamType)-> EFilterVariableResult
-		{
-			return EFilterVariableResult::Include;
-		});
-
 		SearchArgs.OnFilterVariable = VariableScopeFilter;
 
 		// Filter to all declared variables (public and private)
@@ -352,7 +347,7 @@ TSharedRef<SWidget> UAnimNextAssetFindReplaceVariables::MakeFindReplaceWidget()
 			if(SearchType.IsValid())
 			{
 				// [TODO] IsCompatible or IsCompatibleWithDataLoss (e.g. UScriptStruct float entry to double RigVM entry fails IsCompatible)
-				if(!UE::UAF::FParamUtils::GetCompatibility(SearchType, InParamType).IsCompatible())
+				if(!UE::UAF::FParamUtils::GetCompatibility(SearchType, InParamType).IsCompatibleWithDataLoss())
 				{
 					return EFilterVariableResult::Exclude;
 				}
diff --git a/Engine/Plugins/Experimental/UAF/UAF/Source/UAFEditor/Private/Variables/SGraphPinVariableReference.cpp b/Engine/Plugins/Experimental/UAF/UAF/Source/UAFEditor/Private/Variables/SGraphPinVariableReference.cpp
--- a/Engine/Plugins/Experimental/UAF/UAF/Source/UAFEditor/Private/Variables/SGraphPinVariableReference.cpp
+++ b/Engine/Plugins/Experimental/UAF/UAF/Source/UAFEditor/Private/Variables/SGraphPinVariableReference.cpp
@@ -128,16 +128,13 @@ TSharedRef<SWidget> SGraphPinVariableReference::GetDefaultValueWidget()
 		return EFilterVariableResult::Exclude;
 	});
 
-	Args.OnFilterVariableType = FOnFilterVariableType::CreateLambda([this](const FAnimNextParamType& InParamType)-> EFilterVariableResult
+	if (FilterType.IsValid())
 	{
-		if(FilterType.IsValid())
-		{
-			if(!FParamUtils::GetCompatibility(FilterType, InParamType).IsCompatibleWithDataLoss())
-			{
-				return EFilterVariableResult::Exclude;
-			}
-		}
+		Args.CompatibleTypes = { FilterType };
+	}
 
+	Args.OnFilterVariableType = FOnFilterVariableType::CreateLambda([this](const FAnimNextParamType& InParamType) -> EFilterVariableResult
+	{
 		if(InParamType.IsValid())
 		{
 			const FRigVMTemplateArgumentType RigVMType = InParamType.ToRigVMTemplateArgument();
diff --git a/Engine/Plugins/Experimental/UAF/UAF/Source/UAFEditor/Private/Variables/SVariablePicker.cpp b/Engine/Plugins/Experimental/UAF/UAF/Source/UAFEditor/Private/Variables/SVariablePicker.cpp
--- a/Engine/Plugins/Experimental/UAF/UAF/Source/UAFEditor/Private/Variables/SVariablePicker.cpp
+++ b/Engine/Plugins/Experimental/UAF/UAF/Source/UAFEditor/Private/Variables/SVariablePicker.cpp
@@ -4,6 +4,8 @@
 
 #include "UncookedOnlyUtils.h"
 #include "Param/ParamType.h"
+#include "Param/ParamUtils.h"
+#include "Param/ParamCompatibility.h"
 #include "Widgets/Input/SSearchBox.h"
 #include "Widgets/Layout/SBox.h"
 #include "Widgets/Images/SImage.h"
@@ -21,7 +23,12 @@ void SVariablePicker::Construct(const FArguments& InArgs)
 {
 	Args = InArgs._Args;
 
-	FieldIterator = MakeUnique<FFieldIterator>(Args.OnFilterVariableType);
+	FieldIterator = MakeUnique<FFieldIterator>();
+	FieldIterator->Outer = this;
+	FieldIterator->CompatibleTypes = Args.CompatibleTypes;
+	const bool bAutoExpand = Args.CompatibleTypes.Num() > 0;
+	FieldIterator->bAllowStructExpansion = Args.bAllowStructExpansion || bAutoExpand;
+	FieldExpander.bAllowStructExpansion  = Args.bAllowStructExpansion || bAutoExpand;
 
 	ChildSlot
 	[
@@ -110,7 +117,24 @@ void SVariablePicker::RefreshEntries()
 							const FAnimNextSoftVariableReference VariableReference = FAnimNextSoftVariableReference::FromName(Export.Identifier, ExportPair.Key.ToSoftObjectPath());
 							if (!Args.OnFilterVariable.IsBound() || Args.OnFilterVariable.Execute(VariableReference) == EFilterVariableResult::Include)
 							{
-								if (!Args.OnFilterVariableType.IsBound() || Args.OnFilterVariableType.Execute(VariableEntry->Type) == EFilterVariableResult::Include)
+								bool bPassesTypeFilter = IsDirectlyCompatible(VariableEntry->Type);
+
+								if (!bPassesTypeFilter && FieldIterator->bAllowStructExpansion &&
+									VariableEntry->Type.GetValueType() == EPropertyBagPropertyType::Struct)
+								{
+									if (const UScriptStruct* ScriptStruct = Cast<UScriptStruct>(VariableEntry->Type.GetValueTypeObject()))
+									{
+										TSet<const UScriptStruct*> Visited;
+										bPassesTypeFilter = HasCompatibleLeaf(ScriptStruct, Visited);
+									}
+								}
+
+								if (bPassesTypeFilter && Args.OnFilterVariableType.IsBound())
+								{
+									bPassesTypeFilter = (Args.OnFilterVariableType.Execute(VariableEntry->Type) == EFilterVariableResult::Include);
+								}
+
+								if (bPassesTypeFilter)
 								{
 									PropertyDescs.Emplace(Export.Identifier, VariableEntry->Type.GetContainerType(), VariableEntry->Type.GetValueType(), VariableEntry->Type.GetValueTypeObject());
 								}
@@ -174,6 +198,11 @@ void SVariablePicker::HandleFieldPicked(UE::PropertyViewer::SPropertyViewer::FHa
 		FAnimNextSoftVariableReference SoftVariableReference;
 		if(GetFieldInfo(InHandle, InFields.Last(), SoftVariableReference, Type))
 		{
+			if (!IsDirectlyCompatible(Type))
+			{
+				return;  // not directly compatible — user must expand and pick a sub-property
+			}
+
 			const FAnimNextSoftVariableReference VariableReference = FAnimNextSoftVariableReference::FromName(SoftVariableReference.GetName(), SoftVariableReference.GetSoftObjectPath());
 			if(ensure(Type.IsValid() && !VariableReference.IsNone()))
 			{
@@ -181,6 +210,34 @@ void SVariablePicker::HandleFieldPicked(UE::PropertyViewer::SPropertyViewer::FHa
 			}
 		}
 	}
+	else if (InFields.Num() > 1 && Args.OnSubPropertyVariablePicked.IsBound())
+	{
+		// Sub-property path: InFields[0] = container variable, rest = path to leaf
+		FAnimNextParamType TopType;
+		FAnimNextSoftVariableReference TopVarRef;
+		if (GetFieldInfo(InHandle, InFields[0], TopVarRef, TopType))
+		{
+			TArray<FName> PathNames;
+			FAnimNextParamType LeafType;
+			for (int32 i = 1; i < InFields.Num(); ++i)
+			{
+				if (FProperty* Prop = InFields[i].Get<FProperty>())
+				{
+					PathNames.Add(Prop->GetFName());
+					LeafType = FAnimNextParamType::FromProperty(Prop);
+				}
+			}
+			if (PathNames.Num() > 0 && LeafType.IsValid())
+			{
+				if (!IsDirectlyCompatible(LeafType))
+				{
+					return;  // intermediate struct node — user must expand further
+				}
+
+				Args.OnSubPropertyVariablePicked.Execute(TopVarRef, PathNames, LeafType);
+			}
+		}
+	}
 }
 
 TSharedRef<SWidget> SVariablePicker::HandleGenerateContainer(UE::PropertyViewer::SPropertyViewer::FHandle InHandle, TOptional<FText> InDisplayName)
@@ -215,15 +272,26 @@ TSharedRef<SWidget> SVariablePicker::HandleGenerateContainer(UE::PropertyViewer:
 
 TArray<FFieldVariant> SVariablePicker::FFieldIterator::GetFields(const UStruct* Struct, const FName FieldName, const UStruct* ContainerStruct) const
 {
-	auto PassesFilterChecks = [this](const FProperty* InProperty)
+	auto PassesFilterChecks = [this](const FProperty* InProperty) -> bool
 	{
-		if(InProperty && OnFilterVariableType.IsBound())
+		if (!InProperty)
 		{
-			FAnimNextParamType Type = FAnimNextParamType::FromProperty(InProperty);
-			return Type.IsValid() && OnFilterVariableType.Execute(Type) == EFilterVariableResult::Include;
+			return false;
 		}
-
-		return false;
+		FAnimNextParamType Type = FAnimNextParamType::FromProperty(InProperty);
+		if (!Type.IsValid())
+		{
+			return false;
+		}
+		if (!Outer->IsDirectlyCompatible(Type))
+		{
+			return false;
+		}
+		if (Outer->Args.OnFilterVariableType.IsBound())
+		{
+			return Outer->Args.OnFilterVariableType.Execute(Type) == EFilterVariableResult::Include;
+		}
+		return true;
 	};
 
 	TArray<FFieldVariant> Result;
@@ -234,6 +302,18 @@ TArray<FFieldVariant> SVariablePicker::FFieldIterator::GetFields(const UStruct*
 		{
 			Result.Add(FFieldVariant(Property));
 		}
+		else if (bAllowStructExpansion)
+		{
+			// Include struct-type properties that contain any compatible leaf at any depth (handles nested structs like FTransform -> FVector -> float)
+			if (const FStructProperty* StructProp = CastField<FStructProperty>(Property))
+			{
+				TSet<const UScriptStruct*> Visited;
+				if (Outer->HasCompatibleLeaf(StructProp->Struct, Visited))
+				{
+					Result.Add(FFieldVariant(Property));
+				}
+			}
+		}
 	}
 
 	return Result;
@@ -246,7 +326,7 @@ TOptional<const UClass*> SVariablePicker::FFieldExpander::CanExpandObject(const
 
 bool SVariablePicker::FFieldExpander::CanExpandScriptStruct(const FStructProperty* StructProperty) const
 {
-	return false;
+	return bAllowStructExpansion;
 }
 
 TOptional<const UStruct*> SVariablePicker::FFieldExpander::GetExpandedFunction(const UFunction* Function) const
@@ -254,6 +334,47 @@ TOptional<const UStruct*> SVariablePicker::FFieldExpander::GetExpandedFunction(c
 	return TOptional<const UStruct*>();
 }
 
+bool SVariablePicker::IsDirectlyCompatible(const FAnimNextParamType& TestType) const
+{
+	if (Args.CompatibleTypes.IsEmpty())
+	{
+		return true;
+	}
+	for (const FAnimNextParamType& Compat : Args.CompatibleTypes)
+	{
+		if (FParamUtils::GetCompatibility(Compat, TestType).IsCompatibleWithDataLoss())
+		{
+			return true;
+		}
+	}
+	return false;
+}
+
+bool SVariablePicker::HasCompatibleLeaf(const UScriptStruct* Struct, TSet<const UScriptStruct*>& Visited) const
+{
+	if (!Struct || Visited.Contains(Struct))
+	{
+		return false;
+	}
+	Visited.Add(Struct);
+	for (TFieldIterator<FProperty> PropIt(Struct); PropIt; ++PropIt)
+	{
+		FAnimNextParamType SubType = FAnimNextParamType::FromProperty(*PropIt);
+		if (SubType.IsValid() && IsDirectlyCompatible(SubType))
+		{
+			return true;
+		}
+		if (const FStructProperty* SP = CastField<FStructProperty>(*PropIt))
+		{
+			if (HasCompatibleLeaf(SP->Struct, Visited))
+			{
+				return true;
+			}
+		}
+	}
+	return false;
+}
+
 }
 
-#undef LOCTEXT_NAMESPACE
+#undef LOCTEXT_NAMESPACE
\ No newline at end of file
diff --git a/Engine/Plugins/Experimental/UAF/UAF/Source/UAFEditor/Private/Variables/SVariablePicker.h b/Engine/Plugins/Experimental/UAF/UAF/Source/UAFEditor/Private/Variables/SVariablePicker.h
--- a/Engine/Plugins/Experimental/UAF/UAF/Source/UAFEditor/Private/Variables/SVariablePicker.h
+++ b/Engine/Plugins/Experimental/UAF/UAF/Source/UAFEditor/Private/Variables/SVariablePicker.h
@@ -38,6 +38,13 @@ private:
 
 	TSharedRef<SWidget> HandleGenerateContainer(UE::PropertyViewer::SPropertyViewer::FHandle InHandle, TOptional<FText> InDisplayName);
 
+	// Returns true if TestType is directly compatible with any entry in Args.CompatibleTypes.
+	// Returns true unconditionally when CompatibleTypes is empty.
+	bool IsDirectlyCompatible(const FAnimNextParamType& TestType) const;
+
+	// Returns true if Struct has any directly-compatible leaf at any depth.
+	bool HasCompatibleLeaf(const UScriptStruct* Struct, TSet<const UScriptStruct*>& Visited) const;
+
 private:
 	friend class SVariablePickerRow;
 
@@ -51,14 +58,14 @@ private:
 
 	struct FFieldIterator : UE::PropertyViewer::IFieldIterator
 	{
-		FFieldIterator(FOnFilterVariableType InOnFilterVariableType)
-			: OnFilterVariableType(InOnFilterVariableType)
-		{}
-		
+		FFieldIterator() = default;
+
 		virtual TArray<FFieldVariant> GetFields(const UStruct* Struct, const FName FieldName, const UStruct* ContainerStruct) const override;
 
-		FOnFilterVariableType OnFilterVariableType;
+		TArray<FAnimNextParamType> CompatibleTypes;
+		SVariablePicker* Outer = nullptr;
 		const UStruct* CurrentStruct = nullptr;
+		bool bAllowStructExpansion = false;
 	};
 
 	struct FFieldExpander : UE::PropertyViewer::IFieldExpander
@@ -66,6 +73,7 @@ private:
 		virtual TOptional<const UClass*> CanExpandObject(const FObjectPropertyBase* Property, const UObject* Instance) const override;
 		virtual bool CanExpandScriptStruct(const FStructProperty* StructProperty) const override;
 		virtual TOptional<const UStruct*> GetExpandedFunction(const UFunction* Function) const override;
+		bool bAllowStructExpansion = false;
 	} FieldExpander;
 
 	TUniquePtr<FFieldIterator> FieldIterator;
diff --git a/Engine/Plugins/Experimental/UAF/UAF/Source/UAFEditor/Private/Variables/SVariablePickerCombo.cpp b/Engine/Plugins/Experimental/UAF/UAF/Source/UAFEditor/Private/Variables/SVariablePickerCombo.cpp
--- a/Engine/Plugins/Experimental/UAF/UAF/Source/UAFEditor/Private/Variables/SVariablePickerCombo.cpp
+++ b/Engine/Plugins/Experimental/UAF/UAF/Source/UAFEditor/Private/Variables/SVariablePickerCombo.cpp
@@ -20,6 +20,7 @@ void SVariablePickerCombo::Construct(const FArguments& InArgs)
 	VariableTooltip = InArgs._VariableTooltip;
 
 	FOnVariablePicked OnVariablePicked = InArgs._PickerArgs.OnVariablePicked;
+	FOnSubPropertyVariablePicked OnSubPropertyVariablePicked = InArgs._PickerArgs.OnSubPropertyVariablePicked;
 
 	PickerArgs = InArgs._PickerArgs;
 
@@ -33,6 +34,16 @@ void SVariablePickerCombo::Construct(const FArguments& InArgs)
 		RequestRefresh();
 	});
 
+	PickerArgs.OnSubPropertyVariablePicked = FOnSubPropertyVariablePicked::CreateSPLambda(this, [this, OnSubPropertyVariablePicked](const FAnimNextSoftVariableReference& InTopVarRef, const TArray<FName>& InSubPropertyPath, const FAnimNextParamType& InLeafType)
+	{
+		FSlateApplication::Get().DismissAllMenus();
+
+		// Forward to the original delegate
+		OnSubPropertyVariablePicked.ExecuteIfBound(InTopVarRef, InSubPropertyPath, InLeafType);
+
+		RequestRefresh();
+	});
+
 	PickerArgs.bFocusSearchWidget = true;
 
 	ChildSlot
diff --git a/Engine/Plugins/Experimental/UAF/UAF/Source/UAFEditor/Private/Variables/VariableReferencePropertyCustomization.cpp b/Engine/Plugins/Experimental/UAF/UAF/Source/UAFEditor/Private/Variables/VariableReferencePropertyCustomization.cpp
--- a/Engine/Plugins/Experimental/UAF/UAF/Source/UAFEditor/Private/Variables/VariableReferencePropertyCustomization.cpp
+++ b/Engine/Plugins/Experimental/UAF/UAF/Source/UAFEditor/Private/Variables/VariableReferencePropertyCustomization.cpp
@@ -7,8 +7,6 @@
 #include "IPropertyUtilities.h"
 #include "Variables/SVariablePickerCombo.h"
 #include "UncookedOnlyUtils.h"
-#include "Param/ParamCompatibility.h"
-#include "Param/ParamUtils.h"
 #include "RigVMCore/RigVMRegistry.h"
 #include "Variables/AnimNextVariableReference.h"
 #include "Variables/VariablePickerArgs.h"
@@ -109,25 +107,21 @@ void FVariableReferencePropertyCustomization::CustomizeHeader(TSharedRef<IProper
 		UpdateCachedData();
 	});
 
-	Args.OnFilterVariableType = FOnFilterVariableType::CreateLambda([this](const FAnimNextParamType& InParamType)-> EFilterVariableResult
+	if (FilterType.IsValid())
 	{
-		if(FilterType.IsValid())
-		{
-			if(!FParamUtils::GetCompatibility(FilterType, InParamType).IsCompatibleWithDataLoss())
-			{
-				return EFilterVariableResult::Exclude;
-			}
-		}
+		Args.CompatibleTypes = { FilterType };
+	}
 
-		if(InParamType.IsValid())
+	Args.OnFilterVariableType = FOnFilterVariableType::CreateLambda([](const FAnimNextParamType& InParamType) -> EFilterVariableResult
+	{
+		if (InParamType.IsValid())
 		{
 			const FRigVMTemplateArgumentType RigVMType = InParamType.ToRigVMTemplateArgument();
-			if(!RigVMType.IsValid() || FRigVMRegistry::Get().GetTypeIndex(RigVMType) == INDEX_NONE)
+			if (!RigVMType.IsValid() || FRigVMRegistry::Get().GetTypeIndex(RigVMType) == INDEX_NONE)
 			{
 				return EFilterVariableResult::Exclude;
 			}
 		}
-
 		return EFilterVariableResult::Include;
 	});
 
diff --git a/Engine/Plugins/Experimental/UAF/UAF/Source/UAFEditor/UAFEditor.Build.cs b/Engine/Plugins/Experimental/UAF/UAF/Source/UAFEditor/UAFEditor.Build.cs
--- a/Engine/Plugins/Experimental/UAF/UAF/Source/UAFEditor/UAFEditor.Build.cs
+++ b/Engine/Plugins/Experimental/UAF/UAF/Source/UAFEditor/UAFEditor.Build.cs
@@ -68,6 +68,7 @@ namespace UnrealBuildTool.Rules
 					"HierarchyTableRuntime",
 					"HierarchyTableEditor",
 					"HierarchyTableAnimationRuntime",
+				"PropertyBindingUtils",
 				}
 			);
 			
diff --git a/Engine/Plugins/Experimental/UAF/UAF/Source/UAFTestData/Private/Module.cpp b/Engine/Plugins/Experimental/UAF/UAF/Source/UAFTestData/Private/Module.cpp
new file mode 100644
--- /dev/null
+++ b/Engine/Plugins/Experimental/UAF/UAF/Source/UAFTestData/Private/Module.cpp
@@ -0,0 +1,11 @@
+// Copyright Epic Games, Inc. All Rights Reserved.
+
+#include "CoreMinimal.h"
+#include "Modules/ModuleInterface.h"
+#include "Modules/ModuleManager.h"
+
+class FUAFTestDataModule : public IModuleInterface
+{
+};
+
+IMPLEMENT_MODULE(FUAFTestDataModule, UAFTestData)
diff --git a/Engine/Plugins/Experimental/UAF/UAF/Source/UAFTestData/Public/UAFTestVars.h b/Engine/Plugins/Experimental/UAF/UAF/Source/UAFTestData/Public/UAFTestVars.h
new file mode 100644
--- /dev/null
+++ b/Engine/Plugins/Experimental/UAF/UAF/Source/UAFTestData/Public/UAFTestVars.h
@@ -0,0 +1,438 @@
+// Copyright Epic Games, Inc. All Rights Reserved.
+
+#pragma once
+
+#include "CoreMinimal.h"
+#include "UAFTestVars.generated.h"
+
+UENUM()
+enum class EUAFTestEnum : uint8
+{
+	ValueA = 0,
+	ValueB = 1,
+	ValueC = 2,
+};
+
+/** Old-style enum for TEnumAsByte / FByteProperty testing. */
+UENUM()
+namespace EUAFTestByteEnum
+{
+	enum Type : uint8
+	{
+		ByteA = 0,
+		ByteB = 1,
+		ByteC = 2,
+	};
+}
+
+/** Nested struct for FBindableStruct SubProperty binding tests. */
+USTRUCT()
+struct FUAFNestedTestStruct
+{
+	GENERATED_BODY()
+
+	UPROPERTY()
+	FVector Vec = FVector::ZeroVector;
+
+	UPROPERTY()
+	FQuat Quat = FQuat::Identity;
+};
+
+/**
+ * Minimal USTRUCT used only in UAF LLT unit tests.
+ * Provides a set of typed properties that cover the common resolved types
+ * (bool, float, double, int32, FVector, FQuat) so that FBindableXxx runtime resolution
+ * tests can exercise variable and sub-property binding against a live FUAFAssetInstance.
+ */
+USTRUCT()
+struct FUAFTestVars
+{
+	GENERATED_BODY()
+
+	UPROPERTY()
+	bool bBool = false;
+
+	UPROPERTY()
+	float FloatVal = 0.0f;
+
+	UPROPERTY()
+	double DoubleVal = 0.0;
+
+	UPROPERTY()
+	int32 IntVal = 0;
+
+	UPROPERTY()
+	FVector VectorVar = FVector::ZeroVector;
+
+	UPROPERTY()
+	FQuat QuatVar = FQuat::Identity;
+
+	UPROPERTY()
+	FUAFNestedTestStruct NestedVar;
+
+	UPROPERTY()
+	int64 Int64Val = 0;
+
+	UPROPERTY()
+	uint8 ByteVal = 0;
+
+	UPROPERTY()
+	FName NameVar;
+
+	UPROPERTY()
+	EUAFTestEnum EnumVar = EUAFTestEnum::ValueA;
+
+	UPROPERTY()
+	TEnumAsByte<EUAFTestByteEnum::Type> EnumAsByteVar = EUAFTestByteEnum::ByteA;
+
+	UPROPERTY()
+	TObjectPtr<UObject> ObjectVar = nullptr;
+};
+
+/** 10-property struct used by bulk-resolution performance benchmarks. */
+USTRUCT()
+struct FUAFPerfVars10
+{
+	GENERATED_BODY()
+
+	// Float targets/sources for variable-binding benchmark
+	UPROPERTY()
+	float f0 = 0.f;
+
+	UPROPERTY()
+	float f1 = 0.f;
+
+	UPROPERTY()
+	float f2 = 0.f;
+
+	UPROPERTY()
+	float f3 = 0.f;
+
+	UPROPERTY()
+	float f4 = 0.f;
+
+	UPROPERTY()
+	float f5 = 0.f;
+
+	UPROPERTY()
+	float f6 = 0.f;
+
+	UPROPERTY()
+	float f7 = 0.f;
+
+	UPROPERTY()
+	float f8 = 0.f;
+
+	UPROPERTY()
+	float f9 = 0.f;
+
+	// Vector sources for subproperty-binding benchmark (bind .X → f0..f9)
+	UPROPERTY()
+	FVector v0 = FVector::ZeroVector;
+
+	UPROPERTY()
+	FVector v1 = FVector::ZeroVector;
+
+	UPROPERTY()
+	FVector v2 = FVector::ZeroVector;
+
+	UPROPERTY()
+	FVector v3 = FVector::ZeroVector;
+
+	UPROPERTY()
+	FVector v4 = FVector::ZeroVector;
+
+	UPROPERTY()
+	FVector v5 = FVector::ZeroVector;
+
+	UPROPERTY()
+	FVector v6 = FVector::ZeroVector;
+
+	UPROPERTY()
+	FVector v7 = FVector::ZeroVector;
+
+	UPROPERTY()
+	FVector v8 = FVector::ZeroVector;
+
+	UPROPERTY()
+	FVector v9 = FVector::ZeroVector;
+
+	// Bool sources for FBindableBool direct binding benchmark
+	UPROPERTY()
+	bool bBool0 = false;
+
+	UPROPERTY()
+	bool bBool1 = false;
+
+	UPROPERTY()
+	bool bBool2 = false;
+
+	UPROPERTY()
+	bool bBool3 = false;
+
+	UPROPERTY()
+	bool bBool4 = false;
+
+	UPROPERTY()
+	bool bBool5 = false;
+
+	UPROPERTY()
+	bool bBool6 = false;
+
+	UPROPERTY()
+	bool bBool7 = false;
+
+	UPROPERTY()
+	bool bBool8 = false;
+
+	UPROPERTY()
+	bool bBool9 = false;
+
+	// Double sources for FBindableDouble direct binding benchmark
+	UPROPERTY()
+	double d0 = 0.0;
+
+	UPROPERTY()
+	double d1 = 0.0;
+
+	UPROPERTY()
+	double d2 = 0.0;
+
+	UPROPERTY()
+	double d3 = 0.0;
+
+	UPROPERTY()
+	double d4 = 0.0;
+
+	UPROPERTY()
+	double d5 = 0.0;
+
+	UPROPERTY()
+	double d6 = 0.0;
+
+	UPROPERTY()
+	double d7 = 0.0;
+
+	UPROPERTY()
+	double d8 = 0.0;
+
+	UPROPERTY()
+	double d9 = 0.0;
+
+	// Int32 sources for FBindableInt32 direct binding benchmark
+	UPROPERTY()
+	int32 i0 = 0;
+
+	UPROPERTY()
+	int32 i1 = 0;
+
+	UPROPERTY()
+	int32 i2 = 0;
+
+	UPROPERTY()
+	int32 i3 = 0;
+
+	UPROPERTY()
+	int32 i4 = 0;
+
+	UPROPERTY()
+	int32 i5 = 0;
+
+	UPROPERTY()
+	int32 i6 = 0;
+
+	UPROPERTY()
+	int32 i7 = 0;
+
+	UPROPERTY()
+	int32 i8 = 0;
+
+	UPROPERTY()
+	int32 i9 = 0;
+
+	// Int64 sources for FBindableInt64 direct binding benchmark
+	UPROPERTY()
+	int64 l0 = 0;
+
+	UPROPERTY()
+	int64 l1 = 0;
+
+	UPROPERTY()
+	int64 l2 = 0;
+
+	UPROPERTY()
+	int64 l3 = 0;
+
+	UPROPERTY()
+	int64 l4 = 0;
+
+	UPROPERTY()
+	int64 l5 = 0;
+
+	UPROPERTY()
+	int64 l6 = 0;
+
+	UPROPERTY()
+	int64 l7 = 0;
+
+	UPROPERTY()
+	int64 l8 = 0;
+
+	UPROPERTY()
+	int64 l9 = 0;
+
+	// Byte sources for FBindableByte direct binding benchmark
+	UPROPERTY()
+	uint8 u0 = 0;
+
+	UPROPERTY()
+	uint8 u1 = 0;
+
+	UPROPERTY()
+	uint8 u2 = 0;
+
+	UPROPERTY()
+	uint8 u3 = 0;
+
+	UPROPERTY()
+	uint8 u4 = 0;
+
+	UPROPERTY()
+	uint8 u5 = 0;
+
+	UPROPERTY()
+	uint8 u6 = 0;
+
+	UPROPERTY()
+	uint8 u7 = 0;
+
+	UPROPERTY()
+	uint8 u8 = 0;
+
+	UPROPERTY()
+	uint8 u9 = 0;
+
+	// FName sources for FBindableName direct binding benchmark
+	UPROPERTY()
+	FName nm0;
+
+	UPROPERTY()
+	FName nm1;
+
+	UPROPERTY()
+	FName nm2;
+
+	UPROPERTY()
+	FName nm3;
+
+	UPROPERTY()
+	FName nm4;
+
+	UPROPERTY()
+	FName nm5;
+
+	UPROPERTY()
+	FName nm6;
+
+	UPROPERTY()
+	FName nm7;
+
+	UPROPERTY()
+	FName nm8;
+
+	UPROPERTY()
+	FName nm9;
+
+	// Enum sources for FBindableEnum direct binding benchmark
+	UPROPERTY()
+	EUAFTestEnum en0 = EUAFTestEnum::ValueA;
+
+	UPROPERTY()
+	EUAFTestEnum en1 = EUAFTestEnum::ValueA;
+
+	UPROPERTY()
+	EUAFTestEnum en2 = EUAFTestEnum::ValueA;
+
+	UPROPERTY()
+	EUAFTestEnum en3 = EUAFTestEnum::ValueA;
+
+	UPROPERTY()
+	EUAFTestEnum en4 = EUAFTestEnum::ValueA;
+
+	UPROPERTY()
+	EUAFTestEnum en5 = EUAFTestEnum::ValueA;
+
+	UPROPERTY()
+	EUAFTestEnum en6 = EUAFTestEnum::ValueA;
+
+	UPROPERTY()
+	EUAFTestEnum en7 = EUAFTestEnum::ValueA;
+
+	UPROPERTY()
+	EUAFTestEnum en8 = EUAFTestEnum::ValueA;
+
+	UPROPERTY()
+	EUAFTestEnum en9 = EUAFTestEnum::ValueA;
+
+	// Quat sources for FBindableQuat direct binding benchmark
+	UPROPERTY()
+	FQuat q0 = FQuat::Identity;
+
+	UPROPERTY()
+	FQuat q1 = FQuat::Identity;
+
+	UPROPERTY()
+	FQuat q2 = FQuat::Identity;
+
+	UPROPERTY()
+	FQuat q3 = FQuat::Identity;
+
+	UPROPERTY()
+	FQuat q4 = FQuat::Identity;
+
+	UPROPERTY()
+	FQuat q5 = FQuat::Identity;
+
+	UPROPERTY()
+	FQuat q6 = FQuat::Identity;
+
+	UPROPERTY()
+	FQuat q7 = FQuat::Identity;
+
+	UPROPERTY()
+	FQuat q8 = FQuat::Identity;
+
+	UPROPERTY()
+	FQuat q9 = FQuat::Identity;
+
+	// FUAFNestedTestStruct sources for FBindableVector sub-property benchmark (bind .Vec)
+	UPROPERTY()
+	FUAFNestedTestStruct ns0;
+
+	UPROPERTY()
+	FUAFNestedTestStruct ns1;
+
+	UPROPERTY()
+	FUAFNestedTestStruct ns2;
+
+	UPROPERTY()
+	FUAFNestedTestStruct ns3;
+
+	UPROPERTY()
+	FUAFNestedTestStruct ns4;
+
+	UPROPERTY()
+	FUAFNestedTestStruct ns5;
+
+	UPROPERTY()
+	FUAFNestedTestStruct ns6;
+
+	UPROPERTY()
+	FUAFNestedTestStruct ns7;
+
+	UPROPERTY()
+	FUAFNestedTestStruct ns8;
+
+	UPROPERTY()
+	FUAFNestedTestStruct ns9;
+};
diff --git a/Engine/Plugins/Experimental/UAF/UAF/Source/UAFTestData/UAFTestData.Build.cs b/Engine/Plugins/Experimental/UAF/UAF/Source/UAFTestData/UAFTestData.Build.cs
new file mode 100644
--- /dev/null
+++ b/Engine/Plugins/Experimental/UAF/UAF/Source/UAFTestData/UAFTestData.Build.cs
@@ -0,0 +1,18 @@
+// Copyright Epic Games, Inc. All Rights Reserved.
+
+using UnrealBuildTool;
+
+public class UAFTestData : ModuleRules
+{
+	public UAFTestData(ReadOnlyTargetRules Target) : base(Target)
+	{
+		bAllowUETypesInNamespaces = true;
+
+		PublicDependencyModuleNames.AddRange(new string[]
+		{
+			"Core",
+			"CoreUObject",
+			"UAF",
+		});
+	}
+}
diff --git a/Engine/Plugins/Experimental/UAF/UAF/Tests/UAF/Private/BindableValuePerfTests.cpp b/Engine/Plugins/Experimental/UAF/UAF/Tests/UAF/Private/BindableValuePerfTests.cpp
new file mode 100644
--- /dev/null
+++ b/Engine/Plugins/Experimental/UAF/UAF/Tests/UAF/Private/BindableValuePerfTests.cpp
@@ -0,0 +1,441 @@
+// Copyright Epic Games, Inc. All Rights Reserved.
+
+#include "CoreMinimal.h"
+#include "TestHarness.h"
+#include "BindableValue/UAFBindableTypes.h"
+#include "BindableValue/UAFPropertyBinding.h"
+#include "UAFAssetInstance.h"
+#include "Variables/AnimNextVariableReference.h"
+#include "Variables/UAFInstanceVariableData.h"
+#include "UAFTestVars.h"
+
+#include <catch2/catch_test_macros.hpp>
+#include <catch2/benchmark/catch_benchmark.hpp>
+
+// ---------------------------------------------------------------------------
+// Helpers
+// ---------------------------------------------------------------------------
+
+namespace UE::UAF::Tests
+{
+	// Forward-declared friend in FUAFAssetInstance — reuse for perf tests.
+	struct FUAFBindingTestFixture
+	{
+		static void InitFromStruct(FUAFAssetInstance& InInstance, const UScriptStruct* InStruct)
+		{
+			InInstance.Variables.AddVariablesContainerForStruct(InStruct, InInstance, nullptr);
+			InInstance.Variables.RebuildNameMaps();
+		}
+	};
+} // namespace UE::UAF::Tests
+
+namespace
+{
+	void InitPerfInstance(FUAFAssetInstance& Instance)
+	{
+		UE::UAF::Tests::FUAFBindingTestFixture::InitFromStruct(Instance, FUAFPerfVars10::StaticStruct());
+	}
+
+	FAnimNextVariableReference MakePerfRef(FName PropName)
+	{
+		const FProperty* Prop = FUAFPerfVars10::StaticStruct()->FindPropertyByName(PropName);
+		check(Prop);
+		return FAnimNextVariableReference::FromProperty(Prop, FUAFPerfVars10::StaticStruct());
+	}
+
+	FUAFPropertyBinding MakeVarBinding(const FAnimNextVariableReference& SourceVariable)
+	{
+		FUAFPropertyBinding Binding;
+		Binding.SourceType     = EUAFBindingSourceType::Variable;
+		Binding.SourceVariable = SourceVariable;
+		return Binding;
+	}
+
+	FUAFPropertyBinding MakeSubPropBinding(const FAnimNextVariableReference& SourceVariable, FStringView SubPath)
+	{
+		FPropertyBindingPath Path;
+		const bool bParsed = Path.FromString(SubPath);
+		REQUIRE(bParsed);
+
+		FUAFPropertyBinding Binding;
+		Binding.SourceType      = EUAFBindingSourceType::SubProperty;
+		Binding.SourceVariable  = SourceVariable;
+		Binding.SubPropertyPath = MoveTemp(Path);
+		return Binding;
+	}
+
+} // anonymous namespace
+
+// ---------------------------------------------------------------------------
+// Direct (variable binding) benchmarks
+// ---------------------------------------------------------------------------
+
+TEST_CASE("FBindableBool.BulkVariableBinding", "[UAF][perf]")
+{
+	FUAFAssetInstance Owner;
+	InitPerfInstance(Owner);
+
+	FAnimNextVariableReference Refs[10];
+	for (int i = 0; i < 10; ++i)
+	{
+		Refs[i] = MakePerfRef(*FString::Printf(TEXT("bBool%d"), i));
+	}
+
+	FBindableBool Fields[10];
+	for (int i = 0; i < 10; ++i)
+	{
+		Fields[i].SetBinding(MakeVarBinding(Refs[i]));
+	}
+
+	BENCHMARK("10x FBindableBool variable read")
+	{
+		int32 Sum = 0;
+		for (int i = 0; i < 10; ++i)
+		{
+			Sum += Fields[i].GetValue(&Owner) ? 1 : 0;
+		}
+		return Sum;
+	};
+}
+
+TEST_CASE("FBindableFloat.BulkVariableBinding", "[UAF][perf]")
+{
+	FUAFAssetInstance Owner;
+	InitPerfInstance(Owner);
+
+	FAnimNextVariableReference Refs[10];
+	for (int i = 0; i < 10; ++i)
+	{
+		Refs[i] = MakePerfRef(*FString::Printf(TEXT("f%d"), i));
+	}
+
+	FBindableFloat Fields[10];
+	for (int i = 0; i < 10; ++i)
+	{
+		Fields[i].SetBinding(MakeVarBinding(Refs[i]));
+	}
+
+	BENCHMARK("10x FBindableFloat variable read")
+	{
+		float Sum = 0.f;
+		for (int i = 0; i < 10; ++i)
+		{
+			Sum += Fields[i].GetValue(&Owner);
+		}
+		return Sum;
+	};
+}
+
+TEST_CASE("FBindableDouble.BulkVariableBinding", "[UAF][perf]")
+{
+	FUAFAssetInstance Owner;
+	InitPerfInstance(Owner);
+
+	FAnimNextVariableReference Refs[10];
+	for (int i = 0; i < 10; ++i)
+	{
+		Refs[i] = MakePerfRef(*FString::Printf(TEXT("d%d"), i));
+	}
+
+	FBindableDouble Fields[10];
+	for (int i = 0; i < 10; ++i)
+	{
+		Fields[i].SetBinding(MakeVarBinding(Refs[i]));
+	}
+
+	BENCHMARK("10x FBindableDouble variable read")
+	{
+		double Sum = 0.0;
+		for (int i = 0; i < 10; ++i)
+		{
+			Sum += Fields[i].GetValue(&Owner);
+		}
+		return Sum;
+	};
+}
+
+TEST_CASE("FBindableInt32.BulkVariableBinding", "[UAF][perf]")
+{
+	FUAFAssetInstance Owner;
+	InitPerfInstance(Owner);
+
+	FAnimNextVariableReference Refs[10];
+	for (int i = 0; i < 10; ++i)
+	{
+		Refs[i] = MakePerfRef(*FString::Printf(TEXT("i%d"), i));
+	}
+
+	FBindableInt32 Fields[10];
+	for (int i = 0; i < 10; ++i)
+	{
+		Fields[i].SetBinding(MakeVarBinding(Refs[i]));
+	}
+
+	BENCHMARK("10x FBindableInt32 variable read")
+	{
+		int32 Sum = 0;
+		for (int i = 0; i < 10; ++i)
+		{
+			Sum += Fields[i].GetValue(&Owner);
+		}
+		return Sum;
+	};
+}
+
+TEST_CASE("FBindableInt64.BulkVariableBinding", "[UAF][perf]")
+{
+	FUAFAssetInstance Owner;
+	InitPerfInstance(Owner);
+
+	FAnimNextVariableReference Refs[10];
+	for (int i = 0; i < 10; ++i)
+	{
+		Refs[i] = MakePerfRef(*FString::Printf(TEXT("l%d"), i));
+	}
+
+	FBindableInt64 Fields[10];
+	for (int i = 0; i < 10; ++i)
+	{
+		Fields[i].SetBinding(MakeVarBinding(Refs[i]));
+	}
+
+	BENCHMARK("10x FBindableInt64 variable read")
+	{
+		int64 Sum = 0;
+		for (int i = 0; i < 10; ++i)
+		{
+			Sum += Fields[i].GetValue(&Owner);
+		}
+		return Sum;
+	};
+}
+
+TEST_CASE("FBindableByte.BulkVariableBinding", "[UAF][perf]")
+{
+	FUAFAssetInstance Owner;
+	InitPerfInstance(Owner);
+
+	FAnimNextVariableReference Refs[10];
+	for (int i = 0; i < 10; ++i)
+	{
+		Refs[i] = MakePerfRef(*FString::Printf(TEXT("u%d"), i));
+	}
+
+	FBindableByte Fields[10];
+	for (int i = 0; i < 10; ++i)
+	{
+		Fields[i].SetBinding(MakeVarBinding(Refs[i]));
+	}
+
+	BENCHMARK("10x FBindableByte variable read")
+	{
+		uint32 Sum = 0;
+		for (int i = 0; i < 10; ++i)
+		{
+			Sum += Fields[i].GetValue(&Owner);
+		}
+		return Sum;
+	};
+}
+
+TEST_CASE("FBindableName.BulkVariableBinding", "[UAF][perf]")
+{
+	FUAFAssetInstance Owner;
+	InitPerfInstance(Owner);
+
+	FAnimNextVariableReference Refs[10];
+	for (int i = 0; i < 10; ++i)
+	{
+		Refs[i] = MakePerfRef(*FString::Printf(TEXT("nm%d"), i));
+	}
+
+	FBindableName Fields[10];
+	for (int i = 0; i < 10; ++i)
+	{
+		Fields[i].SetBinding(MakeVarBinding(Refs[i]));
+	}
+
+	BENCHMARK("10x FBindableName variable read")
+	{
+		FName Last;
+		for (int i = 0; i < 10; ++i)
+		{
+			Last = Fields[i].GetValue(&Owner);
+		}
+		return Last;
+	};
+}
+
+TEST_CASE("FBindableEnum.BulkVariableBinding", "[UAF][perf]")
+{
+	FUAFAssetInstance Owner;
+	InitPerfInstance(Owner);
+
+	FAnimNextVariableReference Refs[10];
+	for (int i = 0; i < 10; ++i)
+	{
+		Refs[i] = MakePerfRef(*FString::Printf(TEXT("en%d"), i));
+	}
+
+	FBindableEnum Fields[10];
+	for (int i = 0; i < 10; ++i)
+	{
+		Fields[i].SetConstantValue((int32)EUAFTestEnum::ValueA);
+		Fields[i].SetEnumClass(StaticEnum<EUAFTestEnum>());
+		Fields[i].SetBinding(MakeVarBinding(Refs[i]));
+	}
+
+	BENCHMARK("10x FBindableEnum variable read")
+	{
+		int32 Sum = 0;
+		for (int i = 0; i < 10; ++i)
+		{
+			Sum += Fields[i].GetValue(&Owner);
+		}
+		return Sum;
+	};
+}
+
+TEST_CASE("FBindableVector.BulkVariableBinding", "[UAF][perf]")
+{
+	FUAFAssetInstance Owner;
+	InitPerfInstance(Owner);
+
+	FAnimNextVariableReference Refs[10];
+	for (int i = 0; i < 10; ++i)
+	{
+		Refs[i] = MakePerfRef(*FString::Printf(TEXT("v%d"), i));
+	}
+
+	FBindableVector Fields[10];
+	for (int i = 0; i < 10; ++i)
+	{
+		Fields[i].SetBinding(MakeVarBinding(Refs[i]));
+	}
+
+	BENCHMARK("10x FBindableVector variable read")
+	{
+		FVector Last = FVector::ZeroVector;
+		for (int i = 0; i < 10; ++i)
+		{
+			Last = Fields[i].GetValue(&Owner);
+		}
+		return Last;
+	};
+}
+
+TEST_CASE("FBindableQuat.BulkVariableBinding", "[UAF][perf]")
+{
+	FUAFAssetInstance Owner;
+	InitPerfInstance(Owner);
+
+	FAnimNextVariableReference Refs[10];
+	for (int i = 0; i < 10; ++i)
+	{
+		Refs[i] = MakePerfRef(*FString::Printf(TEXT("q%d"), i));
+	}
+
+	FBindableQuat Fields[10];
+	for (int i = 0; i < 10; ++i)
+	{
+		Fields[i].SetBinding(MakeVarBinding(Refs[i]));
+	}
+
+	BENCHMARK("10x FBindableQuat variable read")
+	{
+		FQuat Last = FQuat::Identity;
+		for (int i = 0; i < 10; ++i)
+		{
+			Last = Fields[i].GetValue(&Owner);
+		}
+		return Last;
+	};
+}
+
+// ---------------------------------------------------------------------------
+// Sub-property binding benchmarks
+// ---------------------------------------------------------------------------
+
+TEST_CASE("FBindableFloat.BulkSubPropertyBinding", "[UAF][perf]")
+{
+	FUAFAssetInstance Owner;
+	InitPerfInstance(Owner);
+
+	FAnimNextVariableReference Refs[10];
+	for (int i = 0; i < 10; ++i)
+	{
+		Refs[i] = MakePerfRef(*FString::Printf(TEXT("v%d"), i));
+	}
+
+	FBindableFloat Fields[10];
+	for (int i = 0; i < 10; ++i)
+	{
+		Fields[i].SetBinding(MakeSubPropBinding(Refs[i], TEXT("X")));
+	}
+
+	BENCHMARK("10x FBindableFloat sub-property read")
+	{
+		float Sum = 0.f;
+		for (int i = 0; i < 10; ++i)
+		{
+			Sum += Fields[i].GetValue(&Owner);
+		}
+		return Sum;
+	};
+}
+
+TEST_CASE("FBindableDouble.BulkSubPropertyBinding", "[UAF][perf]")
+{
+	FUAFAssetInstance Owner;
+	InitPerfInstance(Owner);
+
+	FAnimNextVariableReference Refs[10];
+	for (int i = 0; i < 10; ++i)
+	{
+		Refs[i] = MakePerfRef(*FString::Printf(TEXT("v%d"), i));
+	}
+
+	FBindableDouble Fields[10];
+	for (int i = 0; i < 10; ++i)
+	{
+		Fields[i].SetBinding(MakeSubPropBinding(Refs[i], TEXT("X")));
+	}
+
+	BENCHMARK("10x FBindableDouble sub-property read")
+	{
+		double Sum = 0.0;
+		for (int i = 0; i < 10; ++i)
+		{
+			Sum += Fields[i].GetValue(&Owner);
+		}
+		return Sum;
+	};
+}
+
+TEST_CASE("FBindableVector.BulkSubPropertyBinding", "[UAF][perf]")
+{
+	FUAFAssetInstance Owner;
+	InitPerfInstance(Owner);
+
+	FAnimNextVariableReference Refs[10];
+	for (int i = 0; i < 10; ++i)
+	{
+		Refs[i] = MakePerfRef(*FString::Printf(TEXT("ns%d"), i));
+	}
+
+	FBindableVector Fields[10];
+	for (int i = 0; i < 10; ++i)
+	{
+		Fields[i].SetBinding(MakeSubPropBinding(Refs[i], TEXT("Vec")));
+	}
+
+	BENCHMARK("10x FBindableVector sub-property read")
+	{
+		FVector Last = FVector::ZeroVector;
+		for (int i = 0; i < 10; ++i)
+		{
+			Last = Fields[i].GetValue(&Owner);
+		}
+		return Last;
+	};
+}
diff --git a/Engine/Plugins/Experimental/UAF/UAF/Tests/UAF/Private/BindableValueSerializationTests.cpp b/Engine/Plugins/Experimental/UAF/UAF/Tests/UAF/Private/BindableValueSerializationTests.cpp
new file mode 100644
--- /dev/null
+++ b/Engine/Plugins/Experimental/UAF/UAF/Tests/UAF/Private/BindableValueSerializationTests.cpp
@@ -0,0 +1,649 @@
+// Copyright Epic Games, Inc. All Rights Reserved.
+
+#include "CoreMinimal.h"
+#include "TestHarness.h"
+
+#include "BindableValue/UAFBindableTypes.h"
+#include "BindableValue/UAFPropertyBinding.h"
+#include "Serialization/MemoryReader.h"
+#include "Serialization/MemoryWriter.h"
+#include "Serialization/ObjectAndNameAsStringProxyArchive.h"
+#include "UAFTestVars.h"
+#include "Variables/AnimNextVariableReference.h"
+
+#include <catch2/catch_test_macros.hpp>
+
+// ---------------------------------------------------------------------------
+// Helpers
+// ---------------------------------------------------------------------------
+
+namespace
+{
+
+FUAFPropertyBinding MakeTestBinding()
+{
+	const FProperty* Prop = FUAFTestVars::StaticStruct()->FindPropertyByName(TEXT("FloatVal"));
+	check(Prop);
+
+	FUAFPropertyBinding Binding;
+	Binding.SourceType     = EUAFBindingSourceType::Variable;
+	Binding.SourceVariable = FAnimNextVariableReference::FromProperty(Prop, FUAFTestVars::StaticStruct());
+	return Binding;
+}
+
+// ---------------------------------------------------------------------------
+// Binary serialization helpers
+// ---------------------------------------------------------------------------
+
+template<typename BindableType>
+BindableType BinaryRoundTrip(const BindableType& Source)
+{
+	// Serialize — use FObjectAndNameAsStringProxyArchive to handle UObject* references as strings
+	TArray<uint8> Buffer;
+	FMemoryWriter MemWriter(Buffer);
+	FObjectAndNameAsStringProxyArchive Writer(MemWriter, false);
+	UScriptStruct* Struct = BindableType::StaticStruct();
+	Struct->SerializeItem(Writer, const_cast<BindableType*>(&Source), nullptr);
+
+	// Deserialize
+	BindableType Dest;
+	FMemoryReader MemReader(Buffer);
+	FObjectAndNameAsStringProxyArchive Reader(MemReader, false);
+	Struct->SerializeItem(Reader, &Dest, nullptr);
+	return Dest;
+}
+
+// ---------------------------------------------------------------------------
+// Text serialization helpers
+// ---------------------------------------------------------------------------
+
+template<typename BindableType>
+BindableType TextRoundTrip(const BindableType& Source)
+{
+	// Export
+	FString Text;
+	BindableType Default;
+	Source.ExportTextItem(Text, Default, nullptr, 0, nullptr);
+
+	// Import
+	BindableType Dest;
+	const TCHAR* Buffer = *Text;
+	const bool bImported = Dest.ImportTextItem(Buffer, 0, nullptr, nullptr);
+	REQUIRE(bImported);
+	return Dest;
+}
+
+// ---------------------------------------------------------------------------
+// Editor-style text injection helper (mimics SetBindingViaText)
+// ---------------------------------------------------------------------------
+
+/** Mimics the editor flow: set binding on raw data, export, import back. */
+template<typename BindableType>
+BindableType EditorStyleSetBinding(const BindableType& Source, const FUAFPropertyBinding& InBinding)
+{
+	// Step 1: Copy source and set binding (like EnumerateRawData + SetBinding)
+	BindableType WithBinding = Source;
+	WithBinding.SetBinding(InBinding);
+
+	// Step 2: Export to text (like GetValueAsFormattedString → ExportTextItem)
+	FString TextValue;
+	BindableType Default;
+	WithBinding.ExportTextItem(TextValue, Default, nullptr, 0, nullptr);
+
+	// Step 3: Import back (like SetValueFromFormattedString → ImportText_Direct)
+	BindableType Dest;
+	const TCHAR* Buffer = *TextValue;
+	const bool bImported = Dest.ImportTextItem(Buffer, 0, nullptr, nullptr);
+	REQUIRE(bImported);
+	return Dest;
+}
+
+/** Mimics the editor flow: export current state (with binding), clear binding, import back. */
+template<typename BindableType>
+BindableType EditorStyleClearBinding(const BindableType& Source)
+{
+	// Export with binding, then re-import without it — simulates clearing via text round-trip
+	FString TextValue;
+	BindableType Default;
+	Source.ExportTextItem(TextValue, Default, nullptr, 0, nullptr);
+
+	// Strip the __Binding pseudo-property from the text
+	const int32 BindingIdx = TextValue.Find(TEXT("__Binding"));
+	if (BindingIdx != INDEX_NONE)
+	{
+		// Find the comma before __Binding (or opening paren if it's the only member)
+		int32 StartIdx = BindingIdx;
+		if (StartIdx > 0 && TextValue[StartIdx - 1] == TCHAR(','))
+		{
+			StartIdx--;
+		}
+		// Find the end of the binding value — match nested parens
+		int32 EndIdx = BindingIdx;
+		int32 Depth = 0;
+		bool bFoundEquals = false;
+		while (EndIdx < TextValue.Len())
+		{
+			const TCHAR Ch = TextValue[EndIdx];
+			if (!bFoundEquals)
+			{
+				if (Ch == TCHAR('=')) { bFoundEquals = true; }
+			}
+			else
+			{
+				if (Ch == TCHAR('(')) { Depth++; }
+				else if (Ch == TCHAR(')')) { if (--Depth <= 0) { EndIdx++; break; } }
+			}
+			EndIdx++;
+		}
+		TextValue.RemoveAt(StartIdx, EndIdx - StartIdx);
+	}
+
+	BindableType Dest;
+	const TCHAR* Buffer = *TextValue;
+	const bool bImported = Dest.ImportTextItem(Buffer, 0, nullptr, nullptr);
+	REQUIRE(bImported);
+	return Dest;
+}
+
+} // anonymous namespace
+
+// ===========================================================================
+// Binary Serialize round-trip tests
+// ===========================================================================
+
+TEST_CASE("FBindableBool.BinarySerialize.WithBinding", "[UAF][unit][MustPass]")
+{
+	FBindableBool Src(true);
+	Src.SetBinding(MakeTestBinding());
+	FBindableBool Dst = BinaryRoundTrip(Src);
+	CHECK(Dst.GetConstantValue() == true);
+	CHECK(Dst.HasBinding());
+}
+
+TEST_CASE("FBindableBool.BinarySerialize.NoBinding", "[UAF][unit][MustPass]")
+{
+	FBindableBool Src(true);
+	FBindableBool Dst = BinaryRoundTrip(Src);
+	CHECK(Dst.GetConstantValue() == true);
+	CHECK_FALSE(Dst.HasBinding());
+}
+
+TEST_CASE("FBindableFloat.BinarySerialize.WithBinding", "[UAF][unit][MustPass]")
+{
+	FBindableFloat Src(42.5f);
+	Src.SetBinding(MakeTestBinding());
+	FBindableFloat Dst = BinaryRoundTrip(Src);
+	CHECK(Dst.GetConstantValue() == 42.5f);
+	CHECK(Dst.HasBinding());
+}
+
+TEST_CASE("FBindableFloat.BinarySerialize.NoBinding", "[UAF][unit][MustPass]")
+{
+	FBindableFloat Src(42.5f);
+	FBindableFloat Dst = BinaryRoundTrip(Src);
+	CHECK(Dst.GetConstantValue() == 42.5f);
+	CHECK_FALSE(Dst.HasBinding());
+}
+
+TEST_CASE("FBindableDouble.BinarySerialize.WithBinding", "[UAF][unit][MustPass]")
+{
+	FBindableDouble Src(99.9);
+	Src.SetBinding(MakeTestBinding());
+	FBindableDouble Dst = BinaryRoundTrip(Src);
+	CHECK(Dst.GetConstantValue() == 99.9);
+	CHECK(Dst.HasBinding());
+}
+
+TEST_CASE("FBindableInt32.BinarySerialize.WithBinding", "[UAF][unit][MustPass]")
+{
+	FBindableInt32 Src(123);
+	Src.SetBinding(MakeTestBinding());
+	FBindableInt32 Dst = BinaryRoundTrip(Src);
+	CHECK(Dst.GetConstantValue() == 123);
+	CHECK(Dst.HasBinding());
+}
+
+TEST_CASE("FBindableInt64.BinarySerialize.WithBinding", "[UAF][unit][MustPass]")
+{
+	FBindableInt64 Src(9876543210LL);
+	Src.SetBinding(MakeTestBinding());
+	FBindableInt64 Dst = BinaryRoundTrip(Src);
+	CHECK(Dst.GetConstantValue() == 9876543210LL);
+	CHECK(Dst.HasBinding());
+}
+
+TEST_CASE("FBindableByte.BinarySerialize.WithBinding", "[UAF][unit][MustPass]")
+{
+	FBindableByte Src(42);
+	Src.SetBinding(MakeTestBinding());
+	FBindableByte Dst = BinaryRoundTrip(Src);
+	CHECK(Dst.GetConstantValue() == 42);
+	CHECK(Dst.HasBinding());
+}
+
+TEST_CASE("FBindableName.BinarySerialize.WithBinding", "[UAF][unit][MustPass]")
+{
+	FBindableName Src(FName(TEXT("TestName")));
+	Src.SetBinding(MakeTestBinding());
+	FBindableName Dst = BinaryRoundTrip(Src);
+	CHECK(Dst.GetConstantValue() == FName(TEXT("TestName")));
+	CHECK(Dst.HasBinding());
+}
+
+TEST_CASE("FBindableVector.BinarySerialize.WithBinding", "[UAF][unit][MustPass]")
+{
+	FBindableVector Src(FVector(1.0, 2.0, 3.0));
+	Src.SetBinding(MakeTestBinding());
+	FBindableVector Dst = BinaryRoundTrip(Src);
+	CHECK(Dst.GetConstantValue() == FVector(1.0, 2.0, 3.0));
+	CHECK(Dst.HasBinding());
+}
+
+TEST_CASE("FBindableQuat.BinarySerialize.WithBinding", "[UAF][unit][MustPass]")
+{
+	FBindableQuat Src(FQuat(1.0, 2.0, 3.0, 4.0));
+	Src.SetBinding(MakeTestBinding());
+	FBindableQuat Dst = BinaryRoundTrip(Src);
+	CHECK(Dst.GetConstantValue().X == 1.0);
+	CHECK(Dst.GetConstantValue().W == 4.0);
+	CHECK(Dst.HasBinding());
+}
+
+TEST_CASE("FBindableEnum.BinarySerialize.WithBinding", "[UAF][unit][MustPass]")
+{
+	FBindableEnum Src(EUAFTestEnum::ValueB);
+	Src.SetBinding(MakeTestBinding());
+	FBindableEnum Dst = BinaryRoundTrip(Src);
+	CHECK(Dst.GetConstantValue() == static_cast<int32>(EUAFTestEnum::ValueB));
+	CHECK(Dst.GetEnumClass() == StaticEnum<EUAFTestEnum>());
+	CHECK(Dst.HasBinding());
+}
+
+TEST_CASE("FBindableStruct.BinarySerialize.WithBinding", "[UAF][unit][MustPass]")
+{
+	FBindableStruct Src(FUAFTestVars::StaticStruct());
+	Src.SetBinding(MakeTestBinding());
+	FBindableStruct Dst = BinaryRoundTrip(Src);
+	CHECK(Dst.GetStructClass() == FUAFTestVars::StaticStruct());
+	CHECK(Dst.HasBinding());
+}
+
+TEST_CASE("FBindableObject.BinarySerialize.WithBinding", "[UAF][unit][MustPass]")
+{
+	FBindableObject Src(UObject::StaticClass());
+	Src.SetBinding(MakeTestBinding());
+	FBindableObject Dst = BinaryRoundTrip(Src);
+	CHECK(Dst.GetObjectClass() == UObject::StaticClass());
+	CHECK(Dst.HasBinding());
+}
+
+// ===========================================================================
+// Text ExportTextItem/ImportTextItem round-trip tests
+// ===========================================================================
+
+TEST_CASE("FBindableBool.TextSerialize.WithBinding", "[UAF][unit][MustPass]")
+{
+	FBindableBool Src(true);
+	Src.SetBinding(MakeTestBinding());
+	FBindableBool Dst = TextRoundTrip(Src);
+	CHECK(Dst.GetConstantValue() == true);
+	CHECK(Dst.HasBinding());
+}
+
+TEST_CASE("FBindableBool.TextSerialize.NoBinding", "[UAF][unit][MustPass]")
+{
+	FBindableBool Src(true);
+
+	FString Text;
+	FBindableBool Default;
+	Src.ExportTextItem(Text, Default, nullptr, 0, nullptr);
+	CHECK_FALSE(Text.Contains(TEXT("__Binding")));
+
+	FBindableBool Dst = TextRoundTrip(Src);
+	CHECK(Dst.GetConstantValue() == true);
+	CHECK_FALSE(Dst.HasBinding());
+}
+
+TEST_CASE("FBindableFloat.TextSerialize.WithBinding", "[UAF][unit][MustPass]")
+{
+	FBindableFloat Src(42.5f);
+	Src.SetBinding(MakeTestBinding());
+	FBindableFloat Dst = TextRoundTrip(Src);
+	CHECK(Dst.GetConstantValue() == 42.5f);
+	CHECK(Dst.HasBinding());
+}
+
+TEST_CASE("FBindableDouble.TextSerialize.WithBinding", "[UAF][unit][MustPass]")
+{
+	FBindableDouble Src(99.9);
+	Src.SetBinding(MakeTestBinding());
+	FBindableDouble Dst = TextRoundTrip(Src);
+	CHECK(Dst.GetConstantValue() == 99.9);
+	CHECK(Dst.HasBinding());
+}
+
+TEST_CASE("FBindableInt32.TextSerialize.WithBinding", "[UAF][unit][MustPass]")
+{
+	FBindableInt32 Src(123);
+	Src.SetBinding(MakeTestBinding());
+	FBindableInt32 Dst = TextRoundTrip(Src);
+	CHECK(Dst.GetConstantValue() == 123);
+	CHECK(Dst.HasBinding());
+}
+
+TEST_CASE("FBindableInt64.TextSerialize.WithBinding", "[UAF][unit][MustPass]")
+{
+	FBindableInt64 Src(9876543210LL);
+	Src.SetBinding(MakeTestBinding());
+	FBindableInt64 Dst = TextRoundTrip(Src);
+	CHECK(Dst.GetConstantValue() == 9876543210LL);
+	CHECK(Dst.HasBinding());
+}
+
+TEST_CASE("FBindableByte.TextSerialize.WithBinding", "[UAF][unit][MustPass]")
+{
+	FBindableByte Src(42);
+	Src.SetBinding(MakeTestBinding());
+	FBindableByte Dst = TextRoundTrip(Src);
+	CHECK(Dst.GetConstantValue() == 42);
+	CHECK(Dst.HasBinding());
+}
+
+TEST_CASE("FBindableName.TextSerialize.WithBinding", "[UAF][unit][MustPass]")
+{
+	FBindableName Src(FName(TEXT("TestName")));
+	Src.SetBinding(MakeTestBinding());
+	FBindableName Dst = TextRoundTrip(Src);
+	CHECK(Dst.GetConstantValue() == FName(TEXT("TestName")));
+	CHECK(Dst.HasBinding());
+}
+
+TEST_CASE("FBindableVector.TextSerialize.WithBinding", "[UAF][unit][MustPass]")
+{
+	FBindableVector Src(FVector(1.0, 2.0, 3.0));
+	Src.SetBinding(MakeTestBinding());
+	FBindableVector Dst = TextRoundTrip(Src);
+	CHECK(Dst.GetConstantValue() == FVector(1.0, 2.0, 3.0));
+	CHECK(Dst.HasBinding());
+}
+
+TEST_CASE("FBindableQuat.TextSerialize.WithBinding", "[UAF][unit][MustPass]")
+{
+	FBindableQuat Src(FQuat(1.0, 2.0, 3.0, 4.0));
+	Src.SetBinding(MakeTestBinding());
+	FBindableQuat Dst = TextRoundTrip(Src);
+	CHECK(Dst.GetConstantValue().X == 1.0);
+	CHECK(Dst.GetConstantValue().W == 4.0);
+	CHECK(Dst.HasBinding());
+}
+
+TEST_CASE("FBindableEnum.TextSerialize.WithBinding", "[UAF][unit][MustPass]")
+{
+	FBindableEnum Src(EUAFTestEnum::ValueB);
+	Src.SetBinding(MakeTestBinding());
+	FBindableEnum Dst = TextRoundTrip(Src);
+	CHECK(Dst.GetConstantValue() == static_cast<int32>(EUAFTestEnum::ValueB));
+	CHECK(Dst.GetEnumClass() == StaticEnum<EUAFTestEnum>());
+	CHECK(Dst.HasBinding());
+}
+
+TEST_CASE("FBindableStruct.TextSerialize.WithBinding", "[UAF][unit][MustPass]")
+{
+	FBindableStruct Src(FUAFTestVars::StaticStruct());
+	Src.SetBinding(MakeTestBinding());
+	FBindableStruct Dst = TextRoundTrip(Src);
+	CHECK(Dst.GetStructClass() == FUAFTestVars::StaticStruct());
+	CHECK(Dst.HasBinding());
+}
+
+TEST_CASE("FBindableObject.TextSerialize.WithBinding", "[UAF][unit][MustPass]")
+{
+	FBindableObject Src(UObject::StaticClass());
+	Src.SetBinding(MakeTestBinding());
+	FBindableObject Dst = TextRoundTrip(Src);
+	CHECK(Dst.GetObjectClass() == UObject::StaticClass());
+	CHECK(Dst.HasBinding());
+}
+
+// ===========================================================================
+// Editor-style set/clear binding tests (mimics SetBindingViaText flow)
+// ===========================================================================
+
+TEST_CASE("FBindableBool.EditorSetBinding", "[UAF][unit][MustPass]")
+{
+	FBindableBool Src(true);
+	FBindableBool Dst = EditorStyleSetBinding(Src, MakeTestBinding());
+	CHECK(Dst.GetConstantValue() == true);
+	CHECK(Dst.HasBinding());
+}
+
+TEST_CASE("FBindableFloat.EditorSetBinding", "[UAF][unit][MustPass]")
+{
+	FBindableFloat Src(42.5f);
+	FBindableFloat Dst = EditorStyleSetBinding(Src, MakeTestBinding());
+	CHECK(Dst.GetConstantValue() == 42.5f);
+	CHECK(Dst.HasBinding());
+}
+
+TEST_CASE("FBindableDouble.EditorSetBinding", "[UAF][unit][MustPass]")
+{
+	FBindableDouble Src(99.9);
+	FBindableDouble Dst = EditorStyleSetBinding(Src, MakeTestBinding());
+	CHECK(Dst.GetConstantValue() == 99.9);
+	CHECK(Dst.HasBinding());
+}
+
+TEST_CASE("FBindableInt32.EditorSetBinding", "[UAF][unit][MustPass]")
+{
+	FBindableInt32 Src(123);
+	FBindableInt32 Dst = EditorStyleSetBinding(Src, MakeTestBinding());
+	CHECK(Dst.GetConstantValue() == 123);
+	CHECK(Dst.HasBinding());
+}
+
+TEST_CASE("FBindableInt64.EditorSetBinding", "[UAF][unit][MustPass]")
+{
+	FBindableInt64 Src(9876543210LL);
+	FBindableInt64 Dst = EditorStyleSetBinding(Src, MakeTestBinding());
+	CHECK(Dst.GetConstantValue() == 9876543210LL);
+	CHECK(Dst.HasBinding());
+}
+
+TEST_CASE("FBindableByte.EditorSetBinding", "[UAF][unit][MustPass]")
+{
+	FBindableByte Src(42);
+	FBindableByte Dst = EditorStyleSetBinding(Src, MakeTestBinding());
+	CHECK(Dst.GetConstantValue() == 42);
+	CHECK(Dst.HasBinding());
+}
+
+TEST_CASE("FBindableName.EditorSetBinding", "[UAF][unit][MustPass]")
+{
+	FBindableName Src(FName(TEXT("TestName")));
+	FBindableName Dst = EditorStyleSetBinding(Src, MakeTestBinding());
+	CHECK(Dst.GetConstantValue() == FName(TEXT("TestName")));
+	CHECK(Dst.HasBinding());
+}
+
+TEST_CASE("FBindableVector.EditorSetBinding", "[UAF][unit][MustPass]")
+{
+	FBindableVector Src(FVector(1.0, 2.0, 3.0));
+	FBindableVector Dst = EditorStyleSetBinding(Src, MakeTestBinding());
+	CHECK(Dst.GetConstantValue() == FVector(1.0, 2.0, 3.0));
+	CHECK(Dst.HasBinding());
+}
+
+TEST_CASE("FBindableQuat.EditorSetBinding", "[UAF][unit][MustPass]")
+{
+	FBindableQuat Src(FQuat(1.0, 2.0, 3.0, 4.0));
+	FBindableQuat Dst = EditorStyleSetBinding(Src, MakeTestBinding());
+	CHECK(Dst.GetConstantValue().X == 1.0);
+	CHECK(Dst.GetConstantValue().W == 4.0);
+	CHECK(Dst.HasBinding());
+}
+
+TEST_CASE("FBindableEnum.EditorSetBinding", "[UAF][unit][MustPass]")
+{
+	FBindableEnum Src(EUAFTestEnum::ValueB);
+	FBindableEnum Dst = EditorStyleSetBinding(Src, MakeTestBinding());
+	CHECK(Dst.GetConstantValue() == static_cast<int32>(EUAFTestEnum::ValueB));
+	CHECK(Dst.GetEnumClass() == StaticEnum<EUAFTestEnum>());
+	CHECK(Dst.HasBinding());
+}
+
+TEST_CASE("FBindableStruct.EditorSetBinding", "[UAF][unit][MustPass]")
+{
+	FBindableStruct Src(FUAFTestVars::StaticStruct());
+	FBindableStruct Dst = EditorStyleSetBinding(Src, MakeTestBinding());
+	CHECK(Dst.GetStructClass() == FUAFTestVars::StaticStruct());
+	CHECK(Dst.HasBinding());
+}
+
+// ===========================================================================
+// Self-comparison export tests (simulates FRigVMMemoryStorageStruct::GetDataAsString
+// which passes Data as both value and default)
+// ===========================================================================
+
+template<typename BindableType>
+BindableType SelfComparisonRoundTrip(const BindableType& Source)
+{
+	// Simulate GetDataAsString(Data, Data): export with self as default
+	FString Text;
+	Source.ExportTextItem(Text, Source, nullptr, 0, nullptr);
+
+	// Verify the text is well-formed (no leading comma after opening paren)
+	CHECK_FALSE(Text.StartsWith(TEXT("(,")));
+
+	// Import into a fresh instance
+	BindableType Dest;
+	const TCHAR* Buffer = *Text;
+	const bool bImported = Dest.ImportTextItem(Buffer, 0, nullptr, nullptr);
+	REQUIRE(bImported);
+	return Dest;
+}
+
+TEST_CASE("FBindableBool.SelfComparisonExport.WithBinding", "[UAF][unit][MustPass]")
+{
+	FBindableBool Src(true);
+	Src.SetBinding(MakeTestBinding());
+	FBindableBool Dst = SelfComparisonRoundTrip(Src);
+	CHECK(Dst.GetConstantValue() == true);
+	CHECK(Dst.HasBinding());
+}
+
+TEST_CASE("FBindableBool.SelfComparisonExport.NoBinding", "[UAF][unit][MustPass]")
+{
+	FBindableBool Src(true);
+	FBindableBool Dst = SelfComparisonRoundTrip(Src);
+	CHECK(Dst.GetConstantValue() == true);
+	CHECK_FALSE(Dst.HasBinding());
+}
+
+TEST_CASE("FBindableFloat.SelfComparisonExport.WithBinding", "[UAF][unit][MustPass]")
+{
+	FBindableFloat Src(42.5f);
+	Src.SetBinding(MakeTestBinding());
+	FBindableFloat Dst = SelfComparisonRoundTrip(Src);
+	CHECK(Dst.GetConstantValue() == 42.5f);
+	CHECK(Dst.HasBinding());
+}
+
+TEST_CASE("FBindableVector.SelfComparisonExport.WithBinding", "[UAF][unit][MustPass]")
+{
+	FBindableVector Src(FVector(1.0, 2.0, 3.0));
+	Src.SetBinding(MakeTestBinding());
+	FBindableVector Dst = SelfComparisonRoundTrip(Src);
+	CHECK(Dst.GetConstantValue() == FVector(1.0, 2.0, 3.0));
+	CHECK(Dst.HasBinding());
+}
+
+TEST_CASE("FBindableEnum.SelfComparisonExport.WithBinding", "[UAF][unit][MustPass]")
+{
+	FBindableEnum Src(EUAFTestEnum::ValueB);
+	Src.SetBinding(MakeTestBinding());
+	FBindableEnum Dst = SelfComparisonRoundTrip(Src);
+	CHECK(Dst.GetConstantValue() == static_cast<int32>(EUAFTestEnum::ValueB));
+	CHECK(Dst.GetEnumClass() == StaticEnum<EUAFTestEnum>());
+	CHECK(Dst.HasBinding());
+}
+
+// Editor-style clear binding
+TEST_CASE("FBindableBool.EditorClearBinding", "[UAF][unit][MustPass]")
+{
+	FBindableBool Src(true);
+	Src.SetBinding(MakeTestBinding());
+	REQUIRE(Src.HasBinding());
+	FBindableBool Dst = EditorStyleClearBinding(Src);
+	CHECK(Dst.GetConstantValue() == true);
+	CHECK_FALSE(Dst.HasBinding());
+}
+
+TEST_CASE("FBindableVector.EditorClearBinding", "[UAF][unit][MustPass]")
+{
+	FBindableVector Src(FVector(1.0, 2.0, 3.0));
+	Src.SetBinding(MakeTestBinding());
+	REQUIRE(Src.HasBinding());
+	FBindableVector Dst = EditorStyleClearBinding(Src);
+	CHECK(Dst.GetConstantValue() == FVector(1.0, 2.0, 3.0));
+	CHECK_FALSE(Dst.HasBinding());
+}
+
+// ===========================================================================
+// Sibling binding isolation tests — when multiple FBindableXxx members are
+// inside the same parent struct, ImportTextItem receives the full remaining
+// buffer. A member MUST NOT pick up __Binding from a later sibling.
+// ===========================================================================
+
+TEST_CASE("FBindableBool.ImportTextItem.DoesNotPickUpSiblingBinding", "[UAF][unit][MustPass]")
+{
+	// Simulate the buffer that UScriptStruct::ImportText passes when importing
+	// BoolVal inside a parent struct where a LATER sibling (Int64Val) has a binding:
+	//   (BoolVal=(ConstantValue=True),Int64Val=(ConstantValue=0,__Binding=(...)))
+	// When importing BoolVal, the buffer starts at "(ConstantValue=True),Int64Val=..."
+	const FString Buffer = TEXT("(ConstantValue=True),Int64Val=(ConstantValue=0,__Binding=(SourceType=Variable,SourceVariable=(Name=\"FloatVal\")))");
+	const TCHAR* Ptr = *Buffer;
+
+	FBindableBool Dst;
+	const bool bImported = Dst.ImportTextItem(Ptr, 0, nullptr, nullptr);
+	REQUIRE(bImported);
+	CHECK(Dst.GetConstantValue() == true);
+	CHECK_FALSE(Dst.HasBinding()); // Must NOT pick up Int64Val's __Binding
+}
+
+TEST_CASE("FBindableFloat.ImportTextItem.DoesNotPickUpSiblingBinding", "[UAF][unit][MustPass]")
+{
+	const FString Buffer = TEXT("(ConstantValue=42.500000),NextProp=(ConstantValue=0,__Binding=(SourceType=Variable,SourceVariable=(Name=\"FloatVal\")))");
+	const TCHAR* Ptr = *Buffer;
+
+	FBindableFloat Dst;
+	const bool bImported = Dst.ImportTextItem(Ptr, 0, nullptr, nullptr);
+	REQUIRE(bImported);
+	CHECK(Dst.GetConstantValue() == 42.5f);
+	CHECK_FALSE(Dst.HasBinding());
+}
+
+TEST_CASE("FBindableInt64.ImportTextItem.PicksUpOwnBinding", "[UAF][unit][MustPass]")
+{
+	// The member that DOES have the binding should still import it correctly
+	const FString Buffer = TEXT("(ConstantValue=99,__Binding=(SourceType=Variable,SourceVariable=(Name=\"FloatVal\")))");
+	const TCHAR* Ptr = *Buffer;
+
+	FBindableInt64 Dst;
+	const bool bImported = Dst.ImportTextItem(Ptr, 0, nullptr, nullptr);
+	REQUIRE(bImported);
+	CHECK(Dst.GetConstantValue() == 99);
+	CHECK(Dst.HasBinding());
+}
+
+TEST_CASE("FBindableVector.ImportTextItem.DoesNotPickUpSiblingBinding", "[UAF][unit][MustPass]")
+{
+	// FVector has nested parens in its ConstantValue — make sure the struct boundary detection handles this
+	const FString Buffer = TEXT("(ConstantValue=(X=1.000000,Y=2.000000,Z=3.000000)),Sibling=(ConstantValue=0,__Binding=(SourceType=Variable,SourceVariable=(Name=\"FloatVal\")))");
+	const TCHAR* Ptr = *Buffer;
+
+	FBindableVector Dst;
+	const bool bImported = Dst.ImportTextItem(Ptr, 0, nullptr, nullptr);
+	REQUIRE(bImported);
+	CHECK(Dst.GetConstantValue() == FVector(1.0, 2.0, 3.0));
+	CHECK_FALSE(Dst.HasBinding());
+}
diff --git a/Engine/Plugins/Experimental/UAF/UAF/Tests/UAF/Private/BindableValueTests.cpp b/Engine/Plugins/Experimental/UAF/UAF/Tests/UAF/Private/BindableValueTests.cpp
new file mode 100644
--- /dev/null
+++ b/Engine/Plugins/Experimental/UAF/UAF/Tests/UAF/Private/BindableValueTests.cpp
@@ -0,0 +1,588 @@
+// Copyright Epic Games, Inc. All Rights Reserved.
+
+#include "CoreMinimal.h"
+#include "TestHarness.h"
+
+#include "BindableValue/UAFBindableTypes.h"
+#include "BindableValue/UAFPropertyBinding.h"
+#include "UAFAssetInstance.h"
+#include "Variables/AnimNextVariableReference.h"
+#include "Variables/UAFInstanceVariableData.h"
+#include "UAFTestVars.h"
+
+#include <catch2/catch_test_macros.hpp>
+
+// ---------------------------------------------------------------------------
+// Test fixture
+// ---------------------------------------------------------------------------
+
+namespace UE::UAF::Tests
+{
+	struct FUAFBindingTestFixture
+	{
+		static void InitFromStruct(FUAFAssetInstance& InInstance, const UScriptStruct* InStruct)
+		{
+			InInstance.Variables.AddVariablesContainerForStruct(InStruct, InInstance, nullptr);
+			InInstance.Variables.RebuildNameMaps();
+		}
+	};
+} // namespace UE::UAF::Tests
+
+// ---------------------------------------------------------------------------
+// Helpers
+// ---------------------------------------------------------------------------
+
+namespace
+{
+	using FFixture = UE::UAF::Tests::FUAFBindingTestFixture;
+
+	TPair<FUAFAssetInstance, FAnimNextVariableReference>
+	MakeOwnerAndRef(FName InPropertyName)
+	{
+		FUAFAssetInstance Owner;
+		FFixture::InitFromStruct(Owner, FUAFTestVars::StaticStruct());
+
+		const FProperty* Prop = FUAFTestVars::StaticStruct()->FindPropertyByName(InPropertyName);
+		REQUIRE(Prop != nullptr);
+
+		FAnimNextVariableReference Ref =
+			FAnimNextVariableReference::FromProperty(Prop, FUAFTestVars::StaticStruct());
+		return { MoveTemp(Owner), MoveTemp(Ref) };
+	}
+
+	/** Creates a FUAFPropertyBinding with Variable source type. */
+	FUAFPropertyBinding MakeVariablePropertyBinding(const FAnimNextVariableReference& SourceVariable)
+	{
+		FUAFPropertyBinding Binding;
+		Binding.SourceType     = EUAFBindingSourceType::Variable;
+		Binding.SourceVariable = SourceVariable;
+		return Binding;
+	}
+
+	/** Creates a FUAFPropertyBinding with SubProperty source type. */
+	FUAFPropertyBinding MakeSubPropertyBinding(const FAnimNextVariableReference& SourceVariable, FStringView SubPath)
+	{
+		FPropertyBindingPath Path;
+		const bool bParsed = Path.FromString(SubPath);
+		REQUIRE(bParsed);
+
+		FUAFPropertyBinding Binding;
+		Binding.SourceType      = EUAFBindingSourceType::SubProperty;
+		Binding.SourceVariable  = SourceVariable;
+		Binding.SubPropertyPath = MoveTemp(Path);
+		return Binding;
+	}
+
+} // anonymous namespace
+
+// ---------------------------------------------------------------------------
+// FBindableBool
+// ---------------------------------------------------------------------------
+
+TEST_CASE("FBindableBool.NoBindingReturnsFalse", "[UAF][unit][MustPass]")
+{
+	FBindableBool Field;
+	Field.SetConstantValue(false);
+	REQUIRE_FALSE(Field.HasBinding());
+
+	FUAFAssetInstance Empty;
+	REQUIRE(Field.GetValue(&Empty) == false);
+}
+
+TEST_CASE("FBindableBool.NoBindingReturnsValue", "[UAF][unit][MustPass]")
+{
+	FBindableBool Field;
+	Field.SetConstantValue(true);
+	REQUIRE_FALSE(Field.HasBinding());
+
+	FUAFAssetInstance Empty;
+	REQUIRE(Field.GetValue(&Empty) == true);
+}
+
+TEST_CASE("FBindableBool.VariableBindingResolvesTrue", "[UAF][unit][MustPass]")
+{
+	auto [Owner, Ref] = MakeOwnerAndRef("bBool");
+	REQUIRE(Owner.SetVariable(Ref, true) == EPropertyBagResult::Success);
+
+	FBindableBool Field;
+	Field.SetBinding(MakeVariablePropertyBinding(Ref));
+
+	REQUIRE(Field.HasBinding());
+	REQUIRE(Field.GetValue(&Owner) == true);
+}
+
+TEST_CASE("FBindableBool.VariableBindingReflectsUpdatedValue", "[UAF][unit][MustPass]")
+{
+	auto [Owner, Ref] = MakeOwnerAndRef("bBool");
+
+	FBindableBool Field;
+	Field.SetConstantValue(true);
+	Field.SetBinding(MakeVariablePropertyBinding(Ref));
+
+	REQUIRE(Owner.SetVariable(Ref, false) == EPropertyBagResult::Success);
+	REQUIRE(Field.GetValue(&Owner) == false);
+
+	REQUIRE(Owner.SetVariable(Ref, true) == EPropertyBagResult::Success);
+	REQUIRE(Field.GetValue(&Owner) == true);
+}
+
+TEST_CASE("FBindableBool.FailedResolutionFallsBackToValue", "[UAF][unit][MustPass]")
+{
+	// Binding with an empty (IsNone) source variable — should fall back to ConstantValue.
+	FBindableBool Field;
+	Field.SetConstantValue(true);
+
+	FUAFPropertyBinding EmptyBinding;
+	EmptyBinding.SourceType = EUAFBindingSourceType::Variable;
+	// SourceVariable intentionally default-constructed (IsNone() == true)
+	Field.SetBinding(MoveTemp(EmptyBinding));
+
+	REQUIRE(Field.HasBinding());
+
+	FUAFAssetInstance Empty;
+	REQUIRE(Field.GetValue(&Empty) == true); // falls back to ConstantValue
+}
+
+// ---------------------------------------------------------------------------
+// FBindableFloat
+// ---------------------------------------------------------------------------
+
+TEST_CASE("FBindableFloat.NoBindingReturnsValue", "[UAF][unit][MustPass]")
+{
+	FBindableFloat Field;
+	Field.SetConstantValue(3.14f);
+
+	FUAFAssetInstance Empty;
+	REQUIRE(Field.GetValue(&Empty) == 3.14f);
+}
+
+TEST_CASE("FBindableFloat.VariableBindingResolvesValue", "[UAF][unit][MustPass]")
+{
+	auto [Owner, Ref] = MakeOwnerAndRef("FloatVal");
+	REQUIRE(Owner.SetVariable(Ref, 1.5f) == EPropertyBagResult::Success);
+
+	FBindableFloat Field;
+	Field.SetBinding(MakeVariablePropertyBinding(Ref));
+
+	REQUIRE(Field.GetValue(&Owner) == 1.5f);
+}
+
+TEST_CASE("FBindableFloat.SubPropertyVectorXResolvesToFloat", "[UAF][unit][MustPass]")
+{
+	auto [Owner, VectorRef] = MakeOwnerAndRef("VectorVar");
+	const FVector Source(1.5, 2.5, 3.5);
+	REQUIRE(Owner.SetVariable(VectorRef, Source) == EPropertyBagResult::Success);
+
+	FBindableFloat Field;
+	Field.SetBinding(MakeSubPropertyBinding(VectorRef, TEXT("X")));
+
+	REQUIRE(Field.GetValue(&Owner) == static_cast<float>(Source.X));
+}
+
+TEST_CASE("FBindableFloat.SubPropertyVectorYResolvesToFloat", "[UAF][unit][MustPass]")
+{
+	auto [Owner, VectorRef] = MakeOwnerAndRef("VectorVar");
+	const FVector Source(1.5, 2.5, 3.5);
+	REQUIRE(Owner.SetVariable(VectorRef, Source) == EPropertyBagResult::Success);
+
+	FBindableFloat Field;
+	Field.SetBinding(MakeSubPropertyBinding(VectorRef, TEXT("Y")));
+
+	REQUIRE(Field.GetValue(&Owner) == static_cast<float>(Source.Y));
+}
+
+TEST_CASE("FBindableFloat.SubPropertyVectorZResolvesToFloat", "[UAF][unit][MustPass]")
+{
+	auto [Owner, VectorRef] = MakeOwnerAndRef("VectorVar");
+	const FVector Source(1.5, 2.5, 3.5);
+	REQUIRE(Owner.SetVariable(VectorRef, Source) == EPropertyBagResult::Success);
+
+	FBindableFloat Field;
+	Field.SetBinding(MakeSubPropertyBinding(VectorRef, TEXT("Z")));
+
+	REQUIRE(Field.GetValue(&Owner) == static_cast<float>(Source.Z));
+}
+
+TEST_CASE("FBindableFloat.SubPropertyQuatWResolvesToFloat", "[UAF][unit][MustPass]")
+{
+	auto [Owner, QuatRef] = MakeOwnerAndRef("QuatVar");
+	const FQuat Source = FQuat(FVector::UpVector, UE_HALF_PI);
+	REQUIRE(Owner.SetVariable(QuatRef, Source) == EPropertyBagResult::Success);
+
+	FBindableFloat Field;
+	Field.SetBinding(MakeSubPropertyBinding(QuatRef, TEXT("W")));
+
+	REQUIRE(Field.GetValue(&Owner) == static_cast<float>(Source.W));
+}
+
+TEST_CASE("FBindableFloat.TypeMismatchVariableBindingNoCrash", "[UAF][unit][MustPass]")
+{
+	auto [Owner, Ref] = MakeOwnerAndRef("VectorVar");
+	const FVector Expected(1.0, 2.0, 3.0);
+	REQUIRE(Owner.SetVariable(Ref, Expected) == EPropertyBagResult::Success);
+
+	// FVector variable → float target: should fall back gracefully, no crash.
+	FBindableFloat Field;
+	Field.SetConstantValue(99.0f);
+	Field.SetBinding(MakeVariablePropertyBinding(Ref));
+
+	const float Result = Field.GetValue(&Owner);
+	(void)Result;
+}
+
+// ---------------------------------------------------------------------------
+// FBindableDouble
+// ---------------------------------------------------------------------------
+
+TEST_CASE("FBindableDouble.NoBindingReturnsValue", "[UAF][unit][MustPass]")
+{
+	FBindableDouble Field;
+	Field.SetConstantValue(2.71828);
+
+	FUAFAssetInstance Empty;
+	REQUIRE(Field.GetValue(&Empty) == 2.71828);
+}
+
+TEST_CASE("FBindableDouble.VariableBindingResolvesValue", "[UAF][unit][MustPass]")
+{
+	auto [Owner, Ref] = MakeOwnerAndRef("DoubleVal");
+	REQUIRE(Owner.SetVariable(Ref, 2.718) == EPropertyBagResult::Success);
+
+	FBindableDouble Field;
+	Field.SetBinding(MakeVariablePropertyBinding(Ref));
+
+	REQUIRE(Field.GetValue(&Owner) == 2.718);
+}
+
+TEST_CASE("FBindableDouble.SubPropertyVectorXResolvesToDouble", "[UAF][unit][MustPass]")
+{
+	auto [Owner, VectorRef] = MakeOwnerAndRef("VectorVar");
+	const FVector Source(1.23456789, 0.0, 0.0);
+	REQUIRE(Owner.SetVariable(VectorRef, Source) == EPropertyBagResult::Success);
+
+	FBindableDouble Field;
+	Field.SetBinding(MakeSubPropertyBinding(VectorRef, TEXT("X")));
+
+	REQUIRE(Field.GetValue(&Owner) == Source.X);
+}
+
+TEST_CASE("FBindableDouble.SubPropertyQuatWResolvesToDouble", "[UAF][unit][MustPass]")
+{
+	auto [Owner, QuatRef] = MakeOwnerAndRef("QuatVar");
+	const FQuat Source = FQuat(FVector::UpVector, UE_HALF_PI);
+	REQUIRE(Owner.SetVariable(QuatRef, Source) == EPropertyBagResult::Success);
+
+	FBindableDouble Field;
+	Field.SetBinding(MakeSubPropertyBinding(QuatRef, TEXT("W")));
+
+	REQUIRE(Field.GetValue(&Owner) == Source.W);
+}
+
+// ---------------------------------------------------------------------------
+// FBindableInt32
+// ---------------------------------------------------------------------------
+
+TEST_CASE("FBindableInt32.NoBindingReturnsValue", "[UAF][unit][MustPass]")
+{
+	FBindableInt32 Field;
+	Field.SetConstantValue(42);
+
+	FUAFAssetInstance Empty;
+	REQUIRE(Field.GetValue(&Empty) == 42);
+}
+
+TEST_CASE("FBindableInt32.VariableBindingResolvesValue", "[UAF][unit][MustPass]")
+{
+	auto [Owner, Ref] = MakeOwnerAndRef("IntVal");
+	REQUIRE(Owner.SetVariable(Ref, 99) == EPropertyBagResult::Success);
+
+	FBindableInt32 Field;
+	Field.SetBinding(MakeVariablePropertyBinding(Ref));
+
+	REQUIRE(Field.GetValue(&Owner) == 99);
+}
+
+// ---------------------------------------------------------------------------
+// FBindableInt64
+// ---------------------------------------------------------------------------
+
+TEST_CASE("FBindableInt64.NoBindingReturnsValue", "[UAF][unit][MustPass]")
+{
+	FBindableInt64 Field;
+	Field.SetConstantValue(1234567890123LL);
+
+	FUAFAssetInstance Empty;
+	REQUIRE(Field.GetValue(&Empty) == 1234567890123LL);
+}
+
+TEST_CASE("FBindableInt64.VariableBindingResolvesValue", "[UAF][unit][MustPass]")
+{
+	auto [Owner, Ref] = MakeOwnerAndRef("Int64Val");
+	REQUIRE(Owner.SetVariable(Ref, (int64)9876543210LL) == EPropertyBagResult::Success);
+
+	FBindableInt64 Field;
+	Field.SetBinding(MakeVariablePropertyBinding(Ref));
+
+	REQUIRE(Field.GetValue(&Owner) == (int64)9876543210LL);
+}
+
+// ---------------------------------------------------------------------------
+// FBindableByte
+// ---------------------------------------------------------------------------
+
+TEST_CASE("FBindableByte.NoBindingReturnsValue", "[UAF][unit][MustPass]")
+{
+	FBindableByte Field;
+	Field.SetConstantValue((uint8)255);
+
+	FUAFAssetInstance Empty;
+	REQUIRE(Field.GetValue(&Empty) == (uint8)255);
+}
+
+TEST_CASE("FBindableByte.VariableBindingResolvesValue", "[UAF][unit][MustPass]")
+{
+	auto [Owner, Ref] = MakeOwnerAndRef("ByteVal");
+	REQUIRE(Owner.SetVariable(Ref, (uint8)42) == EPropertyBagResult::Success);
+
+	FBindableByte Field;
+	Field.SetBinding(MakeVariablePropertyBinding(Ref));
+
+	REQUIRE(Field.GetValue(&Owner) == (uint8)42);
+}
+
+// ---------------------------------------------------------------------------
+// FBindableName
+// ---------------------------------------------------------------------------
+
+TEST_CASE("FBindableName.NoBindingReturnsValue", "[UAF][unit][MustPass]")
+{
+	FBindableName Field;
+	Field.SetConstantValue(FName("TestName"));
+
+	FUAFAssetInstance Empty;
+	REQUIRE(Field.GetValue(&Empty) == FName("TestName"));
+}
+
+TEST_CASE("FBindableName.VariableBindingResolvesValue", "[UAF][unit][MustPass]")
+{
+	auto [Owner, Ref] = MakeOwnerAndRef("NameVar");
+	REQUIRE(Owner.SetVariable(Ref, FName("BoundName")) == EPropertyBagResult::Success);
+
+	FBindableName Field;
+	Field.SetBinding(MakeVariablePropertyBinding(Ref));
+
+	REQUIRE(Field.GetValue(&Owner) == FName("BoundName"));
+}
+
+// ---------------------------------------------------------------------------
+// FBindableEnum
+// ---------------------------------------------------------------------------
+
+TEST_CASE("FBindableEnum.NoBindingReturnsValue", "[UAF][unit][MustPass]")
+{
+	FBindableEnum Field;
+	Field.SetConstantValue((int32)EUAFTestEnum::ValueC);
+	Field.SetEnumClass(StaticEnum<EUAFTestEnum>());
+
+	FUAFAssetInstance Empty;
+	REQUIRE(Field.GetValue(&Empty) == (int32)EUAFTestEnum::ValueC);
+}
+
+TEST_CASE("FBindableEnum.VariableBindingResolvesValue", "[UAF][unit][MustPass]")
+{
+	auto [Owner, Ref] = MakeOwnerAndRef("EnumVar");
+	REQUIRE(Owner.SetVariable(Ref, (uint8)EUAFTestEnum::ValueB) == EPropertyBagResult::Success);
+
+	FBindableEnum Field;
+	Field.SetConstantValue((int32)EUAFTestEnum::ValueA);
+	Field.SetEnumClass(StaticEnum<EUAFTestEnum>());
+	Field.SetBinding(MakeVariablePropertyBinding(Ref));
+
+	REQUIRE(Field.GetValue(&Owner) == (int32)EUAFTestEnum::ValueB);
+}
+
+TEST_CASE("FBindableEnum.BytePropertyVariableBindingResolvesValue", "[UAF][unit][MustPass]")
+{
+	auto [Owner, Ref] = MakeOwnerAndRef("EnumAsByteVar");
+	REQUIRE(Owner.SetVariable(Ref, (uint8)EUAFTestByteEnum::ByteC) == EPropertyBagResult::Success);
+
+	UEnum* ByteEnumClass = StaticEnum<EUAFTestByteEnum::Type>();
+
+	FBindableEnum Field;
+	Field.SetConstantValue((int32)EUAFTestByteEnum::ByteA);
+	Field.SetEnumClass(ByteEnumClass);
+	Field.SetBinding(MakeVariablePropertyBinding(Ref));
+
+	REQUIRE(Field.GetValue(&Owner) == (int32)EUAFTestByteEnum::ByteC);
+}
+
+// ---------------------------------------------------------------------------
+// FBindableVector
+// ---------------------------------------------------------------------------
+
+TEST_CASE("FBindableVector.NoBindingReturnsConstantValue", "[UAF][unit][MustPass]")
+{
+	FBindableVector Field;
+	Field.SetConstantValue(FVector(1.0, 2.0, 3.0));
+
+	REQUIRE(Field.GetValue(nullptr) == FVector(1.0, 2.0, 3.0));
+}
+
+TEST_CASE("FBindableVector.VariableBindingResolvesVector", "[UAF][unit][MustPass]")
+{
+	auto [Owner, VectorRef] = MakeOwnerAndRef("VectorVar");
+	const FVector Source(10.0, 20.0, 30.0);
+	REQUIRE(Owner.SetVariable(VectorRef, Source) == EPropertyBagResult::Success);
+
+	FBindableVector Field;
+	Field.SetBinding(MakeVariablePropertyBinding(VectorRef));
+
+	REQUIRE(Field.GetValue(&Owner) == Source);
+}
+
+TEST_CASE("FBindableVector.SubPropertyNestedVecResolvesToVector", "[UAF][unit][MustPass]")
+{
+	auto [Owner, NestedRef] = MakeOwnerAndRef("NestedVar");
+	FUAFNestedTestStruct NestedData;
+	NestedData.Vec = FVector(7.0, 8.0, 9.0);
+	REQUIRE(Owner.SetVariable(NestedRef, NestedData) == EPropertyBagResult::Success);
+
+	FBindableVector Field;
+	Field.SetBinding(MakeSubPropertyBinding(NestedRef, TEXT("Vec")));
+
+	REQUIRE(Field.GetValue(&Owner) == NestedData.Vec);
+}
+
+// ---------------------------------------------------------------------------
+// FBindableQuat
+// ---------------------------------------------------------------------------
+
+TEST_CASE("FBindableQuat.NoBindingReturnsConstantValue", "[UAF][unit][MustPass]")
+{
+	const FQuat Expected = FQuat(FVector::UpVector, UE_HALF_PI);
+	FBindableQuat Field;
+	Field.SetConstantValue(Expected);
+
+	REQUIRE(Field.GetValue(nullptr) == Expected);
+}
+
+TEST_CASE("FBindableQuat.VariableBindingResolvesQuat", "[UAF][unit][MustPass]")
+{
+	auto [Owner, QuatRef] = MakeOwnerAndRef("QuatVar");
+	const FQuat Source = FQuat(FVector::UpVector, UE_HALF_PI);
+	REQUIRE(Owner.SetVariable(QuatRef, Source) == EPropertyBagResult::Success);
+
+	FBindableQuat Field;
+	Field.SetBinding(MakeVariablePropertyBinding(QuatRef));
+
+	REQUIRE(Field.GetValue(&Owner) == Source);
+}
+
+TEST_CASE("FBindableQuat.SubPropertyNestedQuatResolvesToQuat", "[UAF][unit][MustPass]")
+{
+	auto [Owner, NestedRef] = MakeOwnerAndRef("NestedVar");
+	FUAFNestedTestStruct NestedData;
+	NestedData.Quat = FQuat(FVector::ForwardVector, UE_HALF_PI);
+	REQUIRE(Owner.SetVariable(NestedRef, NestedData) == EPropertyBagResult::Success);
+
+	FBindableQuat Field;
+	Field.SetBinding(MakeSubPropertyBinding(NestedRef, TEXT("Quat")));
+
+	REQUIRE(Field.GetValue(&Owner) == NestedData.Quat);
+}
+
+// ---------------------------------------------------------------------------
+// FBindableStruct
+// ---------------------------------------------------------------------------
+
+TEST_CASE("FBindableStruct.NoBindingReturnsConstantValue", "[UAF][unit][MustPass]")
+{
+	FBindableStruct Field;
+	FInstancedStruct Tmp;
+	Tmp.InitializeAs<FVector>(FVector(1.0, 2.0, 3.0));
+	Field.SetConstantValue(MoveTemp(Tmp));
+
+	FVector Result = FVector::ZeroVector;
+	Field.GetValue(nullptr, Result);   // null Instance → falls back to ConstantValue
+	REQUIRE(Result == FVector(1.0, 2.0, 3.0));
+}
+
+TEST_CASE("FBindableStruct.VariableBindingResolvesVector", "[UAF][unit][MustPass]")
+{
+	auto [Owner, VectorRef] = MakeOwnerAndRef("VectorVar");
+	const FVector Source(4.0, 5.0, 6.0);
+	REQUIRE(Owner.SetVariable(VectorRef, Source) == EPropertyBagResult::Success);
+
+	FBindableStruct Field;
+	Field.SetBinding(MakeVariablePropertyBinding(VectorRef));
+
+	FVector Result = FVector::ZeroVector;
+	Field.GetValue(&Owner, Result);
+	REQUIRE(Result == Source);
+}
+
+TEST_CASE("FBindableStruct.SubPropertyNestedVecResolvesToVector", "[UAF][unit][MustPass]")
+{
+	auto [Owner, NestedRef] = MakeOwnerAndRef("NestedVar");
+	FUAFNestedTestStruct NestedData;
+	NestedData.Vec = FVector(7.0, 8.0, 9.0);
+	REQUIRE(Owner.SetVariable(NestedRef, NestedData) == EPropertyBagResult::Success);
+
+	FBindableStruct Field;
+	Field.SetBinding(MakeSubPropertyBinding(NestedRef, TEXT("Vec")));
+
+	FVector Result = FVector::ZeroVector;
+	Field.GetValue(&Owner, Result);
+	REQUIRE(Result == NestedData.Vec);
+}
+
+TEST_CASE("FBindableStruct.EmptyConstantValueInitializesOutput", "[UAF][unit][MustPass]")
+{
+	// Construct with StructClass only — no ConstantValue
+	FBindableStruct Field(TBaseStructure<FVector>::Get());
+
+	// Fill OutValue with garbage to detect if it's left unmodified
+	FVector OutValue(999.0, 999.0, 999.0);
+	Field.GetValue<FVector>(nullptr, OutValue);
+
+	// Should be zero-initialized, not left as garbage
+	REQUIRE(OutValue == FVector::ZeroVector);
+}
+
+// ---------------------------------------------------------------------------
+// FBindableObject
+// ---------------------------------------------------------------------------
+
+TEST_CASE("FBindableObject.NoBindingReturnsNull", "[UAF][unit][MustPass]")
+{
+	FBindableObject Field;
+	REQUIRE_FALSE(Field.HasBinding());
+	REQUIRE(Field.GetValue(nullptr) == nullptr);
+}
+
+TEST_CASE("FBindableObject.NoBindingReturnsConstant", "[UAF][unit][MustPass]")
+{
+	UObject* TestObj = UObject::StaticClass()->GetDefaultObject();
+	FBindableObject Field;
+	Field.SetConstantValue(TestObj);
+
+	FUAFAssetInstance Empty;
+	REQUIRE(Field.GetValue(&Empty) == TestObj);
+}
+
+TEST_CASE("FBindableObject.VariableBindingResolvesObject", "[UAF][unit][MustPass]")
+{
+	auto [Owner, ObjRef] = MakeOwnerAndRef("ObjectVar");
+	UObject* TestObj = UObject::StaticClass()->GetDefaultObject();
+
+	const FAnimNextParamType ObjParamType(EPropertyBagPropertyType::Object,
+		EPropertyBagContainerType::None, UObject::StaticClass());
+	REQUIRE(Owner.SetVariable(ObjRef, ObjParamType,
+		TConstArrayView<uint8>(reinterpret_cast<const uint8*>(&TestObj), sizeof(TObjectPtr<UObject>))) == EPropertyBagResult::Success);
+
+	FBindableObject Field;
+	Field.SetObjectClass(UObject::StaticClass());
+	Field.SetBinding(MakeVariablePropertyBinding(ObjRef));
+
+	REQUIRE(Field.GetValue(&Owner) == TestObj);
+}
diff --git a/Engine/Plugins/Experimental/UAF/UAF/Tests/UAF/Private/Module.cpp b/Engine/Plugins/Experimental/UAF/UAF/Tests/UAF/Private/Module.cpp
new file mode 100644
--- /dev/null
+++ b/Engine/Plugins/Experimental/UAF/UAF/Tests/UAF/Private/Module.cpp
@@ -0,0 +1,32 @@
+// Copyright Epic Games, Inc. All Rights Reserved.
+
+#include "CoreMinimal.h"
+#include "Modules/ModuleInterface.h"
+#include "Modules/ModuleManager.h"
+#include "TestCommon/Initialization.h"
+#include "TestRunner.h"
+
+class FUAFTestsModule : public IModuleInterface
+{
+};
+IMPLEMENT_MODULE(FUAFTestsModule, UAFTests)
+
+namespace
+{
+	void SetupUAFTests()
+	{
+		InitAll(/*bAllowLogging=*/false, /*bMultithreaded=*/false);
+	}
+	void TeardownUAFTests()
+	{
+		CleanupAll();
+	}
+	struct FUAFTestsGlobalSetup
+	{
+		FUAFTestsGlobalSetup()
+		{
+			FTestDelegates::GetGlobalSetup().BindStatic(&SetupUAFTests);
+			FTestDelegates::GetGlobalTeardown().BindStatic(&TeardownUAFTests);
+		}
+	} GUAFTestsGlobalSetup;
+}
diff --git a/Engine/Plugins/Experimental/UAF/UAF/Tests/UAFTests.Build.cs b/Engine/Plugins/Experimental/UAF/UAF/Tests/UAFTests.Build.cs
--- a/Engine/Plugins/Experimental/UAF/UAF/Tests/UAFTests.Build.cs
+++ b/Engine/Plugins/Experimental/UAF/UAF/Tests/UAFTests.Build.cs
@@ -17,8 +17,9 @@ public class UAFTests : TestModuleRules
 		PrivateDependencyModuleNames.AddRange(
 			// Any private dependencies to link against
 			new string[] {
-				"UAF"
-			}			
+				"UAF",
+				"UAFTestData",
+			}
 		);
 	}
 }
\ No newline at end of file
diff --git a/Engine/Plugins/Experimental/UAF/UAF/UAF.uplugin b/Engine/Plugins/Experimental/UAF/UAF/UAF.uplugin
--- a/Engine/Plugins/Experimental/UAF/UAF/UAF.uplugin
+++ b/Engine/Plugins/Experimental/UAF/UAF/UAF.uplugin
@@ -49,6 +49,10 @@
 		{
 			"Name": "GameplayInsights",
 			"Enabled": true
+		},
+		{
+			"Name": "PropertyBindingUtils",
+			"Enabled": true
 		}
 	],
 	"Modules" :
@@ -67,6 +71,11 @@
 			"Name" : "UAFUncookedOnly",
 			"Type" : "UncookedOnly",
 			"LoadingPhase" : "Default"
+		},
+		{
+			"Name" : "UAFTestData",
+			"Type" : "DeveloperTool",
+			"LoadingPhase" : "Default"
 		}
 	]
 }
diff --git a/Engine/Plugins/Experimental/UAF/UAFAnimNode/Source/UAFAnimNode/Private/UAF/AnimNodes/UAFBlendByBoolNode.cpp b/Engine/Plugins/Experimental/UAF/UAFAnimNode/Source/UAFAnimNode/Private/UAF/AnimNodes/UAFBlendByBoolNode.cpp
--- a/Engine/Plugins/Experimental/UAF/UAFAnimNode/Source/UAFAnimNode/Private/UAF/AnimNodes/UAFBlendByBoolNode.cpp
+++ b/Engine/Plugins/Experimental/UAF/UAFAnimNode/Source/UAFAnimNode/Private/UAF/AnimNodes/UAFBlendByBoolNode.cpp
@@ -1,87 +1,107 @@
 // Copyright Epic Games, Inc. All Rights Reserved.
+
 #include "UAF/AnimNodes/UAFBlendByBoolNode.h"
-#include "Module/AnimNextModuleInstance.h"
 #include "UAF/AnimNodeCore/UAFAnimNodeUpdate.h"
+#include "UAFAssetInstance.h"
 
 #include UE_INLINE_GENERATED_CPP_BY_NAME(UAFBlendByBoolNode)
 
 namespace UE::UAF
 {
-	FUAFBlendByBoolNode::FUAFBlendByBoolNode(FUAFAnimGraphUpdateContext& Context, const FUAFBlendByBoolNodeData* InData)
-		: FUAFBlendStack(Context), Data(InData)
-	{
-		InitializeAs<FUAFBlendByBoolNode>(Context);
+FUAFBlendByBoolNode::FUAFBlendByBoolNode(FUAFAnimGraphUpdateContext& Context, const FUAFBlendByBoolNodeData* InData)
+	: FUAFBlendStack(Context)
+	, Data(InData)
+{
+	InitializeAs<FUAFBlendByBoolNode>(Context);
 
-		Context.GetVariablesOwner()->GetVariable(Data->Bool, bCurrentValue);
+	// Read the initial bool value, preferring a binding over the data's literal value.
+	bCurrentValue = Data->BoolValue.GetValue(Context.GetVariablesOwner());
 
-		if (bCurrentValue)
+	if (bCurrentValue)
+	{
+		if (Data->TrueNode.IsValid())
 		{
 			TransitionTo(Context, Data->TrueNode.Get().CreateInstance(Context));
 		}
-		else
+	}
+	else
+	{
+		if (Data->FalseNode.IsValid())
 		{
 			TransitionTo(Context, Data->FalseNode.Get().CreateInstance(Context));
 		}
 	}
+}
 
-	void FUAFBlendByBoolNode::PreUpdate(FUAFAnimGraphUpdateContext& Context)
-	{
-		FUAFBlendStack::PreUpdate(Context);
+void FUAFBlendByBoolNode::PreUpdate(FUAFAnimGraphUpdateContext& Context)
+{
+	FUAFBlendStack::PreUpdate(Context);
 
-		bool bNewValue = false;
-		Context.GetVariablesOwner()->GetVariable(Data->Bool, bNewValue);
+	bool bNewValue = Data->BoolValue.GetValue(Context.GetVariablesOwner());
 
-		UAF_TRACE_ANIMNODE_VALUE(Context, this, "Bool", bNewValue);
+	UAF_TRACE_ANIMNODE_VALUE(Context, this, "Bool", bNewValue);
 
-		if (bCurrentValue != bNewValue)
+	if (bCurrentValue != bNewValue)
+	{
+		bCurrentValue = bNewValue;
+		if (bCurrentValue)
 		{
-			bCurrentValue = bNewValue;
-			if (bCurrentValue)
+			if (Data->TrueNode.IsValid())
 			{
 				TransitionTo(Context, Data->TrueNode.Get().CreateInstance(Context), Data->Transition.GetPtr());
 			}
-			else
+		}
+		else
+		{
+			if (Data->FalseNode.IsValid())
 			{
 				TransitionTo(Context, Data->FalseNode.Get().CreateInstance(Context), Data->Transition.GetPtr());
 			}
 		}
-
-		ValidateSingleChild();
 	}
 
+	ValidateSingleChild();
+}
+
 #if UAF_TRACE_ENABLED
 
-	FString FUAFBlendByBoolNode::GetDebugName() const
-	{
-		static FString BlendByBoolName("Blend By Bool");
-		return BlendByBoolName;
-	}
+FString FUAFBlendByBoolNode::GetDebugName() const
+{
+	static FString BlendByBoolName("Blend By Bool");
+	return BlendByBoolName;
+}
 
-	UStruct* FUAFBlendByBoolNode::GetDebugStruct() const
-	{
-		return FUAFBlendByBoolNodeData::StaticStruct();
-	}
+UStruct* FUAFBlendByBoolNode::GetDebugStruct() const
+{
+	return FUAFBlendByBoolNodeData::StaticStruct();
+}
 
 #endif
 
 
-	FUAFAnimNodePtr FUAFBlendByBoolNodeData::CreateInstance(FUAFAnimGraphUpdateContext& Context) const
+FUAFAnimNodePtr FUAFBlendByBoolNodeData::CreateInstance(FUAFAnimGraphUpdateContext& Context) const
+{
+	if (EvaluationFrequency == EBlendByBoolEvaluationFrequency::OnCreate)
 	{
-		if (EvaluationFrequency == EBlendByBoolEvaluationFrequency::OnCreate)
-		{
-			bool bValue = false;
-			Context.GetVariablesOwner()->GetVariable(Bool, bValue);
+		bool bValue = BoolValue.GetValue(Context.GetVariablesOwner());
 
-			if (bValue)
+		if (bValue)
+		{
+			if (TrueNode.IsValid())
 			{
 				return TrueNode.Get().CreateInstance(Context);
 			}
-			else
+		}
+		else
+		{
+			if (FalseNode.IsValid())
 			{
 				return FalseNode.Get().CreateInstance(Context);
 			}
 		}
-		
-		return MakeAnimNode<FUAFBlendByBoolNode>(Context, this);
+		return FUAFAnimNodePtr();
 	}
+
+	return MakeAnimNode<FUAFBlendByBoolNode>(Context, this);
+}
 }
diff --git a/Engine/Plugins/Experimental/UAF/UAFAnimNode/Source/UAFAnimNode/Private/UAF/AnimNodes/UAFScalePlayRate.cpp b/Engine/Plugins/Experimental/UAF/UAFAnimNode/Source/UAFAnimNode/Private/UAF/AnimNodes/UAFScalePlayRate.cpp
--- a/Engine/Plugins/Experimental/UAF/UAFAnimNode/Source/UAFAnimNode/Private/UAF/AnimNodes/UAFScalePlayRate.cpp
+++ b/Engine/Plugins/Experimental/UAF/UAFAnimNode/Source/UAFAnimNode/Private/UAF/AnimNodes/UAFScalePlayRate.cpp
@@ -31,8 +31,9 @@ namespace UE::UAF
 
 	void FUAFScalePlayRate::PreUpdate(FUAFAnimGraphUpdateContext& Context)
 	{
-		Context.PushPlayRate(PlayRateMultiplier);
-		UAF_TRACE_ANIMNODE_VALUE(Context, this, "PlayRateMultiplier", PlayRateMultiplier);
+		float Multiplier = PlayRateMultiplier.GetValue(Context.GetVariablesOwner());
+		UAF_TRACE_ANIMNODE_VALUE(Context, this, "PlayRateMultiplier", Multiplier);
+		Context.PushPlayRate(Multiplier);
 	}
 
 	void FUAFScalePlayRate::PostUpdate(FUAFAnimGraphUpdateContext& Context)
diff --git a/Engine/Plugins/Experimental/UAF/UAFAnimNode/Source/UAFAnimNode/Public/UAF/AnimNodes/UAFBlendByBoolNode.h b/Engine/Plugins/Experimental/UAF/UAFAnimNode/Source/UAFAnimNode/Public/UAF/AnimNodes/UAFBlendByBoolNode.h
--- a/Engine/Plugins/Experimental/UAF/UAFAnimNode/Source/UAFAnimNode/Public/UAF/AnimNodes/UAFBlendByBoolNode.h
+++ b/Engine/Plugins/Experimental/UAF/UAFAnimNode/Source/UAFAnimNode/Public/UAF/AnimNodes/UAFBlendByBoolNode.h
@@ -3,10 +3,9 @@
 #pragma once
 
 #include "CoreMinimal.h"
-#include "EvaluationVM/Tasks/PushAnimSequenceKeyframe.h"
 #include "UAF/AnimNodeCore/UAFAnimNodeData.h"
+#include "BindableValue/UAFBindableTypes.h"
 #include "UAF/AnimNodes/UAFBlendStack.h"
-#include "Variables/AnimNextVariableReference.h"
 
 #include "UAFBlendByBoolNode.generated.h"
 
@@ -22,30 +21,32 @@ namespace UE::UAF
 		// Check the input bool every frame, and when its value changes, transition to the appropriate child node
 		OnUpdate
 	};
-	
+
 	// BlendByBool Anim Node selects between two child nodes based on a bool
 	USTRUCT(DisplayName = "Blend By Bool")
 	struct FUAFBlendByBoolNodeData : public FUAFAnimNodeData
 	{
 		GENERATED_BODY()
 
-		// Bool variable to select based on 
-		// todo: this should be some sort of binding when we build that system
+		/**
+		 * Bool value used to select between TrueNode and FalseNode.
+		 * Bind to a UAF bool variable at runtime by setting the Binding field.
+		 */
 		UPROPERTY(EditAnywhere, Category = "NodeData")
-		FAnimNextVariableReference Bool;
+		FBindableBool BoolValue;
 
-		// Evaluation Frequency: should we select the appriate child once at creation time, or every frame.
+		// Evaluation Frequency: should we select the appropriate child once at creation time, or every frame.
 		UPROPERTY(EditAnywhere, Category = "NodeData")
 		EBlendByBoolEvaluationFrequency EvaluationFrequency = EBlendByBoolEvaluationFrequency::OnUpdate;
 
 		// Child Node to instantiate if input bool is true
 		UPROPERTY(EditAnywhere, Category = "NodeData")
 		TInstancedStruct<FUAFAnimNodeData> TrueNode;
-		
-		// Child Node to instantiate if input bool is false 
+
+		// Child Node to instantiate if input bool is false
 		UPROPERTY(EditAnywhere, Category = "NodeData")
 		TInstancedStruct<FUAFAnimNodeData> FalseNode;
-		
+
 		// Transition and blend data for transitioning between child nodes
 		UPROPERTY(EditAnywhere, Category = "NodeData")
 		TInstancedStruct<FUAFTransitionNodeData> Transition;
diff --git a/Engine/Plugins/Experimental/UAF/UAFAnimNode/Source/UAFAnimNode/Public/UAF/AnimNodes/UAFScalePlayRate.h b/Engine/Plugins/Experimental/UAF/UAFAnimNode/Source/UAFAnimNode/Public/UAF/AnimNodes/UAFScalePlayRate.h
--- a/Engine/Plugins/Experimental/UAF/UAFAnimNode/Source/UAFAnimNode/Public/UAF/AnimNodes/UAFScalePlayRate.h
+++ b/Engine/Plugins/Experimental/UAF/UAFAnimNode/Source/UAFAnimNode/Public/UAF/AnimNodes/UAFScalePlayRate.h
@@ -5,6 +5,7 @@
 #include "CoreMinimal.h"
 #include "UAF/AnimNodeCore/UAFAnimNodeData.h"
 #include "UAF/AnimNodeCore/UAFAnimNode.h"
+#include "BindableValue/UAFBindableTypes.h"
 
 #include "UAFScalePlayRate.generated.h"
 
@@ -22,7 +23,7 @@ namespace UE::UAF
 		TInstancedStruct<FUAFAnimNodeData> Child;
 
 		UPROPERTY(EditAnywhere, Category = "Data")
-		float PlayRateMultiplier = 1.0f;
+		FBindableFloat PlayRateMultiplier = 1.0f;
 
 		// FUAFAnimNodeData impl
 		UE_API virtual FUAFAnimNodePtr CreateInstance(FUAFAnimGraphUpdateContext& Context) const override;
@@ -42,7 +43,7 @@ namespace UE::UAF
 		virtual UStruct* GetDebugStruct() const override;
 #endif
 	private:
-		float PlayRateMultiplier = 1.0f;
+		FBindableFloat PlayRateMultiplier = 1.0f;
 	};
 }
 
diff --git a/Engine/Plugins/Experimental/UAF/UAFAnimNode/Source/UAFAnimNode/UAFAnimNode.Build.cs b/Engine/Plugins/Experimental/UAF/UAFAnimNode/Source/UAFAnimNode/UAFAnimNode.Build.cs
--- a/Engine/Plugins/Experimental/UAF/UAFAnimNode/Source/UAFAnimNode/UAFAnimNode.Build.cs
+++ b/Engine/Plugins/Experimental/UAF/UAFAnimNode/Source/UAFAnimNode/UAFAnimNode.Build.cs
@@ -7,7 +7,7 @@ namespace UnrealBuildTool.Rules
 		public UAFAnimNode(ReadOnlyTargetRules Target) : base(Target)
 		{
 			bAllowUETypesInNamespaces = true;
-			
+
 			PublicDependencyModuleNames.AddRange(new string[] { "UAF" });
 			PrivateDependencyModuleNames.AddRange(
 				new string[]
@@ -21,7 +21,7 @@ namespace UnrealBuildTool.Rules
 					"TraceLog",
 				}
 			);
-			
+
 			if (Target.bBuildEditor)
 			{
 				PrivateDependencyModuleNames.AddRange(
@@ -30,11 +30,10 @@ namespace UnrealBuildTool.Rules
 						"UAFUncookedOnly",
 						"UAFEditor",
 						"WorkspaceEditor",
-						"UAFAnimNodeUncookedOnly"
 					}
 				);
 			}
-			
+
 		}
 	}
-}
\ No newline at end of file
+}
diff --git a/Engine/Plugins/Experimental/UAF/UAFAnimNode/Source/UAFAnimNodeUncookedOnly/UAFAnimNodeUncookedOnly.Build.cs b/Engine/Plugins/Experimental/UAF/UAFAnimNode/Source/UAFAnimNodeUncookedOnly/UAFAnimNodeUncookedOnly.Build.cs
--- a/Engine/Plugins/Experimental/UAF/UAFAnimNode/Source/UAFAnimNodeUncookedOnly/UAFAnimNodeUncookedOnly.Build.cs
+++ b/Engine/Plugins/Experimental/UAF/UAFAnimNode/Source/UAFAnimNodeUncookedOnly/UAFAnimNodeUncookedOnly.Build.cs
@@ -6,26 +6,14 @@ namespace UnrealBuildTool.Rules
 	{
 		public UAFAnimNodeUncookedOnly(ReadOnlyTargetRules Target) : base(Target)
 		{
-			PublicDependencyModuleNames.AddRange(
-				new string[]
-				{
-					"Core",
-				}
-			);
-
 			PrivateDependencyModuleNames.AddRange(
 				new string[]
 				{
+					"Core",
 					"CoreUObject",
-					"Engine",
-					"UAF", 
-					"UAFEditor", 
-					"UAFUncookedOnly",
-					"UAFAnimGraph",
-					"UAFAnimGraphUncookedOnly",
-					"WorkspaceEditor"
+					"UAFAnimNode",
 				}
 			);
 		}
 	}
-}
\ No newline at end of file
+}
diff --git a/Engine/Plugins/Experimental/UAF/UAFTestSuites/Source/UAFAnimNodeTestData/Private/Module.cpp b/Engine/Plugins/Experimental/UAF/UAFTestSuites/Source/UAFAnimNodeTestData/Private/Module.cpp
new file mode 100644
--- /dev/null
+++ b/Engine/Plugins/Experimental/UAF/UAFTestSuites/Source/UAFAnimNodeTestData/Private/Module.cpp
@@ -0,0 +1,11 @@
+// Copyright Epic Games, Inc. All Rights Reserved.
+
+#include "CoreMinimal.h"
+#include "Modules/ModuleInterface.h"
+#include "Modules/ModuleManager.h"
+
+class FUAFAnimNodeTestDataModule : public IModuleInterface
+{
+};
+
+IMPLEMENT_MODULE(FUAFAnimNodeTestDataModule, UAFAnimNodeTestData)
diff --git a/Engine/Plugins/Experimental/UAF/UAFTestSuites/Source/UAFAnimNodeTestData/Public/UAFAnimNodeTestVars.h b/Engine/Plugins/Experimental/UAF/UAFTestSuites/Source/UAFAnimNodeTestData/Public/UAFAnimNodeTestVars.h
new file mode 100644
--- /dev/null
+++ b/Engine/Plugins/Experimental/UAF/UAFTestSuites/Source/UAFAnimNodeTestData/Public/UAFAnimNodeTestVars.h
@@ -0,0 +1,154 @@
+// Copyright Epic Games, Inc. All Rights Reserved.
+
+#pragma once
+
+#include "CoreMinimal.h"
+#include "Animation/TrajectoryTypes.h"
+#include "UAF/AnimNodeCore/UAFAnimNodeData.h"
+#include "BindableValue/UAFBindableTypes.h"
+#include "UAFAnimNodeTestVars.generated.h"
+
+/** Nested struct for FBindableStruct SubProperty binding tests. */
+USTRUCT()
+struct FUAFAnimNodeNestedTestStruct
+{
+	GENERATED_BODY()
+
+	UPROPERTY()
+	FVector Vec = FVector::ZeroVector;
+
+	UPROPERTY()
+	FQuat Quat = FQuat::Identity;
+};
+
+/**
+ * Minimal USTRUCT used only in UAFAnimNode LLT unit tests.
+ * Provides a set of typed properties that cover the common resolved types
+ * (bool, float, double, int32, FVector, FQuat) so that FBindableXxx runtime resolution
+ * tests can exercise variable and sub-property binding against a live FUAFAssetInstance.
+ */
+USTRUCT()
+struct FUAFAnimNodeTestVars
+{
+	GENERATED_BODY()
+
+	UPROPERTY(EditAnywhere, Category = "Variables")
+	bool bBool = false;
+
+	UPROPERTY(EditAnywhere, Category = "Variables")
+	float FloatVal = 0.0f;
+
+	UPROPERTY(EditAnywhere, Category = "Variables")
+	double DoubleVal = 0.0;
+
+	UPROPERTY(EditAnywhere, Category = "Variables")
+	int32 IntVal = 0;
+
+	UPROPERTY(EditAnywhere, Category = "Variables")
+	FVector VectorVar = FVector::ZeroVector;
+
+	UPROPERTY(EditAnywhere, Category = "Variables")
+	FQuat QuatVar = FQuat::Identity;
+
+	UPROPERTY(EditAnywhere, Category = "Variables")
+	FUAFAnimNodeNestedTestStruct NestedVar;
+};
+
+/** 10-property struct used by bulk-resolution performance benchmarks. */
+USTRUCT()
+struct FUAFAnimNodePerfVars10
+{
+	GENERATED_BODY()
+
+	// Float targets/sources for variable-binding benchmark
+	UPROPERTY() float f0 = 0.f;
+	UPROPERTY() float f1 = 0.f;
+	UPROPERTY() float f2 = 0.f;
+	UPROPERTY() float f3 = 0.f;
+	UPROPERTY() float f4 = 0.f;
+	UPROPERTY() float f5 = 0.f;
+	UPROPERTY() float f6 = 0.f;
+	UPROPERTY() float f7 = 0.f;
+	UPROPERTY() float f8 = 0.f;
+	UPROPERTY() float f9 = 0.f;
+
+	// Vector sources for subproperty-binding benchmark (bind .X → f0..f9)
+	UPROPERTY() FVector v0 = FVector::ZeroVector;
+	UPROPERTY() FVector v1 = FVector::ZeroVector;
+	UPROPERTY() FVector v2 = FVector::ZeroVector;
+	UPROPERTY() FVector v3 = FVector::ZeroVector;
+	UPROPERTY() FVector v4 = FVector::ZeroVector;
+	UPROPERTY() FVector v5 = FVector::ZeroVector;
+	UPROPERTY() FVector v6 = FVector::ZeroVector;
+	UPROPERTY() FVector v7 = FVector::ZeroVector;
+	UPROPERTY() FVector v8 = FVector::ZeroVector;
+	UPROPERTY() FVector v9 = FVector::ZeroVector;
+};
+
+namespace UE::UAF
+{
+
+/** Test enum for FBindableEnum exercises. */
+UENUM()
+enum class EUAFAnimNodeTestEnum : uint8
+{
+	Alpha = 0,
+	Beta = 1,
+	Gamma = 2,
+};
+
+USTRUCT(DisplayName = "Test Anim Node")
+struct FUAFTestAnimNodeData : public FUAFAnimNodeData
+{
+	GENERATED_BODY()
+
+	// FUAFAnimNodeData impl
+	virtual FUAFAnimNodePtr CreateInstance(FUAFAnimGraphUpdateContext& Context) const override { return nullptr; }
+
+	UPROPERTY(EditAnywhere, Category = "NodeData")
+	FBindableBool BoolVal = true;
+
+	UPROPERTY(EditAnywhere, Category = "NodeData")
+	FBindableFloat FloatVal = 1.0f;
+
+	UPROPERTY(EditAnywhere, Category = "NodeData")
+	FBindableDouble DoubleVal = 2.0;
+
+	UPROPERTY(EditAnywhere, Category = "NodeData")
+	FBindableInt32 Int32Val = 3;
+
+	UPROPERTY(EditAnywhere, Category = "NodeData")
+	FBindableInt64 Int64Val = 4;
+
+	UPROPERTY(EditAnywhere, Category = "NodeData")
+	FBindableByte ByteVal = 5;
+
+	UPROPERTY(EditAnywhere, Category = "NodeData")
+	FBindableName NameVal = FName(TEXT("TestName"));
+
+	UPROPERTY(EditAnywhere, Category = "NodeData")
+	FBindableVector VectorVal = FVector(1.0f, 2.0f, 3.0f);
+
+	UPROPERTY(EditAnywhere, Category = "NodeData")
+	FBindableQuat QuatVal = FQuat::Identity;
+
+	UPROPERTY(EditAnywhere, Category = "NodeData")
+	FBindableObject ObjectVal = FBindableObject();
+
+	UPROPERTY(EditAnywhere, Category = "NodeData")
+	FBindableObject ActorVal = FBindableObject(AActor::StaticClass());
+
+	UPROPERTY(EditAnywhere, Category = "NodeData")
+	FBindableEnum EnumVal = FBindableEnum(EUAFAnimNodeTestEnum::Beta);
+
+	UPROPERTY(EditAnywhere, Category = "NodeData")
+	FBindableStruct StructVal = FBindableStruct(FUAFAnimNodeTestVars{});
+
+	UPROPERTY(EditAnywhere, Category = "NodeData")
+	FBindableStruct TransformTrajectory = FBindableStruct(FTransformTrajectory());
+
+	UPROPERTY(EditAnywhere, Category = "NodeData")
+	FBindableStruct TransformVal = FBindableStruct(FTransform());
+};
+
+} // namespace UE::UAF
diff --git a/Engine/Plugins/Experimental/UAF/UAFTestSuites/Source/UAFAnimNodeTestData/UAFAnimNodeTestData.Build.cs b/Engine/Plugins/Experimental/UAF/UAFTestSuites/Source/UAFAnimNodeTestData/UAFAnimNodeTestData.Build.cs
new file mode 100644
--- /dev/null
+++ b/Engine/Plugins/Experimental/UAF/UAFTestSuites/Source/UAFAnimNodeTestData/UAFAnimNodeTestData.Build.cs
@@ -0,0 +1,27 @@
+// Copyright Epic Games, Inc. All Rights Reserved.
+
+using UnrealBuildTool;
+
+public class UAFAnimNodeTestData : ModuleRules
+{
+	public UAFAnimNodeTestData(ReadOnlyTargetRules Target) : base(Target)
+	{
+		bAllowUETypesInNamespaces = true;
+
+		PublicDependencyModuleNames.AddRange(new string[]
+		{
+			"Core",
+			"CoreUObject",
+			"UAF",
+			"UAFAnimNode",
+		});
+
+		PrivateDependencyModuleNames.AddRange(new string[]
+		{
+			"UAFAnimGraph",
+			"RigVM",
+			"ControlRig",
+			"Engine",
+		});
+	}
+}
diff --git a/Engine/Plugins/Experimental/UAF/UAFTestSuites/UAFTestSuites.uplugin b/Engine/Plugins/Experimental/UAF/UAFTestSuites/UAFTestSuites.uplugin
--- a/Engine/Plugins/Experimental/UAF/UAFTestSuites/UAFTestSuites.uplugin
+++ b/Engine/Plugins/Experimental/UAF/UAFTestSuites/UAFTestSuites.uplugin
@@ -29,13 +29,22 @@
 			"Name" : "UAFAnimGraphTestSuite",
 			"Type" : "UncookedOnly",
 			"LoadingPhase" : "Default"
+		},
+		{
+			"Name" : "UAFAnimNodeTestData",
+			"Type" : "DeveloperTool",
+			"LoadingPhase" : "Default"
 		}
 	],
-	"Plugins": [		
+	"Plugins": [
 		{
 			"Name": "Workspace",
 			"Enabled": true
 		},
+		{
+			"Name": "UAFAnimNode",
+			"Enabled": true
+		},
 		{
 			"Name": "UAF",
 			"Enabled": true
