UE Editor crash when drag & drop a variable or function into a category
#fix Replace invalid .Pin() calls on TWeakObjectPtr with null-guarded .Get(); fold null checks into existing conditions in category drop handlers
#swarm https://p4-swarm.epicgames.net/reviews/52173301

#Synced-CL 52167502

diff --git a/Engine/Plugins/Experimental/UAF/UAF/Source/UAFEditor/Private/Common/Outliner/CommonOutlinerMode.cpp b/Engine/Plugins/Experimental/UAF/UAF/Source/UAFEditor/Private/Common/Outliner/CommonOutlinerMode.cpp
--- a/Engine/Plugins/Experimental/UAF/UAF/Source/UAFEditor/Private/Common/Outliner/CommonOutlinerMode.cpp
+++ b/Engine/Plugins/Experimental/UAF/UAF/Source/UAFEditor/Private/Common/Outliner/CommonOutlinerMode.cpp
@@ -74,9 +74,9 @@ FCommonOutlinerMode::FCommonOutlinerMode(SSceneOutliner* InOutliner, const TShar
 					OutValidationInfo = FSceneOutlinerDragValidationInfo(ESceneOutlinerDropCompatibility::CompatibleGeneric, FormattedMessage);
 					return FReply::Handled();
 				}
-				else if (ThisAsset != OtherAsset)
+				else if (ThisAsset && ThisAsset != OtherAsset)
 				{
-					const FText ThisAssetText = FText::FromString(DragOverItem->WeakOwner.Pin()->GetName());
+					const FText ThisAssetText = FText::FromString(ThisAsset->GetName());
 					const FText FormattedMessage = FText::Format(LOCTEXT("MoveAndSubCategory", "Move {0} to {1} and make a sub-category of {2}"), DragCategoryText, ThisAssetText, EntryCategoryText);
 					OutValidationInfo = FSceneOutlinerDragValidationInfo(ESceneOutlinerDropCompatibility::CompatibleGeneric, FormattedMessage);
 
@@ -112,9 +112,9 @@ FCommonOutlinerMode::FCommonOutlinerMode(SSceneOutliner* InOutliner, const TShar
 
 					return true;
 				}
-				else if (OtherAsset  && ThisAsset != OtherAsset)
+				else if (ThisAsset && OtherAsset && ThisAsset != OtherAsset)
 				{
-					const FText ThisAssetText = FText::FromString(DroppedOnItem->WeakOwner.Pin()->GetName());
+					const FText ThisAssetText = FText::FromString(ThisAsset->GetName());
 					const FText FormattedMessage = FText::Format(LOCTEXT("MoveAndSubCategory", "Move {0} to {1} and make a sub-category of {2}"), DropCategoryText, ThisAssetText, EntryCategoryText);
 
 					FScopedTransaction Transaction(FormattedMessage);
@@ -285,7 +285,7 @@ FCommonOutlinerMode::FCommonOutlinerMode(SSceneOutliner* InOutliner, const TShar
 			}
 			else if (ThisAsset && OtherAsset)
 			{
-				const FText ThisAssetText = FText::FromString(DroppedOnItem->WeakOwner.Pin()->GetName());
+				const FText ThisAssetText = FText::FromString(ThisAsset->GetName());
 				const FText FormattedMessage = FText::Format(LOCTEXT("MoveReorderCategory", "Move {0} to {1} and reorder before of {2}"), OtherCategoryText, ThisAssetText, ThisCategoryText);
 
 				FScopedTransaction Transaction(FormattedMessage);
diff --git a/Engine/Plugins/Experimental/UAF/UAF/Source/UAFEditor/Private/Functions/Outliner/FunctionsOutlinerMode.cpp b/Engine/Plugins/Experimental/UAF/UAF/Source/UAFEditor/Private/Functions/Outliner/FunctionsOutlinerMode.cpp
--- a/Engine/Plugins/Experimental/UAF/UAF/Source/UAFEditor/Private/Functions/Outliner/FunctionsOutlinerMode.cpp
+++ b/Engine/Plugins/Experimental/UAF/UAF/Source/UAFEditor/Private/Functions/Outliner/FunctionsOutlinerMode.cpp
@@ -107,7 +107,7 @@ FFunctionsOutlinerMode::FFunctionsOutlinerMode(SCommonOutliner* InFunctionsOutli
 				const UUAFRigVMAssetEditorData* VariableEditorData = CollapseNode->GetTypedOuter<UUAFRigVMAssetEditorData>();
 				const UUAFRigVMAssetEditorData* CategoryEditorData = UncookedOnly::FUtils::GetEditorData(DragOverItem->WeakOwner.Get());
 
-				if (VariableEditorData == CategoryEditorData)
+				if (VariableEditorData && VariableEditorData == CategoryEditorData)
 				{
 					const FText FormattedMessage = FText::Format(LOCTEXT("MoveFuncToCategoryFormat", "Moving {0} to Category {1}"), OtherFuncName, ThisFuncCategoryText);	
 					FScopedTransaction Transaction(FormattedMessage);
diff --git a/Engine/Plugins/Experimental/UAF/UAF/Source/UAFEditor/Private/Variables/Outliner/VariablesOutlinerMode.cpp b/Engine/Plugins/Experimental/UAF/UAF/Source/UAFEditor/Private/Variables/Outliner/VariablesOutlinerMode.cpp
--- a/Engine/Plugins/Experimental/UAF/UAF/Source/UAFEditor/Private/Variables/Outliner/VariablesOutlinerMode.cpp
+++ b/Engine/Plugins/Experimental/UAF/UAF/Source/UAFEditor/Private/Variables/Outliner/VariablesOutlinerMode.cpp
@@ -289,29 +289,32 @@ FVariablesOutlinerMode::FVariablesOutlinerMode(SCommonOutliner* InOutliner, cons
 				UUAFRigVMAsset* VariableAsset = VariableEntry->GetTypedOuter<UUAFRigVMAsset>();
 				UUAFRigVMAsset* CategoryAsset = DroppedOnCategory->WeakOwner.Get();
 
-				const FText OtherVariableNameText = FText::FromName(VariableEntry->GetVariableName());
-				const FText ThisVariableCategoryText = FText::FromString(DroppedOnCategory->GetDisplayString());
-				const FText FormattedMessage = FText::Format(LOCTEXT("MoveVariableToCategoryFormat", "Moving {0} to Category {1}"), OtherVariableNameText, ThisVariableCategoryText);	
-				FScopedTransaction Transaction(FormattedMessage);
-
-				UAnimNextVariableEntry* NewVariableEntry = VariableEntry;
-					
-				if (VariableAsset != CategoryAsset)
+				if (CategoryAsset)
 				{
-					const FName VariableEntryName = VariableEntry->GetVariableName();
-					UncookedOnly::FUtils::MoveVariableToAsset(VariableEntry, CategoryAsset, true, true);
-					NewVariableEntry = Cast<UAnimNextVariableEntry>(UncookedOnly::FUtils::GetEditorData(CategoryAsset)->FindEntry(VariableEntryName));
-				}
-					
-				check(NewVariableEntry);
+					const FText OtherVariableNameText = FText::FromName(VariableEntry->GetVariableName());
+					const FText ThisVariableCategoryText = FText::FromString(DroppedOnCategory->GetDisplayString());
+					const FText FormattedMessage = FText::Format(LOCTEXT("MoveVariableToCategoryFormat", "Moving {0} to Category {1}"), OtherVariableNameText, ThisVariableCategoryText);	
+					FScopedTransaction Transaction(FormattedMessage);
+
+					UAnimNextVariableEntry* NewVariableEntry = VariableEntry;
+						
+					if (VariableAsset != CategoryAsset)
+					{
+						const FName VariableEntryName = VariableEntry->GetVariableName();
+						UncookedOnly::FUtils::MoveVariableToAsset(VariableEntry, CategoryAsset, true, true);
+						NewVariableEntry = Cast<UAnimNextVariableEntry>(UncookedOnly::FUtils::GetEditorData(CategoryAsset)->FindEntry(VariableEntryName));
+					}
+						
+					check(NewVariableEntry);
+						
+					// Only add if not already part of category
+					if (NewVariableEntry->GetVariableCategory() != DroppedOnCategory->CategoryPath)
+					{
+						NewVariableEntry->SetVariableCategory(DroppedOnCategory->CategoryPath);
+					}
 					
-				// Only add if not already part of category
-				if (NewVariableEntry->GetVariableCategory() != DroppedOnCategory->CategoryPath)
-				{
-					NewVariableEntry->SetVariableCategory(DroppedOnCategory->CategoryPath);
+					return true;
 				}
-				
-				return true;
 			}
 		}
 		
