[RewindDebugger]

diff --git a/Engine/Plugins/Chooser/Source/ChooserEditor/Private/ChooserTrack.cpp b/Engine/Plugins/Chooser/Source/ChooserEditor/Private/ChooserTrack.cpp
--- a/Engine/Plugins/Chooser/Source/ChooserEditor/Private/ChooserTrack.cpp
+++ b/Engine/Plugins/Chooser/Source/ChooserEditor/Private/ChooserTrack.cpp
@@ -4,6 +4,7 @@
 
 #include "Chooser.h"
 #include "ChooserProvider.h"
+#include "Common/ProviderLock.h"
 #include "Editor.h"
 #include "GameFramework/Actor.h"
 #include "IGameplayProvider.h"
@@ -23,9 +24,17 @@ FChooserTrack::FChooserTrack(uint64 InObjectId, uint64 InChooserId) :
 	Icon = FSlateIcon("EditorStyle", "Sequencer.Tracks.Event", "Sequencer.Tracks.Event");
 	
 	const IGameplayProvider* GameplayProvider = IRewindDebugger::Instance()->GetAnalysisSession()->ReadProvider<IGameplayProvider>("GameplayProvider");
-	const FObjectInfo& ChooserInfo = GameplayProvider->GetObjectInfo(InChooserId);
-	TrackName = FText::FromString(ChooserInfo.Name);
-	ChooserTable = FindObject<UChooserTable>(nullptr, ChooserInfo.PathName);
+	FString ChooserName;
+	FString ChooserPathName;
+	{
+		TraceServices::FProviderReadScopeLock ProviderReadScope(*GameplayProvider);
+		const FObjectInfo& ChooserInfo = GameplayProvider->GetObjectInfo(InChooserId);
+		ChooserName = ChooserInfo.Name;
+		ChooserPathName = ChooserInfo.PathName;
+	}
+	TrackName = FText::FromString(ChooserName);
+	ChooserTable = FindObject<UChooserTable>(nullptr, *ChooserPathName);
+
 }
 
 FChooserTrack::~FChooserTrack()
@@ -76,8 +85,7 @@ bool FChooserTrack::UpdateInternal()
 		EventData->Points.SetNum(0,EAllowShrinking::No);
 		EventData->Windows.SetNum(0);
 
-		TraceServices::FAnalysisSessionReadScope SessionReadScope(*AnalysisSession);
-		
+		TraceServices::FProviderReadScopeLock ProviderReadScope(*ChooserProvider);
 		ChooserProvider->ReadChooserEvaluationTimeline(ObjectId, [this, StartTime, EndTime](const FChooserProvider::ChooserEvaluationTimeline& InTimeline)
 		{
 			InTimeline.EnumerateEvents(StartTime, EndTime, [this, StartTime, EndTime](double InStartTime, double InEndTime, uint32 InDepth, const FChooserEvaluationData& ChooserEvaluationData)
@@ -103,45 +111,56 @@ bool FChooserTrack::HandleDoubleClickInternal()
 	const IRewindDebugger* RewindDebugger = IRewindDebugger::Instance();
 	if (const TraceServices::IAnalysisSession* AnalysisSession = RewindDebugger->GetAnalysisSession())
 	{
-		TraceServices::FAnalysisSessionReadScope SessionReadScope(*AnalysisSession);
-
 		const IGameplayProvider* GameplayProvider = AnalysisSession->ReadProvider<IGameplayProvider>("GameplayProvider");
-		const FObjectInfo& AssetInfo = GameplayProvider->GetObjectInfo(GetChooserId());
+		if (!GameplayProvider)
+		{
+			return false;
+		}
 
-		// attach chooser table editor debugging
-		if (ChooserTable)
+		FString AssetPathName;
+		FString DebugName;
 		{
-			UChooserTable* RootChooser = ChooserTable->GetRootChooser();
-			
-			const FObjectInfo& OwnerObjectInfo  = GameplayProvider->GetObjectInfo(ObjectId);
-			FString DebugName(OwnerObjectInfo.Name);
-			if (const FObjectInfo* ActorInfo = RewindDebugger->FindTypedOuterInfo<AActor>(GameplayProvider, ObjectId))
-			{
-				DebugName += " in " + FString(ActorInfo->Name);
-			}
-			
-			if (const FWorldInfo* WorldInfo = GameplayProvider->FindWorldInfoFromObject(ObjectId))
-			{
-				const FObjectInfo& WorldObjectInfo = GameplayProvider->GetObjectInfo(WorldInfo->Id);
+			TraceServices::FProviderReadScopeLock ProviderReadScope(*GameplayProvider);
+			const FObjectInfo& AssetInfo = GameplayProvider->GetObjectInfo(GetChooserId());
+			AssetPathName = AssetInfo.PathName;
 
-				if (WorldInfo->NetMode == FWorldInfo::ENetMode::DedicatedServer)
+			// attach chooser table editor debugging
+			if (ChooserTable)
+			{
+				const FObjectInfo& OwnerObjectInfo  = GameplayProvider->GetObjectInfo(ObjectId);
+				DebugName = FString(OwnerObjectInfo.Name);
+				if (const FObjectInfo* ActorInfo = RewindDebugger->FindTypedOuterInfo<AActor>(GameplayProvider, ObjectId))
 				{
-					DebugName += " (Server)";
+					DebugName += " in " + FString(ActorInfo->Name);
 				}
-				else if (WorldInfo->NetMode == FWorldInfo::ENetMode::Client)
+
+				if (const FWorldInfo* WorldInfo = GameplayProvider->FindWorldInfoFromObject(ObjectId))
 				{
-					DebugName += FString(" (Client ") + FString::FromInt(WorldInfo->PIEInstanceId) + ")";
+					const FObjectInfo& WorldObjectInfo = GameplayProvider->GetObjectInfo(WorldInfo->Id);
+
+					if (WorldInfo->NetMode == FWorldInfo::ENetMode::DedicatedServer)
+					{
+						DebugName += " (Server)";
+					}
+					else if (WorldInfo->NetMode == FWorldInfo::ENetMode::Client)
+					{
+						DebugName += FString(" (Client ") + FString::FromInt(WorldInfo->PIEInstanceId) + ")";
+					}
 				}
 			}
-			
+		}
+
+		if (ChooserTable)
+		{
+			UChooserTable* RootChooser = ChooserTable->GetRootChooser();
 			RootChooser->SetDebugTarget(DebugName);
 			RootChooser->SetEnableDebugTesting(true);
-			
+
 			GEditor->GetEditorSubsystem<UAssetEditorSubsystem>()->OpenEditorForAsset(ChooserTable);
 		}
 		else
 		{
-			GEditor->GetEditorSubsystem<UAssetEditorSubsystem>()->OpenEditorForAsset(AssetInfo.PathName);
+			GEditor->GetEditorSubsystem<UAssetEditorSubsystem>()->OpenEditorForAsset(AssetPathName);
 		}
 
 		return true;
@@ -208,24 +227,25 @@ bool FChoosersTrack::UpdateInternal()
 	if(ChooserProvider)
 	{
 		TArray<uint64> UniqueTrackIds;
-		
-		TraceServices::FAnalysisSessionReadScope SessionReadScope(*AnalysisSession);
-		
-		ChooserProvider->ReadChooserEvaluationTimeline(ObjectId, [this, StartTime, EndTime, &UniqueTrackIds](const FChooserProvider::ChooserEvaluationTimeline& InTimeline)
 		{
-			InTimeline.EnumerateEvents(StartTime, EndTime, [this, &UniqueTrackIds](double InStartTime, double InEndTime, uint32 InDepth, const FChooserEvaluationData& ChooserEvaluationData)
+			TraceServices::FProviderReadScopeLock ProviderReadScope(*ChooserProvider);
+
+			ChooserProvider->ReadChooserEvaluationTimeline(ObjectId, [this, StartTime, EndTime, &UniqueTrackIds](const FChooserProvider::ChooserEvaluationTimeline& InTimeline)
 			{
-				UniqueTrackIds.AddUnique(ChooserEvaluationData.ChooserId);
-		 		return TraceServices::EEventEnumerate::Continue;
-		 	});
-		});
-		
+				InTimeline.EnumerateEvents(StartTime, EndTime, [this, &UniqueTrackIds](double InStartTime, double InEndTime, uint32 InDepth, const FChooserEvaluationData& ChooserEvaluationData)
+				{
+					UniqueTrackIds.AddUnique(ChooserEvaluationData.ChooserId);
+					return TraceServices::EEventEnumerate::Continue;
+				});
+			});
+		}
+
 		UniqueTrackIds.StableSort();
 		const int32 TrackCount = UniqueTrackIds.Num();
-		
+
 		if (Children.Num() != TrackCount)
 			bChanged = true;
-		
+
 		Children.SetNum(UniqueTrackIds.Num());
 		for(int i = 0; i < TrackCount; i++)
 		{
@@ -234,7 +254,7 @@ bool FChoosersTrack::UpdateInternal()
 				Children[i] = MakeShared<FChooserTrack>(ObjectId, UniqueTrackIds[i]);
 				bChanged = true;
 			}
-		
+
 			if (Children[i]->Update())
 			{
 				bChanged = true;
@@ -255,10 +275,10 @@ bool FChoosersTrackCreator::HasDebugInfoInternal(const RewindDebugger::FObjectId
 	TRACE_CPUPROFILER_EVENT_SCOPE(FChoosersTrack::HasDebugInfoInternal);
 	const TraceServices::IAnalysisSession* AnalysisSession = IRewindDebugger::Instance()->GetAnalysisSession();
 	
-	TraceServices::FAnalysisSessionReadScope SessionReadScope(*AnalysisSession);
 	bool bHasData = false;
 	if (const FChooserProvider* ChooserProvider = AnalysisSession->ReadProvider<FChooserProvider>(FChooserProvider::ProviderName))
 	{
+		TraceServices::FProviderReadScopeLock ProviderReadScope(*ChooserProvider);
 		ChooserProvider->ReadChooserEvaluationTimeline(InObjectId.GetMainId(), [&bHasData](const FChooserProvider::ChooserEvaluationTimeline& InTimeline)
 	 	{
 	 		bHasData = true;
diff --git a/Engine/Plugins/Chooser/Source/ChooserEditor/Private/RewindDebuggerChooser.cpp b/Engine/Plugins/Chooser/Source/ChooserEditor/Private/RewindDebuggerChooser.cpp
--- a/Engine/Plugins/Chooser/Source/ChooserEditor/Private/RewindDebuggerChooser.cpp
+++ b/Engine/Plugins/Chooser/Source/ChooserEditor/Private/RewindDebuggerChooser.cpp
@@ -4,6 +4,7 @@
 
 #include "Chooser.h"
 #include "ChooserProvider.h"
+#include "Common/ProviderLock.h"
 #include "GameFramework/Actor.h"
 #include "IGameplayProvider.h"
 #include "TraceServices/Model/Frames.h"
@@ -16,15 +17,27 @@ FRewindDebuggerChooser::FRewindDebuggerChooser()
 void FRewindDebuggerChooser::Update(float DeltaTime, IRewindDebugger* RewindDebugger)
 {
 	const TraceServices::IAnalysisSession* AnalysisSession = RewindDebugger->GetAnalysisSession();
-	TraceServices::FAnalysisSessionReadScope SessionReadScope(*AnalysisSession);
-	
 	const FChooserProvider* ChooserProvider = AnalysisSession->ReadProvider<FChooserProvider>(FChooserProvider::ProviderName);
 	const IGameplayProvider* GameplayProvider = AnalysisSession->ReadProvider<IGameplayProvider>("GameplayProvider");
-	
-	const TraceServices::IFrameProvider& FrameProvider = TraceServices::ReadFrameProvider(*RewindDebugger->GetAnalysisSession());
+	if (!ChooserProvider || !GameplayProvider)
+	{
+		return;
+	}
+
+	const float CurrentTraceTime = IRewindDebugger::Instance()->CurrentTraceTime();
+
+	bool bFrameFound;
 	TraceServices::FFrame Frame;
-	if (FrameProvider.GetFrameFromTime(ETraceFrameType::TraceFrameType_Game, RewindDebugger->CurrentTraceTime(), Frame))
 	{
+		TraceServices::FAnalysisSessionReadScope SessionReadScope(*AnalysisSession);
+		const TraceServices::IFrameProvider& FrameProvider = TraceServices::ReadFrameProvider(*AnalysisSession);
+		bFrameFound = FrameProvider.GetFrameFromTime(TraceFrameType_Game, CurrentTraceTime, Frame);
+	}
+
+	if (bFrameFound)
+	{
+		TraceServices::FProviderReadScopeLock ChooserProviderReadScope(*ChooserProvider);
+		TraceServices::FProviderReadScopeLock GameplayProviderReadScope(*GameplayProvider);
 		ChooserProvider->EnumerateChooserEvaluationTimelines([GameplayProvider, ChooserProvider, RewindDebugger, Frame](uint64 OwnerId, const FChooserProvider::ChooserEvaluationTimeline& ChooserEvaluationTimeline)
 		{
 			double StartTime = RewindDebugger->GetScrubTime();
diff --git a/Engine/Plugins/Developer/FieldNotificationTrace/Source/FieldNotificationTraceEditor/Private/FieldNotificationRewindDebugger.cpp b/Engine/Plugins/Developer/FieldNotificationTrace/Source/FieldNotificationTraceEditor/Private/FieldNotificationRewindDebugger.cpp
--- a/Engine/Plugins/Developer/FieldNotificationTrace/Source/FieldNotificationTraceEditor/Private/FieldNotificationRewindDebugger.cpp
+++ b/Engine/Plugins/Developer/FieldNotificationTrace/Source/FieldNotificationTraceEditor/Private/FieldNotificationRewindDebugger.cpp
@@ -11,14 +11,16 @@ namespace UE::FieldNotification
 void FRewindDebuggerExtension::Update(float DeltaTime, IRewindDebugger* RewindDebugger)
 {
 	const TraceServices::IAnalysisSession* AnalysisSession = RewindDebugger->GetAnalysisSession();
-	TraceServices::FAnalysisSessionReadScope SessionReadScope(*AnalysisSession);
-	
-	const FTraceProvider* TraceProvider = AnalysisSession->ReadProvider<FTraceProvider>(FTraceProvider::ProviderName);
-	//const IGameplayProvider* GameplayProvider = AnalysisSession->ReadProvider<IGameplayProvider>("GameplayProvider");
-	
-	const TraceServices::IFrameProvider& FrameProvider = TraceServices::ReadFrameProvider(*RewindDebugger->GetAnalysisSession());
+
+	bool bFrameFound;
 	TraceServices::FFrame Frame;
-	if (FrameProvider.GetFrameFromTime(TraceFrameType_Game, RewindDebugger->CurrentTraceTime(), Frame))
+	{
+		TraceServices::FAnalysisSessionReadScope SessionReadScope(*AnalysisSession);
+		const TraceServices::IFrameProvider& FrameProvider = TraceServices::ReadFrameProvider(*AnalysisSession);
+		bFrameFound = FrameProvider.GetFrameFromTime(TraceFrameType_Game, RewindDebugger->CurrentTraceTime(), Frame);
+	}
+
+	if (bFrameFound)
 	{
 		// each tick update the UMG Preview Window with what we are trying to debug
 	}
diff --git a/Engine/Plugins/Developer/FieldNotificationTrace/Source/FieldNotificationTraceEditor/Private/FieldNotificationTrack.cpp b/Engine/Plugins/Developer/FieldNotificationTrace/Source/FieldNotificationTraceEditor/Private/FieldNotificationTrack.cpp
--- a/Engine/Plugins/Developer/FieldNotificationTrace/Source/FieldNotificationTraceEditor/Private/FieldNotificationTrack.cpp
+++ b/Engine/Plugins/Developer/FieldNotificationTrace/Source/FieldNotificationTraceEditor/Private/FieldNotificationTrack.cpp
@@ -2,6 +2,7 @@
 
 #include "FieldNotificationTrack.h"
 
+#include "Common/ProviderLock.h"
 #include "Editor.h"
 #include "FieldNotificationTraceProvider.h"
 #include "IGameplayProvider.h"
@@ -110,9 +111,9 @@ bool FFieldNotifyTrack::HandleDoubleClickInternal()
 	const IRewindDebugger* RewindDebugger = IRewindDebugger::Instance();
 	if (const TraceServices::IAnalysisSession* AnalysisSession = RewindDebugger->GetAnalysisSession())
 	{
-		TraceServices::FAnalysisSessionReadScope SessionReadScope(*AnalysisSession);
-
 		const IGameplayProvider* GameplayProvider = AnalysisSession->ReadProvider<IGameplayProvider>("GameplayProvider");
+		TraceServices::FProviderReadScopeLock ProviderReadScope(*GameplayProvider);
+
 		//const FObjectInfo& AssetInfo = GameplayProvider->GetObjectInfo(GetId());
 
 		// attach the viewmodel to the UMG Preview Window for debugging
diff --git a/Engine/Plugins/Experimental/Animation/TrajectoryTools/Source/TrajectoryTools/Private/TrajectoryRewindDebuggerExtension.cpp b/Engine/Plugins/Experimental/Animation/TrajectoryTools/Source/TrajectoryTools/Private/TrajectoryRewindDebuggerExtension.cpp
--- a/Engine/Plugins/Experimental/Animation/TrajectoryTools/Source/TrajectoryTools/Private/TrajectoryRewindDebuggerExtension.cpp
+++ b/Engine/Plugins/Experimental/Animation/TrajectoryTools/Source/TrajectoryTools/Private/TrajectoryRewindDebuggerExtension.cpp
@@ -2,6 +2,7 @@
 
 #include "TrajectoryRewindDebuggerExtension.h"
 
+#include "Common/ProviderLock.h"
 #include "ToolMenus.h"
 #include "IAnimationProvider.h"
 #include "IGameplayProvider.h"
@@ -126,8 +127,8 @@ void FRewindDebuggerTrajectory::UpdateState(IRewindDebugger* RewindDebugger)
 
 		if (GameplayProvider && AnimationProvider)
 		{
-			TraceServices::FAnalysisSessionReadScope AnalysisSessionReadScope(*Session);
-			
+			TraceServices::FProviderReadScopeLock GameplayProviderReadScope(*GameplayProvider);
+
 			// Note that recording index is always zero since RewindDebugger starts a new trace for every recording.
 			if (const IGameplayProvider::RecordingInfoTimeline* Recording = GameplayProvider->GetRecordingInfo(0))
 			{
@@ -149,8 +150,11 @@ void FRewindDebuggerTrajectory::UpdateState(IRewindDebugger* RewindDebugger)
 			}
 			
 			BuildTrajectoryOwnersList(RewindDebugger, GameplayProvider, State.Extract.TraceStarTime, State.Extract.TraceEndTime, State.ObjectInfos);
-			BuildTrajectorySkeletalMeshInfoList(GameplayProvider, AnimationProvider, State.Extract.TraceStarTime, State.Extract.TraceEndTime, State.ObjectInfos, State.SkelMeshInfos);
-			BuildTrajectories(GameplayProvider, AnimationProvider, State.Extract.TraceStarTime, State.Extract.TraceEndTime, State.ObjectInfos, State.Trajectories);
+			{
+				TraceServices::FProviderReadScopeLock AnimationProviderReadScope(*AnimationProvider);
+				BuildTrajectorySkeletalMeshInfoList(GameplayProvider, AnimationProvider, State.Extract.TraceStarTime, State.Extract.TraceEndTime, State.ObjectInfos, State.SkelMeshInfos);
+				BuildTrajectories(GameplayProvider, AnimationProvider, State.Extract.TraceStarTime, State.Extract.TraceEndTime, State.ObjectInfos, State.Trajectories);
+			}
 			UpdateDebugInfos(GameplayProvider, State.ObjectInfos, State.DebugInfos);
 		}
 	}
diff --git a/Engine/Plugins/Experimental/UAF/UAF/Source/UAFEditor/Private/RewindDebugger/AnimNextAnalyzer.cpp b/Engine/Plugins/Experimental/UAF/UAF/Source/UAFEditor/Private/RewindDebugger/AnimNextAnalyzer.cpp
--- a/Engine/Plugins/Experimental/UAF/UAF/Source/UAFEditor/Private/RewindDebugger/AnimNextAnalyzer.cpp
+++ b/Engine/Plugins/Experimental/UAF/UAF/Source/UAFEditor/Private/RewindDebugger/AnimNextAnalyzer.cpp
@@ -7,6 +7,7 @@
 #include "HAL/LowLevelMemTracker.h"
 #include "Serialization/MemoryReader.h"
 #include "TraceServices/Model/AnalysisSession.h"
+#include "Common/ProviderLock.h"
 #include "AnimNextProvider.h"
 #include "TraceServices/Utils.h"
 #include "StructUtils/PropertyBag.h"
@@ -32,7 +33,7 @@ bool FUAFAnalyzer::OnEvent(uint16 RouteId, EStyle Style, const FOnEventContext&
 {
 	LLM_SCOPE_BYNAME(TEXT("Insights/FUAFAnalyzer"));
 
-	TraceServices::FAnalysisSessionEditScope _(Session);
+	TraceServices::FProviderEditScopeLock ProviderEditScope(Provider);
 
 	const auto& EventData = Context.EventData;
 	switch (RouteId)
diff --git a/Engine/Plugins/Experimental/UAF/UAF/Source/UAFEditor/Private/RewindDebugger/AnimNextModuleTrack.cpp b/Engine/Plugins/Experimental/UAF/UAF/Source/UAFEditor/Private/RewindDebugger/AnimNextModuleTrack.cpp
--- a/Engine/Plugins/Experimental/UAF/UAF/Source/UAFEditor/Private/RewindDebugger/AnimNextModuleTrack.cpp
+++ b/Engine/Plugins/Experimental/UAF/UAF/Source/UAFEditor/Private/RewindDebugger/AnimNextModuleTrack.cpp
@@ -8,6 +8,7 @@
 
 #if UAF_TRACE_ENABLED
 #include "AnimNextProvider.h"
+#include "Common/ProviderLock.h"
 #include "Editor.h"
 #include "IGameplayProvider.h"
 #include "IRewindDebugger.h"
@@ -36,12 +37,11 @@ FText FAnimNextModuleTrack::GetDisplayNameInternal() const
 {
 	if (DisplayNameCache.IsEmpty())
 	{
-		TraceServices::FAnalysisSessionReadScope SessionReadScope(*IRewindDebugger::Instance()->GetAnalysisSession());
-		const FUAFProvider* UAFProvider = IRewindDebugger::Instance()->GetAnalysisSession()->ReadProvider<FUAFProvider>("UAFProvider");
 		const IGameplayProvider* GameplayProvider = IRewindDebugger::Instance()->GetAnalysisSession()->ReadProvider<IGameplayProvider>("GameplayProvider");
 
-		if (InstanceId != 0)
+		if (InstanceId != 0 && GameplayProvider)
 		{
+			TraceServices::FProviderReadScopeLock ProviderReadScope(*GameplayProvider);
 			const FObjectInfo& ModuleInfo = GameplayProvider->GetObjectInfo(InstanceId);
 			DisplayNameCache = FText::FromString(ModuleInfo.Name);
 			return DisplayNameCache;
@@ -120,7 +120,7 @@ bool FAnimNextModuleTrack::UpdateInternal()
 
 	bool bChanged = false;
 
-	if (const FUAFProvider* UAFProvider = IRewindDebugger::Instance()->GetAnalysisSession()->ReadProvider<FUAFProvider>("UAFProvider"))
+	if (const FUAFProvider* UAFProvider = AnalysisSession->ReadProvider<FUAFProvider>("UAFProvider"))
 	{
 		double CurrentScrubTime = IRewindDebugger::Instance()->CurrentTraceTime();
 
@@ -136,21 +136,25 @@ bool FAnimNextModuleTrack::UpdateInternal()
 
 		if (InstanceId != 0)
 		{
+			TraceServices::FProviderReadScopeLock UAFProviderReadScope(*UAFProvider);
 			if (const FDataInterfaceData* Data = UAFProvider->GetDataInterfaceData(InstanceId))
 			{
 				if (PreviousScrubTime != CurrentScrubTime)
 				{
 					PreviousScrubTime = CurrentScrubTime;
 
-					const TraceServices::IFrameProvider& FramesProvider = TraceServices::ReadFrameProvider(*AnalysisSession);
-					TraceServices::FAnalysisSessionReadScope SessionReadScope(*AnalysisSession);
-
-
 					DetailsObject->Properties.SetNum(0, EAllowShrinking::No);
 					DetailsObject->NativeProperties.SetNum(0, EAllowShrinking::No);
 
+					bool bFrameFound;
 					TraceServices::FFrame MarkerFrame;
-					if (FramesProvider.GetFrameFromTime(ETraceFrameType::TraceFrameType_Game, CurrentScrubTime, MarkerFrame))
+					{
+						TraceServices::FAnalysisSessionReadScope SessionReadScope(*AnalysisSession);
+						const TraceServices::IFrameProvider& FrameProvider = TraceServices::ReadFrameProvider(*AnalysisSession);
+						bFrameFound = FrameProvider.GetFrameFromTime(TraceFrameType_Game, CurrentScrubTime, MarkerFrame);
+					}
+		
+					if (bFrameFound)
 					{
 						Data->VariablesTimeline.EnumerateEvents(MarkerFrame.StartTime, MarkerFrame.EndTime, [UAFProvider, DetailsObject, AnalysisSession, this](double InStartTime, double InEndTime, uint32 InDepth, const FPropertyVariableData& VariableListData)
 						{
@@ -223,13 +227,13 @@ bool FAnimNextModuleTrack::HandleDoubleClickInternal()
 	{
 		FString AssetPath;
 		{
-			TraceServices::FAnalysisSessionReadScope SessionReadScope(*AnalysisSession);
-
 			const IGameplayProvider* GameplayProvider = AnalysisSession->ReadProvider<IGameplayProvider>("GameplayProvider");
 			const FUAFProvider* AnimNextProvider = IRewindDebugger::Instance()->GetAnalysisSession()->ReadProvider<FUAFProvider>("UAFProvider");
 
+			TraceServices::FProviderReadScopeLock AnimNextProviderReadScope(*AnimNextProvider);
 			if (const FDataInterfaceData* ModuleData = AnimNextProvider->GetDataInterfaceData(InstanceId))
 			{
+				TraceServices::FProviderReadScopeLock GameplayProviderReadScope(*GameplayProvider);
 				const FObjectInfo& AssetInfo = GameplayProvider->GetObjectInfo(ModuleData->AssetId);
 				if (!EnumHasAnyFlags(AssetInfo.Flags, EObjectInfoFlags::TransientObject))
 				{
diff --git a/Engine/Plugins/Experimental/UAF/UAF/Source/UAFEditor/Private/RewindDebugger/AnimNextProvider.cpp b/Engine/Plugins/Experimental/UAF/UAF/Source/UAFEditor/Private/RewindDebugger/AnimNextProvider.cpp
--- a/Engine/Plugins/Experimental/UAF/UAF/Source/UAFEditor/Private/RewindDebugger/AnimNextProvider.cpp
+++ b/Engine/Plugins/Experimental/UAF/UAF/Source/UAFEditor/Private/RewindDebugger/AnimNextProvider.cpp
@@ -4,6 +4,11 @@
 #if UAF_TRACE_ENABLED
 #include "ObjectTrace.h"
 
+namespace TraceServices
+{
+thread_local FProviderLock::FThreadLocalState GUAFProviderLockState;
+}
+
 FName FUAFProvider::ProviderName("UAFProvider");
 
 #define LOCTEXT_NAMESPACE "UAFProvider"
@@ -15,7 +20,7 @@ FUAFProvider::FUAFProvider(TraceServices::IAnalysisSession& InSession)
 
 void FUAFProvider::AppendInstance(uint64 InstanceId, uint64 HostInstanceId, uint64 AssetId, uint64 OuterObjectId)
 {
-	Session.WriteAccessCheck();
+	EditAccessCheck();
 
 	if (HostInstanceId == OuterObjectId)
 	{
@@ -36,7 +41,7 @@ void FUAFProvider::AppendInstance(uint64 InstanceId, uint64 HostInstanceId, uint
 
 void FUAFProvider::AppendVariables(double ProfileTime, double RecordingTime, uint64 DataInterfaceId, EPropertyVariableDataType PropertyVariableDataType, uint32 PropertyDescriptionHash, const TArrayView<const uint8>& VariableData)
 {
-	Session.WriteAccessCheck();
+	EditAccessCheck();
 	
 	if (const TSharedRef<FDataInterfaceData>* Data = DataInterfaceData.Find(DataInterfaceId))
 	{
@@ -54,7 +59,7 @@ void FUAFProvider::AppendVariables(double ProfileTime, double RecordingTime, uin
 
 void FUAFProvider::AppendVariableDescriptions(uint32 PropertyDescriptionHash, const TArrayView<const uint8>& VariableDescriptionData, const FString& Name)
 {
-	Session.WriteAccessCheck();
+	EditAccessCheck();
 	
 	if (HashToPropertyDescriptionDataMap.Contains(PropertyDescriptionHash) == false)
 	{
@@ -68,7 +73,7 @@ void FUAFProvider::AppendVariableDescriptions(uint32 PropertyDescriptionHash, co
 
 bool FUAFProvider::GetModuleId(uint64 ComponentId, uint64& OutModuleId) const
 {
-	Session.ReadAccessCheck();
+	ReadAccessCheck();
 	
 	if (const uint64* Id = ComponentIdToSystemId.Find(ComponentId))
 	{
@@ -80,6 +85,8 @@ bool FUAFProvider::GetModuleId(uint64 ComponentId, uint64& OutModuleId) const
 
 const FDataInterfaceData* FUAFProvider::GetDataInterfaceData(uint64 DataInterfaceId) const
 {
+	ReadAccessCheck();
+
 	const TSharedRef<FDataInterfaceData>* Data = DataInterfaceData.Find(DataInterfaceId);
 	if (Data)
 	{
@@ -90,7 +97,7 @@ const FDataInterfaceData* FUAFProvider::GetDataInterfaceData(uint64 DataInterfac
 
 void FUAFProvider::EnumerateChildInstances(uint64 InstanceId, TFunctionRef<void(const FDataInterfaceData&)> Callback) const
 {
-	Session.ReadAccessCheck();
+	ReadAccessCheck();
 
 	if (const TArray<TSharedRef<FDataInterfaceData>>* Children = HostToChildDataMap.Find(InstanceId))
 	{
@@ -105,7 +112,7 @@ void FUAFProvider::EnumerateChildInstances(uint64 InstanceId, TFunctionRef<void(
 
 void FUAFProvider::EnumerateDataInterfaces(const double InStartTime, TFunctionRef<void(const TSharedRef<const FDataInterfaceData>)> Callback) const
 {
-	Session.ReadAccessCheck();
+	ReadAccessCheck();
 
 	for (const auto& [InstanceId, InterfaceData] : DataInterfaceData)
 	{
diff --git a/Engine/Plugins/Experimental/UAF/UAF/Source/UAFEditor/Private/RewindDebugger/AnimNextProvider.h b/Engine/Plugins/Experimental/UAF/UAF/Source/UAFEditor/Private/RewindDebugger/AnimNextProvider.h
--- a/Engine/Plugins/Experimental/UAF/UAF/Source/UAFEditor/Private/RewindDebugger/AnimNextProvider.h
+++ b/Engine/Plugins/Experimental/UAF/UAF/Source/UAFEditor/Private/RewindDebugger/AnimNextProvider.h
@@ -11,8 +11,13 @@
 #include "Templates/SharedPointer.h"
 #include "Model/IntervalTimeline.h"
 #include "TraceServices/Model/AnalysisSession.h"
+#include "Common/ProviderLock.h"
 
-namespace TraceServices { class IAnalysisSession; }
+namespace TraceServices
+{
+	class IAnalysisSession;
+	extern thread_local FProviderLock::FThreadLocalState GUAFProviderLockState;
+}
 
 struct FPropertyDescriptionData
 {
@@ -60,26 +65,59 @@ struct FDataInterfaceData
 	double EndTime;
 };
 
-class FUAFProvider : public TraceServices::IProvider
+class FUAFProvider : public TraceServices::IProvider, public TraceServices::IEditableProvider
 {
 public:
 	static FName ProviderName;
 
 	FUAFProvider(TraceServices::IAnalysisSession& InSession);
 
+	//~ IProvider (read lock)
+	virtual void BeginRead() const override
+	{
+		Lock.BeginRead(TraceServices::GUAFProviderLockState);
+	}
+
+	virtual void EndRead() const override
+	{
+		Lock.EndRead(TraceServices::GUAFProviderLockState);
+	}
+
+	virtual void ReadAccessCheck() const override
+	{
+		Lock.ReadAccessCheck(TraceServices::GUAFProviderLockState);
+	}
+
+	//~ IEditableProvider (edit lock)
+	virtual void BeginEdit() const override
+	{
+		Lock.BeginWrite(TraceServices::GUAFProviderLockState);
+	}
+
+	virtual void EndEdit() const override
+	{
+		Lock.EndWrite(TraceServices::GUAFProviderLockState);
+	}
+
+	virtual void EditAccessCheck() const override
+	{
+		Lock.WriteAccessCheck(TraceServices::GUAFProviderLockState);
+	}
+
 	void AppendInstance(uint64 InstanceId, uint64 HostInstanceId,  uint64 AssetId, uint64 OuterObjectId);
 	void AppendVariables(double ProfileTime, double RecordingTime, uint64 InModuleInstanceId, EPropertyVariableDataType PropertyVariableDataType, uint32 PropertyDescriptionHash, const TArrayView<const uint8>& VariableData);
 	void AppendVariableDescriptions(uint32 PropertyDescriptionHash, const TArrayView<const uint8>& VariableDescriptionData, const FString& Name);
 
 	bool GetModuleId(uint64 ComponentId, uint64& OutModuleId) const;
 	const FDataInterfaceData* GetDataInterfaceData(uint64 DataInterfaceId) const;
-	const FPropertyDescriptionData* GetPropertyDescriptionData(uint32 PropertyDescriptionHash) const { return HashToPropertyDescriptionDataMap.Find(PropertyDescriptionHash);}
+	const FPropertyDescriptionData* GetPropertyDescriptionData(uint32 PropertyDescriptionHash) const { ReadAccessCheck(); return HashToPropertyDescriptionDataMap.Find(PropertyDescriptionHash);}
 
 	void EnumerateChildInstances(uint64 InstanceId, TFunctionRef<void(const FDataInterfaceData&)> Callback) const;
 	void EnumerateDataInterfaces(const double StartTime, TFunctionRef<void(const TSharedRef<const FDataInterfaceData>)> Callback) const;
 
 private:
 	TraceServices::IAnalysisSession& Session;
+	mutable TraceServices::FProviderLock Lock;
 
 	TMap<uint64, uint64> ComponentIdToSystemId;
 	TMap<uint64, TSharedRef<FDataInterfaceData>> DataInterfaceData;
diff --git a/Engine/Plugins/Experimental/UAF/UAF/Source/UAFEditor/Private/UAF/Debug/SVariableWatch.cpp b/Engine/Plugins/Experimental/UAF/UAF/Source/UAFEditor/Private/UAF/Debug/SVariableWatch.cpp
--- a/Engine/Plugins/Experimental/UAF/UAF/Source/UAFEditor/Private/UAF/Debug/SVariableWatch.cpp
+++ b/Engine/Plugins/Experimental/UAF/UAF/Source/UAFEditor/Private/UAF/Debug/SVariableWatch.cpp
@@ -187,28 +187,27 @@ namespace UE::UAF::Editor
 		{
 			return;
 		}
-		
+
 		const TraceServices::IAnalysisSession* AnalysisSession = RewindDebugger->GetAnalysisSession();
-		TraceServices::FAnalysisSessionReadScope SessionReadScope(*AnalysisSession);
 
 		const double CurrentScrubTime = RewindDebugger->GetScrubTime();
-		const double CurrentTraceTime  = RewindDebugger->CurrentTraceTime();
+		const double CurrentTraceTime = RewindDebugger->CurrentTraceTime();
 
 		RootDebugInstances.Empty();
 		DebugInstances.Empty();
 
-		const TraceServices::IFrameProvider& FrameProvider = TraceServices::ReadFrameProvider(*RewindDebugger->GetAnalysisSession());
-
 		TraceServices::FFrame Frame;
-		if (!FrameProvider.GetFrameFromTime(ETraceFrameType::TraceFrameType_Game, CurrentTraceTime, Frame))
 		{
-			return;
+			const TraceServices::IFrameProvider& FrameProvider = TraceServices::ReadFrameProvider(*RewindDebugger->GetAnalysisSession());
+			TraceServices::FAnalysisSessionReadScope SessionReadScope(*AnalysisSession);
+			if (!FrameProvider.GetFrameFromTime(ETraceFrameType::TraceFrameType_Game, CurrentTraceTime, Frame))
+			{
+				return;
+			}
 		}
 
-		const IGameplayProvider* GameplayProvider = AnalysisSession->ReadProvider<IGameplayProvider>("GameplayProvider");
 		const FUAFProvider* UAFProvider = AnalysisSession->ReadProvider<FUAFProvider>("UAFProvider");
-		
-		if (!GameplayProvider || !UAFProvider)
+		if (!UAFProvider)
 		{
 			return;
 		}
@@ -226,61 +225,64 @@ namespace UE::UAF::Editor
 			WorkspaceDebugObject = EditorMode->GetDebugObjectSelector()->GetDebugObject().Get();
 		}
 
-		UAFProvider->EnumerateDataInterfaces(CurrentScrubTime, [this, FocusedAssetPath, UAFProvider, WorkspaceDebugObject](const TSharedRef<const FDataInterfaceData> DataInterfaceData)
 		{
-			UObject* Asset = FObjectTrace::GetObjectFromId(DataInterfaceData->AssetId);
-			if (FSoftObjectPath(Asset) != FocusedAssetPath)
-			{
-				return;
-			}
-			
-			const UObject* OwnerObject = FObjectTrace::GetObjectFromId(DataInterfaceData->OuterObjectId);
-			if (!WorkspaceDebugObject || (WorkspaceDebugObject != OwnerObject && !OwnerObject->IsInOuter(WorkspaceDebugObject)))
-			{
-				return;
-			}
-
-			TSharedPtr<FDebugInstance> NewDebugInstance = MakeShared<FDebugInstance>();
-			NewDebugInstance->InstanceId = DataInterfaceData->InstanceId;
-			NewDebugInstance->HostInstanceId = DataInterfaceData->HostInstanceId;
-			NewDebugInstance->Asset = FObjectTrace::GetObjectFromId(DataInterfaceData->AssetId);
-			
-			DebugInstances.Add(NewDebugInstance->InstanceId, NewDebugInstance);
-			
-			uint64 HostInstanceId = NewDebugInstance->HostInstanceId;
-			const FDataInterfaceData* HostInstanceData = UAFProvider->GetDataInterfaceData(HostInstanceId);
-			while (HostInstanceData)
-			{
-				if (!DebugInstances.Contains(HostInstanceId))
-				{
-					TSharedPtr<FDebugInstance> NewHostDebugInstance = MakeShared<FDebugInstance>();
-					NewHostDebugInstance->InstanceId = HostInstanceData->InstanceId;
-					NewHostDebugInstance->HostInstanceId = HostInstanceData->HostInstanceId;
-					NewHostDebugInstance->Asset = FObjectTrace::GetObjectFromId(HostInstanceData->AssetId);
-					
-					DebugInstances.Add(NewHostDebugInstance->InstanceId, NewHostDebugInstance);
-					
-					HostInstanceId = NewHostDebugInstance->HostInstanceId;
-                    HostInstanceData = UAFProvider->GetDataInterfaceData(HostInstanceId);
-				}
-				else
+			TraceServices::FProviderReadScopeLock ProviderReadScope(*UAFProvider);
+			UAFProvider->EnumerateDataInterfaces(CurrentScrubTime, [this, FocusedAssetPath, UAFProvider, WorkspaceDebugObject](const TSharedRef<const FDataInterfaceData> DataInterfaceData)
 				{
-					break;
-				}
-			}
-			
-			if (!DebugInstances.Contains(HostInstanceId))
-			{
-				TSharedPtr<FDebugInstance> NewHostDebugInstance = MakeShared<FDebugInstance>();
-				NewHostDebugInstance->InstanceId = HostInstanceId;
-				NewHostDebugInstance->HostInstanceId = 0;
-				NewHostDebugInstance->Asset = FObjectTrace::GetObjectFromId(HostInstanceId);
-			
-				DebugInstances.Add(NewHostDebugInstance->InstanceId, NewHostDebugInstance);
-				RootDebugInstances.Add(NewHostDebugInstance);
-			}
-		});
-				
+					UObject* Asset = FObjectTrace::GetObjectFromId(DataInterfaceData->AssetId);
+					if (FSoftObjectPath(Asset) != FocusedAssetPath)
+					{
+						return;
+					}
+
+					const UObject* OwnerObject = FObjectTrace::GetObjectFromId(DataInterfaceData->OuterObjectId);
+					if (!WorkspaceDebugObject || (WorkspaceDebugObject != OwnerObject && !OwnerObject->IsInOuter(WorkspaceDebugObject)))
+					{
+						return;
+					}
+
+					TSharedPtr<FDebugInstance> NewDebugInstance = MakeShared<FDebugInstance>();
+					NewDebugInstance->InstanceId = DataInterfaceData->InstanceId;
+					NewDebugInstance->HostInstanceId = DataInterfaceData->HostInstanceId;
+					NewDebugInstance->Asset = FObjectTrace::GetObjectFromId(DataInterfaceData->AssetId);
+
+					DebugInstances.Add(NewDebugInstance->InstanceId, NewDebugInstance);
+
+					uint64 HostInstanceId = NewDebugInstance->HostInstanceId;
+					const FDataInterfaceData* HostInstanceData = UAFProvider->GetDataInterfaceData(HostInstanceId);
+					while (HostInstanceData)
+					{
+						if (!DebugInstances.Contains(HostInstanceId))
+						{
+							TSharedPtr<FDebugInstance> NewHostDebugInstance = MakeShared<FDebugInstance>();
+							NewHostDebugInstance->InstanceId = HostInstanceData->InstanceId;
+							NewHostDebugInstance->HostInstanceId = HostInstanceData->HostInstanceId;
+							NewHostDebugInstance->Asset = FObjectTrace::GetObjectFromId(HostInstanceData->AssetId);
+
+							DebugInstances.Add(NewHostDebugInstance->InstanceId, NewHostDebugInstance);
+
+							HostInstanceId = NewHostDebugInstance->HostInstanceId;
+							HostInstanceData = UAFProvider->GetDataInterfaceData(HostInstanceId);
+						}
+						else
+						{
+							break;
+						}
+					}
+
+					if (!DebugInstances.Contains(HostInstanceId))
+					{
+						TSharedPtr<FDebugInstance> NewHostDebugInstance = MakeShared<FDebugInstance>();
+						NewHostDebugInstance->InstanceId = HostInstanceId;
+						NewHostDebugInstance->HostInstanceId = 0;
+						NewHostDebugInstance->Asset = FObjectTrace::GetObjectFromId(HostInstanceId);
+
+						DebugInstances.Add(NewHostDebugInstance->InstanceId, NewHostDebugInstance);
+						RootDebugInstances.Add(NewHostDebugInstance);
+					}
+				});
+		}
+
 		for (const auto& [InstanceId, DebugTreeItem] : DebugInstances)
 		{
 			DebugInstanceTreeView->SetItemExpansion(DebugTreeItem, true);
@@ -325,78 +327,80 @@ namespace UE::UAF::Editor
 	{
 #if UAF_TRACE_ENABLED
 		VariablesBoxSource.Empty();
-	
-		const TraceServices::IAnalysisSession* AnalysisSession = RewindDebugger->GetAnalysisSession();
-		TraceServices::FAnalysisSessionReadScope SessionReadScope(*AnalysisSession);
 
-		const double CurrentScrubTime = RewindDebugger->GetScrubTime();
-		const double CurrentTraceTime  = RewindDebugger->CurrentTraceTime();
-				
+		const double CurrentTraceTime = RewindDebugger->CurrentTraceTime();
+
+		const TraceServices::IAnalysisSession* AnalysisSession = RewindDebugger->GetAnalysisSession();
 		const IGameplayProvider* GameplayProvider = AnalysisSession->ReadProvider<IGameplayProvider>("GameplayProvider");
 		const FUAFProvider* UAFProvider = AnalysisSession->ReadProvider<FUAFProvider>("UAFProvider");
-        if (!GameplayProvider || !UAFProvider)
-        {
-	        return;
-        }  	
-        
-		const FDataInterfaceData* DataInterfaceData = UAFProvider->GetDataInterfaceData(SelectedDebugInstanceId);
-		if (!DataInterfaceData)
+		if (!GameplayProvider || !UAFProvider)
 		{
 			return;
 		}
-		
-		const TraceServices::IFrameProvider& FrameProvider = TraceServices::ReadFrameProvider(*AnalysisSession);
 
-		TraceServices::FFrame Frame;
-		if (!FrameProvider.GetFrameFromTime(ETraceFrameType::TraceFrameType_Game, CurrentTraceTime, Frame))
+		const FDataInterfaceData* DataInterfaceData = UAFProvider->GetDataInterfaceData(SelectedDebugInstanceId);
+		if (!DataInterfaceData)
 		{
 			return;
 		}
-			
-		DataInterfaceData->VariablesTimeline.EnumerateEvents(Frame.StartTime, Frame.EndTime, [UAFProvider, GameplayProvider, AnalysisSession, this](double InStartTime, double InEndTime, uint32 InDepth, const FPropertyVariableData& VariableListData)
+
+		TraceServices::FFrame Frame;
 		{
-			// Most of this reading code is ripped from FAnimNextModuleTrack
-			if (VariableListData.ValueType == EPropertyVariableDataType::PropertyBag)
+			TraceServices::FAnalysisSessionReadScope SessionReadScope(*AnalysisSession);
+			if (!TraceServices::ReadFrameProvider(*AnalysisSession).GetFrameFromTime(ETraceFrameType::TraceFrameType_Game, CurrentTraceTime, Frame))
 			{
-				// First look up property description
-				const FPropertyDescriptionData* DescriptionData = UAFProvider->GetPropertyDescriptionData(VariableListData.DescriptionHash);
-				if (DescriptionData == nullptr)
-				{
-					// Can't do anything without a description
-					return TraceServices::EEventEnumerate::Continue;
-				}
-
-				// Load the property descriptions
-				TArray<FPropertyBagPropertyDesc> PropertyDescriptions;
-				FMemoryReader DescriptionReader(DescriptionData->PropertyBagData);
-				FObjectAsTraceIdProxyArchiveReader DescriptionReaderProxy(DescriptionReader, GameplayProvider);
-				DescriptionReaderProxy.UsingCustomVersion(UE::UAF::FUAFTrace::CustomVersionGUID);
-				DescriptionReaderProxy << PropertyDescriptions;
-				
-				//Properties->Properties.SetNum(Properties->Properties.Num() + 1);
-				Properties.AddProperties(PropertyDescriptions);
-
-				// Load the property values
-				FMemoryReader Reader(VariableListData.ValueData);
-				FObjectAsTraceIdProxyArchiveReader ReaderProxy(Reader, GameplayProvider);
+				return;
+			}
+		}
 
-				UPropertyBag* PropertyBagStruct = const_cast<UPropertyBag*>(Properties.GetPropertyBagStruct());
-				if (PropertyBagStruct != nullptr)
+		{
+			TraceServices::FProviderReadScopeLock GameplayProviderReadScope(*GameplayProvider);
+			TraceServices::FProviderReadScopeLock UAFProviderReadScope(*UAFProvider);
+			DataInterfaceData->VariablesTimeline.EnumerateEvents(Frame.StartTime, Frame.EndTime, [UAFProvider, GameplayProvider, AnalysisSession, this](double InStartTime, double InEndTime, uint32 InDepth, const FPropertyVariableData& VariableListData)
 				{
-					PropertyBagStruct->SerializeItem(ReaderProxy, Properties.GetMutableValue().GetMemory(), nullptr);
-				}
-			}
-			else
-			{
-				FMemoryReader Reader(VariableListData.ValueData);
-				FObjectAsTraceIdProxyArchiveReader ReaderProxy(Reader, GameplayProvider);
+					// Most of this reading code is ripped from FAnimNextModuleTrack
+					if (VariableListData.ValueType == EPropertyVariableDataType::PropertyBag)
+					{
+						// First look up property description
+						const FPropertyDescriptionData* DescriptionData = UAFProvider->GetPropertyDescriptionData(VariableListData.DescriptionHash);
+						if (DescriptionData == nullptr)
+						{
+							// Can't do anything without a description
+							return TraceServices::EEventEnumerate::Continue;
+						}
+
+						// Load the property descriptions
+						TArray<FPropertyBagPropertyDesc> PropertyDescriptions;
+						FMemoryReader DescriptionReader(DescriptionData->PropertyBagData);
+						FObjectAsTraceIdProxyArchiveReader DescriptionReaderProxy(DescriptionReader, GameplayProvider);
+						DescriptionReaderProxy.UsingCustomVersion(UE::UAF::FUAFTrace::CustomVersionGUID);
+						DescriptionReaderProxy << PropertyDescriptions;
+
+						//Properties->Properties.SetNum(Properties->Properties.Num() + 1);
+						Properties.AddProperties(PropertyDescriptions);
+
+						// Load the property values
+						FMemoryReader Reader(VariableListData.ValueData);
+						FObjectAsTraceIdProxyArchiveReader ReaderProxy(Reader, GameplayProvider);
+
+						UPropertyBag* PropertyBagStruct = const_cast<UPropertyBag*>(Properties.GetPropertyBagStruct());
+						if (PropertyBagStruct != nullptr)
+						{
+							PropertyBagStruct->SerializeItem(ReaderProxy, Properties.GetMutableValue().GetMemory(), nullptr);
+						}
+					}
+					else
+					{
+						FMemoryReader Reader(VariableListData.ValueData);
+						FObjectAsTraceIdProxyArchiveReader ReaderProxy(Reader, GameplayProvider);
+
+						FInstancedStruct::StaticStruct()->SerializeItem(ReaderProxy, &NativeProperties, nullptr);
+					}
 
-				FInstancedStruct::StaticStruct()->SerializeItem(ReaderProxy, &NativeProperties, nullptr);
-			}
+					return TraceServices::EEventEnumerate::Continue;
+				});
+		}
 
-			return TraceServices::EEventEnumerate::Continue;
-		});
-		        		
 		// This gives you the schema (like a struct definition)
 		const UPropertyBag* BagStruct = Properties.GetPropertyBagStruct();
 		if (!BagStruct)
diff --git a/Engine/Plugins/Experimental/UAF/UAFAnimGraph/Source/UAFAnimGraphEditor/Private/RewindDebugger/AnimNextAnimGraphProvider.cpp b/Engine/Plugins/Experimental/UAF/UAFAnimGraph/Source/UAFAnimGraphEditor/Private/RewindDebugger/AnimNextAnimGraphProvider.cpp
--- a/Engine/Plugins/Experimental/UAF/UAFAnimGraph/Source/UAFAnimGraphEditor/Private/RewindDebugger/AnimNextAnimGraphProvider.cpp
+++ b/Engine/Plugins/Experimental/UAF/UAFAnimGraph/Source/UAFAnimGraphEditor/Private/RewindDebugger/AnimNextAnimGraphProvider.cpp
@@ -3,6 +3,11 @@
 #include "AnimNextAnimGraphProvider.h"
 #include "ObjectTrace.h"
 
+namespace TraceServices
+{
+thread_local FProviderLock::FThreadLocalState GAnimNextAnimGraphProviderLockState;
+}
+
 FName FAnimNextAnimGraphProvider::ProviderName("AnimNextAnimGraphProvider");
 
 #define LOCTEXT_NAMESPACE "AnimNextAnimGraphProvider"
@@ -14,7 +19,7 @@ FAnimNextAnimGraphProvider::FAnimNextAnimGraphProvider(TraceServices::IAnalysisS
 
 void FAnimNextAnimGraphProvider::AppendEvaluationProgram(double ProfileTime, double RecordingTime, uint64 OuterObjectId, uint64 GraphInstanceId, const TArrayView<const uint8>& ProgramData)
 {
-	Session.WriteAccessCheck();
+	EditAccessCheck();
 
 	EvaluationGraphs.Add(OuterObjectId, GraphInstanceId);
 	
@@ -35,7 +40,7 @@ void FAnimNextAnimGraphProvider::AppendEvaluationProgram(double ProfileTime, dou
 
 const FEvaluationProgramData* FAnimNextAnimGraphProvider::GetEvaluationProgramData(uint64 GraphInstanceId) const
 {
-	Session.ReadAccessCheck();
+	ReadAccessCheck();
 	
 	const TSharedRef<FEvaluationProgramData>* Data = EvaluationProgramData.Find(GraphInstanceId);
 	if (Data)
@@ -48,7 +53,7 @@ const FEvaluationProgramData* FAnimNextAnimGraphProvider::GetEvaluationProgramDa
 
 void FAnimNextAnimGraphProvider::EnumerateEvaluationGraphs(uint64 OuterObjectId, TFunctionRef<void(uint64 GraphInstanceId)> Callback) const
 {
-	Session.ReadAccessCheck();
+	ReadAccessCheck();
 
 	TArray<uint64> GraphIds;
 	EvaluationGraphs.MultiFind(OuterObjectId, GraphIds);
diff --git a/Engine/Plugins/Experimental/UAF/UAFAnimGraph/Source/UAFAnimGraphEditor/Private/RewindDebugger/AnimNextAnimGraphProvider.h b/Engine/Plugins/Experimental/UAF/UAFAnimGraph/Source/UAFAnimGraphEditor/Private/RewindDebugger/AnimNextAnimGraphProvider.h
--- a/Engine/Plugins/Experimental/UAF/UAFAnimGraph/Source/UAFAnimGraphEditor/Private/RewindDebugger/AnimNextAnimGraphProvider.h
+++ b/Engine/Plugins/Experimental/UAF/UAFAnimGraph/Source/UAFAnimGraphEditor/Private/RewindDebugger/AnimNextAnimGraphProvider.h
@@ -8,8 +8,13 @@
 #include "Templates/SharedPointer.h"
 #include "Model/IntervalTimeline.h"
 #include "TraceServices/Model/AnalysisSession.h"
+#include "Common/ProviderLock.h"
 
-namespace TraceServices { class IAnalysisSession; }
+namespace TraceServices
+{
+	class IAnalysisSession;
+	extern thread_local FProviderLock::FThreadLocalState GAnimNextAnimGraphProviderLockState;
+}
 
 struct FEvaluationProgramData
 {
@@ -24,13 +29,45 @@ struct FEvaluationProgramData
 	TraceServices::TPointTimeline<TArray<uint8>> EvaluationProgramTimeline;
 };
 
-class FAnimNextAnimGraphProvider : public TraceServices::IProvider
+class FAnimNextAnimGraphProvider : public TraceServices::IProvider, public TraceServices::IEditableProvider
 {
 public:
 	static FName ProviderName;
 
 	FAnimNextAnimGraphProvider(TraceServices::IAnalysisSession& InSession);
 
+	//~ IProvider (read lock)
+	virtual void BeginRead() const override
+	{
+		Lock.BeginRead(TraceServices::GAnimNextAnimGraphProviderLockState);
+	}
+
+	virtual void EndRead() const override
+	{
+		Lock.EndRead(TraceServices::GAnimNextAnimGraphProviderLockState);
+	}
+
+	virtual void ReadAccessCheck() const override
+	{
+		Lock.ReadAccessCheck(TraceServices::GAnimNextAnimGraphProviderLockState);
+	}
+
+	//~ IEditableProvider (edit lock)
+	virtual void BeginEdit() const override
+	{
+		Lock.BeginWrite(TraceServices::GAnimNextAnimGraphProviderLockState);
+	}
+
+	virtual void EndEdit() const override
+	{
+		Lock.EndWrite(TraceServices::GAnimNextAnimGraphProviderLockState);
+	}
+
+	virtual void EditAccessCheck() const override
+	{
+		Lock.WriteAccessCheck(TraceServices::GAnimNextAnimGraphProviderLockState);
+	}
+
 	const FEvaluationProgramData* GetEvaluationProgramData(uint64 GraphInstanceId) const;
 	void AppendEvaluationProgram(double ProfileTime, double RecordingTime, uint64 ObjectInstanceId, uint64 GraphInstanceId, const TArrayView<const uint8>& ProgramData);
 	
@@ -38,6 +75,7 @@ public:
 	
 private:
 	TraceServices::IAnalysisSession& Session;
+	mutable TraceServices::FProviderLock Lock;
 
 	TMap<uint64, TSharedRef<FEvaluationProgramData>> EvaluationProgramData;
 
diff --git a/Engine/Plugins/Experimental/UAF/UAFAnimGraph/Source/UAFAnimGraphEditor/Private/RewindDebugger/EvaluationProgramTrack.cpp b/Engine/Plugins/Experimental/UAF/UAFAnimGraph/Source/UAFAnimGraphEditor/Private/RewindDebugger/EvaluationProgramTrack.cpp
--- a/Engine/Plugins/Experimental/UAF/UAFAnimGraph/Source/UAFAnimGraphEditor/Private/RewindDebugger/EvaluationProgramTrack.cpp
+++ b/Engine/Plugins/Experimental/UAF/UAFAnimGraph/Source/UAFAnimGraphEditor/Private/RewindDebugger/EvaluationProgramTrack.cpp
@@ -3,6 +3,7 @@
 #include "EvaluationProgramTrack.h"
 
 #include "AnimNextAnimGraphProvider.h"
+#include "Common/ProviderLock.h"
 #include "Editor.h"
 #include "IGameplayProvider.h"
 #include "IRewindDebugger.h"
@@ -101,9 +102,9 @@ bool FEvaluationProgramTrack::UpdateInternal()
 
  	bool bChanged = false;
 	
-	if (const FAnimNextAnimGraphProvider* AnimNextAnimGraphProvider = IRewindDebugger::Instance()->GetAnalysisSession()->ReadProvider<FAnimNextAnimGraphProvider>("AnimNextAnimGraphProvider"))
+	if (const FAnimNextAnimGraphProvider* AnimNextAnimGraphProvider = AnalysisSession->ReadProvider<FAnimNextAnimGraphProvider>("AnimNextAnimGraphProvider"))
 	{
-		double CurrentScrubTime = IRewindDebugger::Instance()->CurrentTraceTime();
+		double CurrentScrubTime = RewindDebugger->CurrentTraceTime();
 
 		UEvaluationProgramDetailsObject* DetailsObject = DetailsObjectWeakPtr.Get();
 		if (DetailsObject == nullptr)
@@ -112,6 +113,8 @@ bool FEvaluationProgramTrack::UpdateInternal()
 			Initialize();
 		}
 
+		TraceServices::FProviderReadScopeLock AnimNextProviderReadScope(*AnimNextAnimGraphProvider);
+
 		if (InstanceId == 0)
 		{
 		 	AnimNextAnimGraphProvider->EnumerateEvaluationGraphs(ObjectId, [this](uint64 GraphId)
@@ -127,16 +130,21 @@ bool FEvaluationProgramTrack::UpdateInternal()
 				if (PreviousScrubTime != CurrentScrubTime)
 				{
 					PreviousScrubTime = CurrentScrubTime;
-					
-					const TraceServices::IFrameProvider& FramesProvider = TraceServices::ReadFrameProvider(*AnalysisSession);
-					const IGameplayProvider* GameplayProvider = AnalysisSession->ReadProvider<IGameplayProvider>("GameplayProvider");
-					TraceServices::FAnalysisSessionReadScope SessionReadScope(*AnalysisSession);
 
+					bool bFrameFound;
 					TraceServices::FFrame MarkerFrame;
-					if(FramesProvider.GetFrameFromTime(ETraceFrameType::TraceFrameType_Game, CurrentScrubTime, MarkerFrame))
 					{
+						TraceServices::FAnalysisSessionReadScope SessionReadScope(*AnalysisSession);
+						const TraceServices::IFrameProvider& FrameProvider = TraceServices::ReadFrameProvider(*AnalysisSession);
+						bFrameFound = FrameProvider.GetFrameFromTime(TraceFrameType_Game, CurrentScrubTime, MarkerFrame);
+					}
+
+					if (bFrameFound)
+					{
+						const IGameplayProvider* GameplayProvider = AnalysisSession->ReadProvider<IGameplayProvider>("GameplayProvider");
+						TraceServices::FProviderReadScopeLock GameplayProviderReadScope(*GameplayProvider);
 						Data->EvaluationProgramTimeline.EnumerateEvents(MarkerFrame.StartTime, MarkerFrame.EndTime, [DetailsObject, GameplayProvider](double InStartTime, double InEndTime, uint32 InDepth, const TArray<uint8>& VariableData)
-						{				
+						{
 							FMemoryReader Reader(VariableData);
 							FObjectAsTraceIdProxyArchiveReader Archive(Reader, GameplayProvider);
 
@@ -188,10 +196,9 @@ bool FEvaluationProgramTrackCreator::HasDebugInfoInternal(const RewindDebugger::
 	const IRewindDebugger* RewindDebugger = IRewindDebugger::Instance();
 	if (const TraceServices::IAnalysisSession* AnalysisSession = RewindDebugger->GetAnalysisSession())
 	{
-		TraceServices::FAnalysisSessionReadScope SessionReadScope(*AnalysisSession);
-
 		if (const FAnimNextAnimGraphProvider* AnimNextAnimGraphProvider = IRewindDebugger::Instance()->GetAnalysisSession()->ReadProvider<FAnimNextAnimGraphProvider>("AnimNextAnimGraphProvider"))
 		{
+			TraceServices::FProviderReadScopeLock ProviderReadScope(*AnimNextAnimGraphProvider);
 			uint64 InstanceId = 0;
 			AnimNextAnimGraphProvider->EnumerateEvaluationGraphs(InObjectId.GetMainId(), [&InstanceId](uint64 GraphId)
 				{
diff --git a/Engine/Plugins/Experimental/UAF/UAFAnimGraph/Source/UAFAnimGraphEditor/Private/RewindDebugger/SequenceInfoTrack.cpp b/Engine/Plugins/Experimental/UAF/UAFAnimGraph/Source/UAFAnimGraphEditor/Private/RewindDebugger/SequenceInfoTrack.cpp
--- a/Engine/Plugins/Experimental/UAF/UAFAnimGraph/Source/UAFAnimGraphEditor/Private/RewindDebugger/SequenceInfoTrack.cpp
+++ b/Engine/Plugins/Experimental/UAF/UAFAnimGraph/Source/UAFAnimGraphEditor/Private/RewindDebugger/SequenceInfoTrack.cpp
@@ -3,6 +3,7 @@
 #include "SequenceInfoTrack.h"
 
 #include "AnimNextAnimGraphProvider.h"
+#include "Common/ProviderLock.h"
 #include "DetailWidgetRow.h"
 #include "IDetailChildrenBuilder.h"
 #include "IGameplayProvider.h"
@@ -307,19 +308,26 @@ bool FSequenceInfoTrack::UpdateInternal()
 
 		if (InstanceId != 0)
 		{
+			TraceServices::FProviderReadScopeLock AnimNextProviderReadScope(*AnimNextAnimGraphProvider);
 			if (const FEvaluationProgramData* Data = AnimNextAnimGraphProvider->GetEvaluationProgramData(InstanceId))
 			{
 				if (PreviousScrubTime != CurrentScrubTime)
 				{
 					PreviousScrubTime = CurrentScrubTime;
 
-					const TraceServices::IFrameProvider& FramesProvider = TraceServices::ReadFrameProvider(*AnalysisSession);
-					const IGameplayProvider* GameplayProvider = AnalysisSession->ReadProvider<IGameplayProvider>("GameplayProvider");
-					TraceServices::FAnalysisSessionReadScope SessionReadScope(*AnalysisSession);
-
+					bool bFrameFound;
 					TraceServices::FFrame MarkerFrame;
-					if (FramesProvider.GetFrameFromTime(ETraceFrameType::TraceFrameType_Game, CurrentScrubTime, MarkerFrame))
 					{
+						TraceServices::FAnalysisSessionReadScope SessionReadScope(*AnalysisSession);
+						const TraceServices::IFrameProvider& FrameProvider = TraceServices::ReadFrameProvider(*AnalysisSession);
+						bFrameFound = FrameProvider.GetFrameFromTime(TraceFrameType_Game, CurrentScrubTime, MarkerFrame);
+					}
+		
+					if (bFrameFound)
+					{
+						const IGameplayProvider* GameplayProvider = AnalysisSession->ReadProvider<IGameplayProvider>("GameplayProvider");
+						TraceServices::FProviderReadScopeLock GameplayProviderReadScope(*GameplayProvider);
+
 						Data->EvaluationProgramTimeline.EnumerateEvents(MarkerFrame.StartTime, MarkerFrame.EndTime,
 						[DetailsObject, GameplayProvider, this] ( double InStartTime, double InEndTime, uint32 InDepth, const TArray<uint8>& VariableData)
 						{
@@ -350,10 +358,9 @@ bool FSequenceInfoTrackCreator::HasDebugInfoInternal(const RewindDebugger::FObje
 	const IRewindDebugger* RewindDebugger = IRewindDebugger::Instance();
 	if (const TraceServices::IAnalysisSession* AnalysisSession = RewindDebugger->GetAnalysisSession())
 	{
-		TraceServices::FAnalysisSessionReadScope SessionReadScope(*AnalysisSession);
-
 		if (const FAnimNextAnimGraphProvider* AnimNextAnimGraphProvider = IRewindDebugger::Instance()->GetAnalysisSession()->ReadProvider<FAnimNextAnimGraphProvider>("AnimNextAnimGraphProvider"))
 		{
+			TraceServices::FProviderReadScopeLock ProviderReadScope(*AnimNextAnimGraphProvider);
 			uint64 InstanceId = 0;
 			AnimNextAnimGraphProvider->EnumerateEvaluationGraphs(InObjectId.GetMainId(), [&InstanceId](uint64 GraphId)
 				{
diff --git a/Engine/Plugins/Experimental/UAF/UAFAnimNode/Source/UAFAnimNodeEditor/Private/RewindDebugger/UAFAnimNodeProvider.cpp b/Engine/Plugins/Experimental/UAF/UAFAnimNode/Source/UAFAnimNodeEditor/Private/RewindDebugger/UAFAnimNodeProvider.cpp
--- a/Engine/Plugins/Experimental/UAF/UAFAnimNode/Source/UAFAnimNodeEditor/Private/RewindDebugger/UAFAnimNodeProvider.cpp
+++ b/Engine/Plugins/Experimental/UAF/UAFAnimNode/Source/UAFAnimNodeEditor/Private/RewindDebugger/UAFAnimNodeProvider.cpp
@@ -3,6 +3,11 @@
 #include "UAFAnimNodeProvider.h"
 #include "ObjectTrace.h"
 
+namespace TraceServices
+{
+thread_local FProviderLock::FThreadLocalState GUAFAnimNodeProviderLockState;
+}
+
 FName FUAFAnimNodeProvider::ProviderName("UAFAnimNodeProvider");
 
 #define LOCTEXT_NAMESPACE "UAFAnimNodeProvider"
@@ -14,7 +19,7 @@ FUAFAnimNodeProvider::FUAFAnimNodeProvider(TraceServices::IAnalysisSession& InSe
 
 void FUAFAnimNodeProvider::AppendAnimOpList(double ProfileTime, double RecordingTime, uint64 OuterObjectId, uint64 GraphInstanceId, const TArrayView<const uint8>& ListData)
 {
-	Session.WriteAccessCheck();
+	EditAccessCheck();
 
 	EvaluationGraphs.Add(OuterObjectId, GraphInstanceId);
 
@@ -35,7 +40,7 @@ void FUAFAnimNodeProvider::AppendAnimOpList(double ProfileTime, double Recording
 
 void FUAFAnimNodeProvider::AppendAnimNodeUpdate(double ProfileTime, double RecordingTime, uint64 RootGraphId, uint64 NodeId, uint64 ParentNodeId, float TotalWeight)
 {
-	Session.WriteAccessCheck();
+	EditAccessCheck();
     
 	const TSharedRef<FAnimNodeUpdateTimelineData>* Data = AnimNodeUpdateTimelineData.Find(RootGraphId);
 	if (Data == nullptr)
@@ -52,7 +57,7 @@ void FUAFAnimNodeProvider::AppendAnimNodeUpdate(double ProfileTime, double Recor
 
 void FUAFAnimNodeProvider::AppendAnimNodeValue(double ProfileTime, double RecordingTime, uint64 RootGraphId, uint64 NodeId, uint8 Type, uint32 NameId, TConstArrayView<uint8> Value)
 {
-	Session.WriteAccessCheck();
+	EditAccessCheck();
     
 	const TSharedRef<FAnimNodeValueTimelineData>* Data = AnimNodeValueTimelineData.Find(RootGraphId);
 	if (Data == nullptr)
@@ -77,7 +82,7 @@ void FUAFAnimNodeProvider::AppendAnimNodeValue(double ProfileTime, double Record
 
 const FAnimOpData* FUAFAnimNodeProvider::GetAnimOpData(uint64 GraphInstanceId) const
 {
-	Session.ReadAccessCheck();
+	ReadAccessCheck();
 
 	const TSharedRef<FAnimOpData>* Data = AnimOpData.Find(GraphInstanceId);
 	if (Data)
@@ -89,7 +94,7 @@ const FAnimOpData* FUAFAnimNodeProvider::GetAnimOpData(uint64 GraphInstanceId) c
 
 const FAnimNodeUpdateTimelineData* FUAFAnimNodeProvider::GetAnimNodeTimelineData(uint64 GraphInstanceId) const
 {
-	Session.ReadAccessCheck();
+	ReadAccessCheck();
 	
 	const TSharedRef<FAnimNodeUpdateTimelineData>* Data = AnimNodeUpdateTimelineData.Find(GraphInstanceId);
 	if (Data)
@@ -102,7 +107,7 @@ const FAnimNodeUpdateTimelineData* FUAFAnimNodeProvider::GetAnimNodeTimelineData
 
 const FAnimNodeValueTimelineData* FUAFAnimNodeProvider::GetAnimNodeValueTimelineData(uint64 GraphInstanceId) const
 {
-	Session.ReadAccessCheck();
+	ReadAccessCheck();
 	
 	const TSharedRef<FAnimNodeValueTimelineData>* Data = AnimNodeValueTimelineData.Find(GraphInstanceId);
 	if (Data)
@@ -115,7 +120,7 @@ const FAnimNodeValueTimelineData* FUAFAnimNodeProvider::GetAnimNodeValueTimeline
 
 void FUAFAnimNodeProvider::EnumerateEvaluationGraphs(uint64 OuterObjectId, TFunctionRef<void(uint64 GraphInstanceId)> Callback) const
 {
-	Session.ReadAccessCheck();
+	ReadAccessCheck();
 
 	TArray<uint64> GraphIds;
 	EvaluationGraphs.MultiFind(OuterObjectId, GraphIds);
diff --git a/Engine/Plugins/Experimental/UAF/UAFAnimNode/Source/UAFAnimNodeEditor/Private/RewindDebugger/UAFAnimNodeProvider.h b/Engine/Plugins/Experimental/UAF/UAFAnimNode/Source/UAFAnimNodeEditor/Private/RewindDebugger/UAFAnimNodeProvider.h
--- a/Engine/Plugins/Experimental/UAF/UAFAnimNode/Source/UAFAnimNodeEditor/Private/RewindDebugger/UAFAnimNodeProvider.h
+++ b/Engine/Plugins/Experimental/UAF/UAFAnimNode/Source/UAFAnimNodeEditor/Private/RewindDebugger/UAFAnimNodeProvider.h
@@ -8,8 +8,13 @@
 #include "Templates/SharedPointer.h"
 #include "Model/IntervalTimeline.h"
 #include "TraceServices/Model/AnalysisSession.h"
+#include "Common/ProviderLock.h"
 
-namespace TraceServices { class IAnalysisSession; }
+namespace TraceServices
+{
+	class IAnalysisSession;
+	extern thread_local FProviderLock::FThreadLocalState GUAFAnimNodeProviderLockState;
+}
 
 
 struct FSmallPageSettings
@@ -77,13 +82,45 @@ struct FAnimNodeValueTimelineData
 	TraceServices::TPointTimeline<FAnimNodeValueData,FSmallPageSettings> NodeValueTimeline;
 };
 
-class FUAFAnimNodeProvider : public TraceServices::IProvider
+class FUAFAnimNodeProvider : public TraceServices::IProvider, public TraceServices::IEditableProvider
 {
 public:
 	static FName ProviderName;
 
 	FUAFAnimNodeProvider(TraceServices::IAnalysisSession& InSession);
 
+	//~ IProvider (read lock)
+	virtual void BeginRead() const override
+	{
+		Lock.BeginRead(TraceServices::GUAFAnimNodeProviderLockState);
+	}
+
+	virtual void EndRead() const override
+	{
+		Lock.EndRead(TraceServices::GUAFAnimNodeProviderLockState);
+	}
+
+	virtual void ReadAccessCheck() const override
+	{
+		Lock.ReadAccessCheck(TraceServices::GUAFAnimNodeProviderLockState);
+	}
+
+	//~ IEditableProvider (edit lock)
+	virtual void BeginEdit() const override
+	{
+		Lock.BeginWrite(TraceServices::GUAFAnimNodeProviderLockState);
+	}
+
+	virtual void EndEdit() const override
+	{
+		Lock.EndWrite(TraceServices::GUAFAnimNodeProviderLockState);
+	}
+
+	virtual void EditAccessCheck() const override
+	{
+		Lock.WriteAccessCheck(TraceServices::GUAFAnimNodeProviderLockState);
+	}
+
 	const FAnimOpData* GetAnimOpData(uint64 GraphInstanceId) const;
 	const FAnimNodeUpdateTimelineData* GetAnimNodeTimelineData(uint64 GraphInstanceId) const;
 	const FAnimNodeValueTimelineData* GetAnimNodeValueTimelineData(uint64 GraphInstanceId) const;
@@ -96,6 +133,7 @@ public:
 	
 private:
 	TraceServices::IAnalysisSession& Session;
+	mutable TraceServices::FProviderLock Lock;
 
 	TMap<uint64, TSharedRef<FAnimOpData>> AnimOpData;
 	TMap<uint64, TSharedRef<FAnimNodeUpdateTimelineData>> AnimNodeUpdateTimelineData;
diff --git a/Engine/Plugins/Experimental/UAF/UAFAnimNode/Source/UAFAnimNodeEditor/Private/RewindDebugger/UAFAnimNodeTrack.cpp b/Engine/Plugins/Experimental/UAF/UAFAnimNode/Source/UAFAnimNodeEditor/Private/RewindDebugger/UAFAnimNodeTrack.cpp
--- a/Engine/Plugins/Experimental/UAF/UAFAnimNode/Source/UAFAnimNodeEditor/Private/RewindDebugger/UAFAnimNodeTrack.cpp
+++ b/Engine/Plugins/Experimental/UAF/UAFAnimNode/Source/UAFAnimNodeEditor/Private/RewindDebugger/UAFAnimNodeTrack.cpp
@@ -2,6 +2,7 @@
 
 #include "UAFAnimNodeTrack.h"
 
+#include "Common/ProviderLock.h"
 #include "UAFAnimNodeProvider.h"
 #include "DetailWidgetRow.h"
 #include "IGameplayProvider.h"
@@ -144,7 +145,9 @@ bool FUAFAnimNodeTrack::UpdateInternal()
 
 	if (GameplayProvider && AnimationProvider && UAFAnimNodeProvider)
 	{
-		TraceServices::FAnalysisSessionReadScope SessionReadScope(*Session);
+		TraceServices::FProviderReadScopeLock GameplayProviderReadScope(*GameplayProvider);
+		TraceServices::FProviderReadScopeLock AnimationProviderReadScope(*AnimationProvider);
+		TraceServices::FProviderReadScopeLock UAFAnimNodeProviderReadScope(*UAFAnimNodeProvider);
 		double CurrentScrubTime = RewindDebugger->CurrentTraceTime();
 
 		if (GraphId == 0)
diff --git a/Engine/Plugins/Experimental/UAF/UAFAnimNode/Source/UAFAnimNodeEditor/Private/RewindDebugger/UAFAnimOpTrack.cpp b/Engine/Plugins/Experimental/UAF/UAFAnimNode/Source/UAFAnimNodeEditor/Private/RewindDebugger/UAFAnimOpTrack.cpp
--- a/Engine/Plugins/Experimental/UAF/UAFAnimNode/Source/UAFAnimNodeEditor/Private/RewindDebugger/UAFAnimOpTrack.cpp
+++ b/Engine/Plugins/Experimental/UAF/UAFAnimNode/Source/UAFAnimNodeEditor/Private/RewindDebugger/UAFAnimOpTrack.cpp
@@ -2,6 +2,7 @@
 
 #include "UAFAnimOpTrack.h"
 
+#include "Common/ProviderLock.h"
 #include "UAFAnimNodeProvider.h"
 #include "Editor.h"
 #include "IGameplayProvider.h"
@@ -100,7 +101,7 @@ bool FUAFAnimOpTrack::UpdateInternal()
 
  	bool bChanged = false;
 	
-	if (const FUAFAnimNodeProvider* UAFAnimNodeProvider = IRewindDebugger::Instance()->GetAnalysisSession()->ReadProvider<FUAFAnimNodeProvider>("UAFAnimNodeProvider"))
+	if (const FUAFAnimNodeProvider* UAFAnimNodeProvider = AnalysisSession->ReadProvider<FUAFAnimNodeProvider>("UAFAnimNodeProvider"))
 	{
 		double CurrentScrubTime = IRewindDebugger::Instance()->CurrentTraceTime();
 
@@ -126,14 +127,20 @@ bool FUAFAnimOpTrack::UpdateInternal()
 				if (PreviousScrubTime != CurrentScrubTime)
 				{
 					PreviousScrubTime = CurrentScrubTime;
-					
-					const TraceServices::IFrameProvider& FramesProvider = TraceServices::ReadFrameProvider(*AnalysisSession);
-					const IGameplayProvider* GameplayProvider = AnalysisSession->ReadProvider<IGameplayProvider>("GameplayProvider");
-					TraceServices::FAnalysisSessionReadScope SessionReadScope(*AnalysisSession);
 
+					bool bFrameFound;
 					TraceServices::FFrame MarkerFrame;
-					if(FramesProvider.GetFrameFromTime(ETraceFrameType::TraceFrameType_Game, CurrentScrubTime, MarkerFrame))
 					{
+						TraceServices::FAnalysisSessionReadScope SessionReadScope(*AnalysisSession);
+						const TraceServices::IFrameProvider& FrameProvider = TraceServices::ReadFrameProvider(*AnalysisSession);
+						bFrameFound = FrameProvider.GetFrameFromTime(TraceFrameType_Game, CurrentScrubTime, MarkerFrame);
+					}
+		
+					if (bFrameFound)
+					{
+						const IGameplayProvider* GameplayProvider = AnalysisSession->ReadProvider<IGameplayProvider>("GameplayProvider");
+						TraceServices::FProviderReadScopeLock GameplayProviderReadScope(*GameplayProvider);
+						TraceServices::FProviderReadScopeLock UAFAnimNodeProviderReadScope(*UAFAnimNodeProvider);
 						Data->AnimOpTimeline.EnumerateEvents(MarkerFrame.StartTime, MarkerFrame.EndTime, [DetailsObject, GameplayProvider](double InStartTime, double InEndTime, uint32 InDepth, const TArray<uint8>& VariableData)
 						{				
 							FMemoryReader Reader(VariableData);
@@ -187,10 +194,9 @@ bool FUAFAnimOpTrackCreator::HasDebugInfoInternal(const RewindDebugger::FObjectI
 	const IRewindDebugger* RewindDebugger = IRewindDebugger::Instance();
 	if (const TraceServices::IAnalysisSession* AnalysisSession = RewindDebugger->GetAnalysisSession())
 	{
-		TraceServices::FAnalysisSessionReadScope SessionReadScope(*AnalysisSession);
-
-		if (const FUAFAnimNodeProvider* UAFAnimNodeProvider = IRewindDebugger::Instance()->GetAnalysisSession()->ReadProvider<FUAFAnimNodeProvider>("UAFAnimNodeProvider"))
+		if (const FUAFAnimNodeProvider* UAFAnimNodeProvider = AnalysisSession->ReadProvider<FUAFAnimNodeProvider>("UAFAnimNodeProvider"))
 		{
+			TraceServices::FProviderReadScopeLock ProviderReadScope(*UAFAnimNodeProvider);
 			uint64 InstanceId = 0;
 			UAFAnimNodeProvider->EnumerateEvaluationGraphs(InObjectId.GetMainId(), [&InstanceId](uint64 GraphId)
 				{
diff --git a/Engine/Plugins/Experimental/UAF/UAFAnimNode/Source/UAFAnimNodeEditor/Private/RewindDebugger/UAFSequenceInfoTrack.cpp b/Engine/Plugins/Experimental/UAF/UAFAnimNode/Source/UAFAnimNodeEditor/Private/RewindDebugger/UAFSequenceInfoTrack.cpp
--- a/Engine/Plugins/Experimental/UAF/UAFAnimNode/Source/UAFAnimNodeEditor/Private/RewindDebugger/UAFSequenceInfoTrack.cpp
+++ b/Engine/Plugins/Experimental/UAF/UAFAnimNode/Source/UAFAnimNodeEditor/Private/RewindDebugger/UAFSequenceInfoTrack.cpp
@@ -3,6 +3,7 @@
 #include "UAFSequenceInfoTrack.h"
 
 #include "UAFAnimNodeProvider.h"
+#include "Common/ProviderLock.h"
 #include "DetailWidgetRow.h"
 #include "IDetailChildrenBuilder.h"
 #include "IGameplayProvider.h"
@@ -313,13 +314,19 @@ bool FUAFSequenceInfoTrack::UpdateInternal()
 				{
 					PreviousScrubTime = CurrentScrubTime;
 
-					const TraceServices::IFrameProvider& FramesProvider = TraceServices::ReadFrameProvider(*AnalysisSession);
-					const IGameplayProvider* GameplayProvider = AnalysisSession->ReadProvider<IGameplayProvider>("GameplayProvider");
-					TraceServices::FAnalysisSessionReadScope SessionReadScope(*AnalysisSession);
-
+					bool bFrameFound;
 					TraceServices::FFrame MarkerFrame;
-					if (FramesProvider.GetFrameFromTime(ETraceFrameType::TraceFrameType_Game, CurrentScrubTime, MarkerFrame))
 					{
+						TraceServices::FAnalysisSessionReadScope SessionReadScope(*AnalysisSession);
+						const TraceServices::IFrameProvider& FrameProvider = TraceServices::ReadFrameProvider(*AnalysisSession);
+						bFrameFound = FrameProvider.GetFrameFromTime(TraceFrameType_Game, CurrentScrubTime, MarkerFrame);
+					}
+		
+					if (bFrameFound)
+					{
+						const IGameplayProvider* GameplayProvider = AnalysisSession->ReadProvider<IGameplayProvider>("GameplayProvider");
+						TraceServices::FProviderReadScopeLock GameplayProviderReadScope(*GameplayProvider);
+
 						Data->AnimOpTimeline.EnumerateEvents(MarkerFrame.StartTime, MarkerFrame.EndTime,
 						[DetailsObject, GameplayProvider, this] ( double InStartTime, double InEndTime, uint32 InDepth, const TArray<uint8>& VariableData)
 						{
@@ -349,10 +356,9 @@ bool FUAFSequenceInfoTrackCreator::HasDebugInfoInternal(const RewindDebugger::FO
 	const IRewindDebugger* RewindDebugger = IRewindDebugger::Instance();
 	if (const TraceServices::IAnalysisSession* AnalysisSession = RewindDebugger->GetAnalysisSession())
 	{
-		TraceServices::FAnalysisSessionReadScope SessionReadScope(*AnalysisSession);
-
 		if (const FUAFAnimNodeProvider* UAFAnimNodeProvider = IRewindDebugger::Instance()->GetAnalysisSession()->ReadProvider<FUAFAnimNodeProvider>("UAFAnimNodeProvider"))
 		{
+			TraceServices::FProviderReadScopeLock ProviderReadScope(*UAFAnimNodeProvider);
 			uint64 InstanceId = 0;
 			UAFAnimNodeProvider->EnumerateEvaluationGraphs(InObjectId.GetMainId(), [&InstanceId](uint64 GraphId)
 				{
diff --git a/Engine/Plugins/Runtime/RigVM/Source/RigVMEditor/Private/Trace/RigVMTraceModule.cpp b/Engine/Plugins/Runtime/RigVM/Source/RigVMEditor/Private/Trace/RigVMTraceModule.cpp
--- a/Engine/Plugins/Runtime/RigVM/Source/RigVMEditor/Private/Trace/RigVMTraceModule.cpp
+++ b/Engine/Plugins/Runtime/RigVM/Source/RigVMEditor/Private/Trace/RigVMTraceModule.cpp
@@ -8,6 +8,7 @@
 #include "TraceServices/Model/Frames.h"
 #include "ObjectTrace.h"
 #include "RigVMHost.h"
+#include "Common/ProviderLock.h"
 #include "Engine/Level.h"
 
 #if RIGVM_TRACE_ENABLED
@@ -52,6 +53,7 @@ void FRewindDebuggerForRigVM::Update(float DeltaTime, IRewindDebugger* RewindDeb
 	
 	if (const TraceServices::IAnalysisSession* Session = RewindDebugger->GetAnalysisSession())
 	{
+		// Session scope required only for RigVMProvider
 		TraceServices::FAnalysisSessionReadScope SessionReadScope(*Session);
 
 		const IGameplayProvider* GameplayProvider = Session->ReadProvider<IGameplayProvider>("GameplayProvider");
@@ -108,76 +110,84 @@ void FRewindDebuggerForRigVM::Update(float DeltaTime, IRewindDebugger* RewindDeb
 					TArray<FString> Names, Paths;
 
 					uint64 ObjectId = HostId;
-					do
+					uint64 HostInfoClassId = 0;
+					const TCHAR* HostInfoName = nullptr;
+					const TCHAR* HostInfoPathName = nullptr;
 					{
-						const FObjectInfo* ObjectInfo = GameplayProvider->FindObjectInfo(ObjectId);
-						if (!ObjectInfo)
+						TraceServices::FProviderReadScopeLock ProviderReadScope(*GameplayProvider);
+						do
 						{
-							break;
-						}
+							const FObjectInfo* ObjectInfo = GameplayProvider->FindObjectInfo(ObjectId);
+							if (!ObjectInfo)
+							{
+								break;
+							}
 
-						bool bIsLevel = false;
+							bool bIsLevel = false;
 
-						uint64 ClassId = ObjectInfo->ClassId;
-						do
-						{
-							if (ClassId == ActorClassId)
+							uint64 ClassId = ObjectInfo->ClassId;
+							do
 							{
-								// find actor based on name
-								TObjectPtr<AActor>* FoundActor = Level->Actors.FindByPredicate([ObjectInfo](const TObjectPtr<AActor>& ExistingActor) -> bool
+								if (ClassId == ActorClassId)
 								{
-									if (ExistingActor)
-									{
-										FString ActorLabel = ExistingActor->GetActorLabel();
-										if (ActorLabel.StartsWith("RewindDebugger: "))
+									// find actor based on name
+									TObjectPtr<AActor>* FoundActor = Level->Actors.FindByPredicate([ObjectInfo](const TObjectPtr<AActor>& ExistingActor) -> bool
 										{
-											ActorLabel.RemoveFromStart("RewindDebugger: ");
-										}
-										return ActorLabel == ObjectInfo->Name;
+											if (ExistingActor)
+											{
+												FString ActorLabel = ExistingActor->GetActorLabel();
+												if (ActorLabel.StartsWith("RewindDebugger: "))
+												{
+													ActorLabel.RemoveFromStart("RewindDebugger: ");
+												}
+												return ActorLabel == ObjectInfo->Name;
+											}
+											return false;
+										});
+
+									if (FoundActor)
+									{
+										Actor = FoundActor->Get();
 									}
-									return false;
-								});
+									break;
+								}
+								if (ClassId == LevelClassId)
+								{
+									bIsLevel = true;
+									break;
+								}
 
-								if (FoundActor)
+								const FClassInfo* ClassInfo = GameplayProvider->FindClassInfo(ClassId);
+								if (!ClassInfo)
 								{
-									Actor = FoundActor->Get();
+									break;
 								}
-								break;
-							}
-							if (ClassId == LevelClassId)
+
+								ClassId = ClassInfo->SuperId;
+							} while (ClassId);
+
+							if (Actor || bIsLevel)
 							{
-								bIsLevel = true;
 								break;
 							}
-							
-							const FClassInfo* ClassInfo = GameplayProvider->FindClassInfo(ClassId);
-							if (!ClassInfo)
+
+							if (ObjectId != HostId)
 							{
-								break;
+								Names.Insert(ObjectInfo->Name, 0);
+								Paths.Insert(ObjectInfo->PathName, 0);
 							}
+							ObjectId = ObjectInfo->GetOuterId().GetMainId();
+						} while (ObjectId);
 
-							ClassId = ClassInfo->SuperId;
-						}
-						while (ClassId);
-
-						if (Actor || bIsLevel)
+						const FObjectInfo* FoundInfo = GameplayProvider->FindObjectInfo(HostId);
+						if (!FoundInfo)
 						{
-							break;
+							return;
 						}
 
-						if (ObjectId != HostId)
-						{
-							Names.Insert(ObjectInfo->Name, 0);
-							Paths.Insert(ObjectInfo->PathName, 0);
-						}
-						ObjectId = ObjectInfo->GetOuterId().GetMainId();
-					}
-					while (ObjectId);
-
-					const FObjectInfo* HostInfo = GameplayProvider->FindObjectInfo(HostId);
-					if (!HostInfo)
-					{
-						return;
+						HostInfoClassId = FoundInfo->ClassId;
+						HostInfoName = FoundInfo->Name;
+						HostInfoPathName = FoundInfo->PathName;
 					}
 
 					const bool NotYetLoggedAtThisTime = CurrentTraceTime != LastLogTime;
@@ -187,7 +197,7 @@ void FRewindDebuggerForRigVM::Update(float DeltaTime, IRewindDebugger* RewindDeb
 					{
 						if(NotYetLoggedAtThisTime)
 						{
-							UE_LOG(LogRigVM, Warning, TEXT("Cannot find actor for Host '%s' in visualize world."), HostInfo->PathName);
+							UE_LOG(LogRigVM, Warning, TEXT("Cannot find actor for Host '%s' in visualize world."), HostInfoPathName);
 						}
 						return;
 					}
@@ -207,20 +217,20 @@ void FRewindDebuggerForRigVM::Update(float DeltaTime, IRewindDebugger* RewindDeb
 						}
 					}
 
-					URigVMHost* Host = FindObject<URigVMHost>(Outer, HostInfo->Name);
+					URigVMHost* Host = FindObject<URigVMHost>(Outer, HostInfoName);
 					if (!Host)
 					{
-						UClass* HostClass = Cast<UClass>(FObjectTrace::GetObjectFromId(HostInfo->ClassId));
+						UClass* HostClass = Cast<UClass>(FObjectTrace::GetObjectFromId(HostInfoClassId));
 						if (!HostClass)
 						{
 							if(NotYetLoggedAtThisTime)
 							{
-								UE_LOG(LogRigVM, Warning, TEXT("Cannot recreate host '%s' in visualize world - class not found."), HostInfo->PathName);
+								UE_LOG(LogRigVM, Warning, TEXT("Cannot recreate host '%s' in visualize world - class not found."), HostInfoPathName);
 							}
 							return;
 						}
 						
-						Host = NewObject<URigVMHost>(Outer, HostClass, HostInfo->Name);
+						Host = NewObject<URigVMHost>(Outer, HostClass, HostInfoName);
 						Host->Initialize(true);
 
 						// add this here to that GC doesn't destroy the host while the visualize world exist
diff --git a/Engine/Plugins/Runtime/StateTree/Source/StateTreeEditorModule/Private/Debugger/SStateTreeDebuggerView.cpp b/Engine/Plugins/Runtime/StateTree/Source/StateTreeEditorModule/Private/Debugger/SStateTreeDebuggerView.cpp
--- a/Engine/Plugins/Runtime/StateTree/Source/StateTreeEditorModule/Private/Debugger/SStateTreeDebuggerView.cpp
+++ b/Engine/Plugins/Runtime/StateTree/Source/StateTreeEditorModule/Private/Debugger/SStateTreeDebuggerView.cpp
@@ -8,6 +8,7 @@
 #include "Framework/Docking/TabManager.h"
 #include "Framework/MultiBox/MultiBoxBuilder.h"
 #include "IGameplayProvider.h"
+#include "Common/ProviderLock.h"
 #include "IRewindDebuggerModule.h"
 #include "Kismet2/DebuggerCommands.h"
 #include "LevelEditor.h"
@@ -598,7 +599,7 @@ void SStateTreeDebuggerView::ConstructView(const TSharedRef<FUICommandList>& InC
 							if (const TraceServices::IAnalysisSession* AnalysisSession = RewindDebugger->GetAnalysisSession())
 							{
 								const IGameplayProvider* GameplayProvider = AnalysisSession->ReadProvider<IGameplayProvider>("GameplayProvider");
-								TraceServices::FAnalysisSessionReadScope SessionReadScope(*AnalysisSession);
+								TraceServices::FProviderReadScopeLock ProviderReadScope(*GameplayProvider);
 								const FClassInfo* LevelClass = GameplayProvider->FindClassInfo(TEXT("/Script/Engine.Level"));
 								const FClassInfo* WorldClass = GameplayProvider->FindClassInfo(TEXT("/Script/Engine.World"));
 								const FObjectInfo* ObjectInfo = GameplayProvider->FindObjectInfo(RewindDebugger::FObjectId{InstanceId.ToUint64()});
diff --git a/Engine/Plugins/Runtime/StateTree/Source/StateTreeEditorModule/Private/Debugger/StateTreeRewindDebuggerExtensions.cpp b/Engine/Plugins/Runtime/StateTree/Source/StateTreeEditorModule/Private/Debugger/StateTreeRewindDebuggerExtensions.cpp
--- a/Engine/Plugins/Runtime/StateTree/Source/StateTreeEditorModule/Private/Debugger/StateTreeRewindDebuggerExtensions.cpp
+++ b/Engine/Plugins/Runtime/StateTree/Source/StateTreeEditorModule/Private/Debugger/StateTreeRewindDebuggerExtensions.cpp
@@ -18,13 +18,17 @@ namespace UE::StateTreeDebugger
 void FRewindDebuggerPlaybackExtension::Update(float DeltaTime, IRewindDebugger* RewindDebugger)
 {
 	const TraceServices::IAnalysisSession* AnalysisSession = RewindDebugger->GetAnalysisSession();
-	TraceServices::FAnalysisSessionReadScope SessionReadScope(*AnalysisSession);
 
-	const TraceServices::IFrameProvider& FrameProvider = TraceServices::ReadFrameProvider(*RewindDebugger->GetAnalysisSession());
+	bool bFrameFound;
 	TraceServices::FFrame Frame;
+	{
+		TraceServices::FAnalysisSessionReadScope SessionReadScope(*AnalysisSession);
+		const TraceServices::IFrameProvider& FrameProvider = TraceServices::ReadFrameProvider(*AnalysisSession);
+		bFrameFound = FrameProvider.GetFrameFromTime(TraceFrameType_Game, RewindDebugger->CurrentTraceTime(), Frame);
+	}
 
 	// Require some debug frame to exist before doing any processing, currently the frame itself is not used
-	if (FrameProvider.GetFrameFromTime(TraceFrameType_Game, RewindDebugger->CurrentTraceTime(), Frame))
+	if (bFrameFound)
 	{
 		auto SetScrubberTimeline = [&]()
 		{
diff --git a/Engine/Plugins/Runtime/StateTree/Source/StateTreeEditorModule/Private/Debugger/StateTreeRewindDebuggerTrack.cpp b/Engine/Plugins/Runtime/StateTree/Source/StateTreeEditorModule/Private/Debugger/StateTreeRewindDebuggerTrack.cpp
--- a/Engine/Plugins/Runtime/StateTree/Source/StateTreeEditorModule/Private/Debugger/StateTreeRewindDebuggerTrack.cpp
+++ b/Engine/Plugins/Runtime/StateTree/Source/StateTreeEditorModule/Private/Debugger/StateTreeRewindDebuggerTrack.cpp
@@ -39,6 +39,7 @@ bool FRewindDebuggerTrackCreator::HasDebugInfoInternal(const RewindDebugger::FOb
 		{
 			if (const IStateTreeTraceProvider* StateTreeTraceProvider = Session->ReadProvider<IStateTreeTraceProvider>(FStateTreeTraceProvider::ProviderName))
 			{
+				TraceServices::FProviderReadScopeLock ProviderReadScope(*StateTreeTraceProvider);
 				if (StateTreeTraceProvider->GetInstanceDescriptor(FStateTreeInstanceDebugId(InObjectId.GetMainId())).IsValid())
 				{
 					return true;
@@ -74,6 +75,7 @@ FRewindDebuggerTrack::FRewindDebuggerTrack(const RewindDebugger::FObjectId& InOb
 		{
 			if (const IStateTreeTraceProvider* StateTreeTraceProvider = Session->ReadProvider<IStateTreeTraceProvider>(FStateTreeTraceProvider::ProviderName))
 			{
+				TraceServices::FProviderReadScopeLock ProviderReadScope(*StateTreeTraceProvider);
 				Descriptor = StateTreeTraceProvider->GetInstanceDescriptor(EventCollection.InstanceId);
 			}
 		}
diff --git a/Engine/Plugins/Runtime/StateTree/Source/StateTreeModule/Private/Debugger/StateTreeDebugger.cpp b/Engine/Plugins/Runtime/StateTree/Source/StateTreeModule/Private/Debugger/StateTreeDebugger.cpp
--- a/Engine/Plugins/Runtime/StateTree/Source/StateTreeModule/Private/Debugger/StateTreeDebugger.cpp
+++ b/Engine/Plugins/Runtime/StateTree/Source/StateTreeModule/Private/Debugger/StateTreeDebugger.cpp
@@ -3,6 +3,7 @@
 #if WITH_STATETREE_TRACE_DEBUGGER
 
 #include "Debugger/StateTreeDebugger.h"
+#include "Common/ProviderLock.h"
 #include "Debugger/IStateTreeTraceProvider.h"
 #include "Debugger/StateTreeTraceProvider.h"
 #include "Debugger/StateTreeTraceTypes.h"
@@ -207,10 +208,9 @@ TSharedPtr<const UE::StateTreeDebugger::FInstanceDescriptor> FStateTreeDebugger:
 {
 	if (const TraceServices::IAnalysisSession* Session = GetAnalysisSession())
 	{
-		TraceServices::FAnalysisSessionReadScope SessionReadScope(*Session);
-
 		if (const IStateTreeTraceProvider* Provider = Session->ReadProvider<IStateTreeTraceProvider>(FStateTreeTraceProvider::ProviderName))
 		{
+			TraceServices::FProviderReadScopeLock ProviderReadScope(*Provider);
 			return Provider->GetInstanceDescriptor(InstanceId);
 		}
 	}
@@ -259,10 +259,9 @@ void FStateTreeDebugger::GetSessionInstances(TArray<UE::StateTreeDebugger::FInst
 PRAGMA_DISABLE_DEPRECATION_WARNINGS
 	if (const TraceServices::IAnalysisSession* Session = GetAnalysisSession())
 	{
-		TraceServices::FAnalysisSessionReadScope SessionReadScope(*Session);
-
 		if (const IStateTreeTraceProvider* Provider = Session->ReadProvider<IStateTreeTraceProvider>(FStateTreeTraceProvider::ProviderName))
 		{
+			TraceServices::FProviderReadScopeLock ProviderReadScope(*Provider);
 			Provider->GetInstances(OutInstances);
 		}
 	}
@@ -273,10 +272,9 @@ void FStateTreeDebugger::GetSessionInstanceDescriptors(TArray<const TSharedRef<c
 {
 	if (const TraceServices::IAnalysisSession* Session = GetAnalysisSession())
 	{
-		TraceServices::FAnalysisSessionReadScope SessionReadScope(*Session);
-
 		if (const IStateTreeTraceProvider* Provider = Session->ReadProvider<IStateTreeTraceProvider>(FStateTreeTraceProvider::ProviderName))
 		{
+			TraceServices::FProviderReadScopeLock ProviderReadScope(*Provider);
 			Provider->GetInstances(OutInstances);
 		}
 	}
@@ -562,10 +560,9 @@ bool FStateTreeDebugger::IsActiveInstance(const double Time, const FStateTreeIns
 {
 	if (const TraceServices::IAnalysisSession* Session = GetAnalysisSession())
 	{
-		TraceServices::FAnalysisSessionReadScope SessionReadScope(*Session);
-
 		if (const IStateTreeTraceProvider* Provider = Session->ReadProvider<IStateTreeTraceProvider>(FStateTreeTraceProvider::ProviderName))
 		{
+			TraceServices::FProviderReadScopeLock ProviderReadScope(*Provider);
 			const TSharedPtr<const UE::StateTreeDebugger::FInstanceDescriptor> Descriptor = Provider->GetInstanceDescriptor(InstanceId);
 			return Descriptor.IsValid() && Descriptor->Lifetime.Contains(Time);
 		}
@@ -801,6 +798,7 @@ void FStateTreeDebugger::ReadTrace(const uint64 FrameIndex)
 {
 	if (const TraceServices::IAnalysisSession* Session = GetAnalysisSession())
 	{
+		// UE::StateTreeDebugger::ReadTrace also uses the FrameProvider so we scope the lock for the whole operation
 		TraceServices::FAnalysisSessionReadScope SessionReadScope(*Session);
 
 		const TraceServices::IFrameProvider& FrameProvider = TraceServices::ReadFrameProvider(*Session);
@@ -1059,6 +1057,7 @@ bool ProcessEvent(
 void AddEvents(
 	const double InStartTime,
 	const double InEndTime,
+	const TraceServices::IAnalysisSession& Session,
 	const TraceServices::IFrameProvider& InFrameProvider,
 	const FStateTreeInstanceDebugId InInstanceId,
 	const IStateTreeTraceProvider::FEventsTimeline& TimelineData,
@@ -1075,7 +1074,7 @@ void AddEvents(
 	Frame.Index = INDEX_NONE;
 
 	TimelineData.EnumerateEvents(InStartTime, InEndTime,
-		[&InFrameProvider, &Frame, EventCollection, TraceReader](const double EventStartTime, const double EventEndTime, uint32 InDepth, const FStateTreeTraceEventVariantType& Event)
+		[&InFrameProvider, &Frame, EventCollection, TraceReader, &Session](const double EventStartTime, const double EventEndTime, uint32 InDepth, const FStateTreeTraceEventVariantType& Event)
 		{
 			bool bValidFrame = true;
 
@@ -1083,6 +1082,7 @@ void AddEvents(
 			if (Frame.Index == INDEX_NONE ||
 				(EventEndTime < Frame.StartTime || Frame.EndTime < EventStartTime))
 			{
+				TraceServices::FAnalysisSessionReadScope SessionReadScope(Session);
 				bValidFrame = InFrameProvider.GetFrameFromTime(TraceFrameType_Game, EventStartTime, Frame);
 
 				if (bValidFrame == false)
@@ -1118,31 +1118,33 @@ void ReadTrace(
 	const FTraceFilter& Filter,
 	double& OutLastTraceReadTime)
 {
-	TraceServices::FAnalysisSessionReadScope SessionReadScope(Session);
-
+	TraceServices::FFrame TargetFrame;
+	bool bValidFrame = true;
 	const TraceServices::IFrameProvider& FrameProvider = TraceServices::ReadFrameProvider(Session);
 
-	TraceServices::FFrame TargetFrame;
-	if (FrameProvider.GetFrameFromTime(TraceFrameType_Game, ScrubTime, TargetFrame))
 	{
-		// Process only completed frames
-		bool bValidFrame = true;
-		if (TargetFrame.EndTime == std::numeric_limits<double>::infinity())
+		TraceServices::FAnalysisSessionReadScope SessionReadScope(Session);
+		if (FrameProvider.GetFrameFromTime(TraceFrameType_Game, ScrubTime, TargetFrame))
 		{
-			if (const TraceServices::FFrame* PreviousCompleteFrame = FrameProvider.GetFrame(TraceFrameType_Game, TargetFrame.Index - 1))
-			{
-				TargetFrame = *PreviousCompleteFrame;
-			}
-			else
+			// Process only completed frames
+
+			if (TargetFrame.EndTime == std::numeric_limits<double>::infinity())
 			{
-				bValidFrame = false;
+				if (const TraceServices::FFrame* PreviousCompleteFrame = FrameProvider.GetFrame(TraceFrameType_Game, TargetFrame.Index - 1))
+				{
+					TargetFrame = *PreviousCompleteFrame;
+				}
+				else
+				{
+					bValidFrame = false;
+				}
 			}
 		}
+	}
 
-		if (bValidFrame)
-		{
-			ReadTrace(Session, FrameProvider, TargetFrame, TraceReader, Filter, OutLastTraceReadTime);
-		}
+	if (bValidFrame)
+	{
+		ReadTrace(Session, FrameProvider, TargetFrame, TraceReader, Filter, OutLastTraceReadTime);
 	}
 }
 
@@ -1154,26 +1156,31 @@ void ReadTrace(
 	const FTraceFilter& Filter,
 	double& OutLastTraceReadTime)
 {
+	bool bValidLastReadFrame;
 	TraceServices::FFrame LastReadFrame;
-	const bool bValidLastReadFrame = FrameProvider.GetFrameFromTime(TraceFrameType_Game, OutLastTraceReadTime, LastReadFrame);
+	{
+		TraceServices::FAnalysisSessionReadScope SessionReadScope(Session);
+		bValidLastReadFrame = FrameProvider.GetFrameFromTime(TraceFrameType_Game, OutLastTraceReadTime, LastReadFrame);
+	}
 	if (OutLastTraceReadTime == 0 || (bValidLastReadFrame && Frame.Index > LastReadFrame.Index))
 	{
 		if (const IStateTreeTraceProvider* Provider = Session.ReadProvider<IStateTreeTraceProvider>(FStateTreeTraceProvider::ProviderName))
 		{
+			TraceServices::FProviderReadScopeLock ProviderReadScope(*Provider);
 			if (Filter.Instance.IsValid())
 			{
 				Provider->ReadTimelines(Filter.Instance,
-					[StartTime = OutLastTraceReadTime, EndTime = Frame.EndTime, &FrameProvider, TraceReader](const FStateTreeInstanceDebugId InstanceId, const IStateTreeTraceProvider::FEventsTimeline& TimelineData)
+					[StartTime = OutLastTraceReadTime, EndTime = Frame.EndTime, &FrameProvider, TraceReader, &Session](const FStateTreeInstanceDebugId InstanceId, const IStateTreeTraceProvider::FEventsTimeline& TimelineData)
 					{
-						AddEvents(StartTime, EndTime, FrameProvider, InstanceId, TimelineData, TraceReader);
+						AddEvents(StartTime, EndTime, Session, FrameProvider, InstanceId, TimelineData, TraceReader);
 					});
 			}
 			else if (Filter.Asset)
 			{
 				Provider->ReadTimelines(*Filter.Asset,
-					[StartTime = OutLastTraceReadTime, EndTime = Frame.EndTime, &FrameProvider, TraceReader](const FStateTreeInstanceDebugId InstanceId, const IStateTreeTraceProvider::FEventsTimeline& TimelineData)
+					[StartTime = OutLastTraceReadTime, EndTime = Frame.EndTime, &FrameProvider, TraceReader, &Session](const FStateTreeInstanceDebugId InstanceId, const IStateTreeTraceProvider::FEventsTimeline& TimelineData)
 					{
-						AddEvents(StartTime, EndTime, FrameProvider, InstanceId, TimelineData, TraceReader);
+						AddEvents(StartTime, EndTime, Session, FrameProvider, InstanceId, TimelineData, TraceReader);
 					});
 			}
 
diff --git a/Engine/Plugins/Runtime/StateTree/Source/StateTreeModule/Private/Debugger/StateTreeTraceProvider.cpp b/Engine/Plugins/Runtime/StateTree/Source/StateTreeModule/Private/Debugger/StateTreeTraceProvider.cpp
--- a/Engine/Plugins/Runtime/StateTree/Source/StateTreeModule/Private/Debugger/StateTreeTraceProvider.cpp
+++ b/Engine/Plugins/Runtime/StateTree/Source/StateTreeModule/Private/Debugger/StateTreeTraceProvider.cpp
@@ -5,6 +5,11 @@
 #include "Debugger/StateTreeTraceProvider.h"
 #include "Debugger/StateTreeDebugger.h"
 
+namespace TraceServices
+{
+	thread_local FProviderLock::FThreadLocalState GStateTreeProviderLockState;
+}
+
 FName FStateTreeTraceProvider::ProviderName("StateTreeDebuggerProvider");
 
 #define LOCTEXT_NAMESPACE "StateTreeDebuggerProvider"
@@ -16,7 +21,7 @@ FStateTreeTraceProvider::FStateTreeTraceProvider(TraceServices::IAnalysisSession
 
 bool FStateTreeTraceProvider::ReadTimelines(const FStateTreeInstanceDebugId InstanceId, const TFunctionRef<void(const FStateTreeInstanceDebugId ProcessedInstanceId, const FEventsTimeline&)> Callback) const
 {
-	Session.ReadAccessCheck();
+	ReadAccessCheck();
 
 	// Read specific timeline if specified
 	if (InstanceId.IsValid())
@@ -46,7 +51,7 @@ bool FStateTreeTraceProvider::ReadTimelines(const FStateTreeInstanceDebugId Inst
 
 bool FStateTreeTraceProvider::ReadTimelines(const UStateTree& StateTree, TFunctionRef<void(const FStateTreeInstanceDebugId ProcessedInstanceId, const FEventsTimeline&)> Callback) const
 {
-	Session.ReadAccessCheck();
+	ReadAccessCheck();
 
 	for (auto It = InstanceIdToDebuggerEntryTimelines.CreateConstIterator(); It; ++It)
 	{
@@ -64,7 +69,7 @@ bool FStateTreeTraceProvider::ReadTimelines(const UStateTree& StateTree, TFuncti
 
 void FStateTreeTraceProvider::AppendEvent(const FStateTreeInstanceDebugId InInstanceId, const double InTime, const FStateTreeTraceEventVariantType& InEvent)
 {
-	Session.WriteAccessCheck();
+	EditAccessCheck();
 
 	// It is currently possible to receive events from an instance without receiving event `EStateTreeTraceEventType::Push` first
 	// (i.e. traces were activated after the statetree instance execution was started).    
@@ -73,8 +78,11 @@ void FStateTreeTraceProvider::AppendEvent(const FStateTreeInstanceDebugId InInst
 	{
 		EventsTimelines[*IndexPtr]->AppendEvent(InTime, InEvent);
 	}
-	
-	Session.UpdateDurationSeconds(InTime);
+
+	{
+		TraceServices::FAnalysisSessionEditScope SessionEditScope(Session);
+		Session.UpdateDurationSeconds(InTime);
+	}
 }
 
 void FStateTreeTraceProvider::AppendInstanceEvent(
@@ -85,6 +93,7 @@ void FStateTreeTraceProvider::AppendInstanceEvent(
 	const double InWorldRecordingTime,
 	const EStateTreeTraceEventType InEventType)
 {
+	EditAccessCheck();
 	using namespace UE::StateTreeDebugger;
 
 	if (InEventType == EStateTreeTraceEventType::Push)
@@ -151,11 +160,15 @@ void FStateTreeTraceProvider::AppendInstanceEvent(
 		}
 	}
 
-	Session.UpdateDurationSeconds(InTime);
+	{
+		TraceServices::FAnalysisSessionEditScope SessionEditScope(Session);
+		Session.UpdateDurationSeconds(InTime);
+	}
 }
 
 void FStateTreeTraceProvider::AppendAssetDebugId(const UStateTree* InStateTree, const FStateTreeIndex16 AssetDebugId)
 {
+	EditAccessCheck();
 	TWeakObjectPtr<const UStateTree> WeakStateTree;
 	if (ensureMsgf(AssetDebugId.IsValid(), TEXT("Expecting valid asset debug Id."))
 		&& !GetAssetFromDebugId(AssetDebugId, WeakStateTree))
@@ -166,6 +179,7 @@ void FStateTreeTraceProvider::AppendAssetDebugId(const UStateTree* InStateTree,
 
 bool FStateTreeTraceProvider::GetAssetFromDebugId(const FStateTreeIndex16 AssetDebugId, TWeakObjectPtr<const UStateTree>& WeakStateTree) const
 {
+	ReadAccessCheck();
 	verifyf(AssetDebugId.IsValid(), TEXT("Expecting valid asset debug Id."));
 	const FStateTreeDebugIdPair* ExistingPair = StateTreeAssets.FindByPredicate([AssetDebugId](const FStateTreeDebugIdPair& Pair)
 	{
@@ -179,6 +193,7 @@ bool FStateTreeTraceProvider::GetAssetFromDebugId(const FStateTreeIndex16 AssetD
 
 bool FStateTreeTraceProvider::GetAssetFromInstanceId(const FStateTreeInstanceDebugId InstanceId, TWeakObjectPtr<const UStateTree>& WeakStateTree) const
 {
+	ReadAccessCheck();
 	if (const uint32* IndexPtr = InstanceIdToDebuggerEntryTimelines.Find(InstanceId))
 	{
 		check(Descriptors.Num() == EventsTimelines.Num());
@@ -191,6 +206,7 @@ bool FStateTreeTraceProvider::GetAssetFromInstanceId(const FStateTreeInstanceDeb
 
 TSharedPtr<const UE::StateTreeDebugger::FInstanceDescriptor> FStateTreeTraceProvider::GetInstanceDescriptor(const FStateTreeInstanceDebugId InstanceId) const
 {
+	ReadAccessCheck();
 	const TSharedRef<UE::StateTreeDebugger::FInstanceDescriptor>* Descriptor = Descriptors.FindByPredicate([InstanceId](const TSharedRef<const UE::StateTreeDebugger::FInstanceDescriptor>& Descriptor)
 	{
 		return Descriptor.Get().Id == InstanceId;
@@ -201,6 +217,7 @@ TSharedPtr<const UE::StateTreeDebugger::FInstanceDescriptor> FStateTreeTraceProv
 
 void FStateTreeTraceProvider::GetInstances(TArray<const TSharedRef<const UE::StateTreeDebugger::FInstanceDescriptor>>& OutInstances) const
 {
+	ReadAccessCheck();
 	OutInstances.Reserve(Descriptors.Num());
 	Algo::Transform(Descriptors, OutInstances, [](const TSharedRef<UE::StateTreeDebugger::FInstanceDescriptor>& Descriptor)
 	{
diff --git a/Engine/Plugins/Runtime/StateTree/Source/StateTreeModule/Public/Debugger/StateTreeTraceProvider.h b/Engine/Plugins/Runtime/StateTree/Source/StateTreeModule/Public/Debugger/StateTreeTraceProvider.h
--- a/Engine/Plugins/Runtime/StateTree/Source/StateTreeModule/Public/Debugger/StateTreeTraceProvider.h
+++ b/Engine/Plugins/Runtime/StateTree/Source/StateTreeModule/Public/Debugger/StateTreeTraceProvider.h
@@ -5,20 +5,53 @@
 #if WITH_STATETREE_TRACE_DEBUGGER
 
 #include "IStateTreeTraceProvider.h"
+#include "Common/ProviderLock.h"
 #include "Model/PointTimeline.h"
 #include "StateTreeIndexTypes.h"
 
-namespace TraceServices { class IAnalysisSession; }
+namespace TraceServices
+{
+	class IAnalysisSession;
+	extern thread_local FProviderLock::FThreadLocalState GStateTreeProviderLockState;
+}
 class UStateTree;
 struct FStateTreeInstanceDebugId;
 namespace UE::StateTreeDebugger { struct FInstanceDescriptor; }
 
-class FStateTreeTraceProvider : public IStateTreeTraceProvider
+class FStateTreeTraceProvider : public IStateTreeTraceProvider, public TraceServices::IEditableProvider
 {
 public:
 	STATETREEMODULE_API static FName ProviderName;
 
 	explicit FStateTreeTraceProvider(TraceServices::IAnalysisSession& InSession);
+
+	//~ IProvider (read lock)
+	virtual void BeginRead() const override
+	{
+		Lock.BeginRead(TraceServices::GStateTreeProviderLockState);
+	}
+	virtual void EndRead() const override
+	{
+		Lock.EndRead(TraceServices::GStateTreeProviderLockState);
+	}
+	virtual void ReadAccessCheck() const override
+	{
+		Lock.ReadAccessCheck(TraceServices::GStateTreeProviderLockState);
+	}
+
+	//~ IEditableProvider (edit lock)
+	virtual void BeginEdit() const override
+	{
+		Lock.BeginWrite(TraceServices::GStateTreeProviderLockState);
+	}
+	virtual void EndEdit() const override
+	{
+		Lock.EndWrite(TraceServices::GStateTreeProviderLockState);
+	}
+	virtual void EditAccessCheck() const override
+	{
+		Lock.WriteAccessCheck(TraceServices::GStateTreeProviderLockState);
+	}
 	
 	void AppendEvent(FStateTreeInstanceDebugId InInstanceId, double InTime, const FStateTreeTraceEventVariantType& InEvent);
 	void AppendInstanceEvent(
@@ -42,6 +75,7 @@ protected:
 
 private:
 	TraceServices::IAnalysisSession& Session;
+	mutable TraceServices::FProviderLock Lock;
 
 	TMap<FStateTreeInstanceDebugId, uint32> InstanceIdToDebuggerEntryTimelines;
 	TArray<TSharedRef<UE::StateTreeDebugger::FInstanceDescriptor>> Descriptors;
