UAF: Implicit cross-skeleton remapping at graph input boundary

Remaps poses automatically when a foreign skeleton enters a graph via PushPose,
so node authors never deal with cross-skeleton concerns.

- FRemapPoseData: Add RemapPose(FLODPose&) overload for zero-alloc remap into
pre-prepared targets (heap or stack). Templatize RemapAttributes.
- FKeyframeState::CopyFrom(FAnimNextGraphLODPose): Transparently handles
same-skeleton copy or cross-skeleton remap, including bones, attributes, curves.
- PushPose::Execute: Simplified to one-liner using Keyframe.CopyFrom().
- Existing RemapPose(FLODPoseHeap&) now delegates to the base overload.

#Synced-CL 52100020

diff --git a/Engine/Plugins/Experimental/UAF/UAF/Source/UAF/Private/Graph/RigUnit_AnimNextRemapPose.cpp b/Engine/Plugins/Experimental/UAF/UAF/Source/UAF/Private/Graph/RigUnit_AnimNextRemapPose.cpp
--- a/Engine/Plugins/Experimental/UAF/UAF/Source/UAF/Private/Graph/RigUnit_AnimNextRemapPose.cpp
+++ b/Engine/Plugins/Experimental/UAF/UAF/Source/UAF/Private/Graph/RigUnit_AnimNextRemapPose.cpp
@@ -2,7 +2,6 @@
 
 #include "RigUnit_AnimNextRemapPose.h"
 
-#include "GenerationTools.h"
 #include "Components/SkeletalMeshComponent.h"
 #include "Engine/SkeletalMesh.h"
 #include "UAFLogging.h"
@@ -30,30 +29,9 @@ FRigUnit_AnimNextRemapPose_Execute()
 
 	FMemMark MemMark(FMemStack::Get());
 
-	const UE::UAF::FLODPoseHeap& SourcePose = Pose.LODPose;
-	const UE::UAF::FReferencePose& SourceRefPose = SourcePose.GetRefPose();
-	const USkeletalMesh* SourceMesh = SourceRefPose.SkeletalMesh.Get();
-
 	const UE::UAF::FReferencePose& TargetRefPose = TargetAnimGraphRefPose.ReferencePose.GetRef<UE::UAF::FReferencePose>();
-	UE::UAF::FLODPoseHeap& TargetPose = Result.LODPose;
-	const USkeletalMesh* TargetMesh = TargetRefPose.SkeletalMesh.Get();
 
-	if (SourceMesh == TargetMesh)
-	{
-		TargetPose.PrepareForLOD(TargetRefPose, SourcePose.LODLevel, /*bSetRefPose=*/false, SourcePose.IsAdditive());
-		// Just copy the pose to the target in case there is nothing to convert.
-		TargetPose.CopyFrom(SourcePose);
-	}
-	else
-	{
-		// Pre-create the mapping and cache it to avoid runtime lookups.
-		if (RemapPoseData.ShouldReinit(SourceRefPose, TargetRefPose))
-		{
-			RemapPoseData.Reinit(SourceRefPose, TargetRefPose);
-		}
-
-		Result.Curves.CopyFrom(Pose.Curves);
-		RemapPoseData.RemapPose(SourcePose, TargetPose);
-		RemapPoseData.RemapAttributes(Pose.LODPose, Pose.Attributes, Result.LODPose, Result.Attributes);
-	}
-}
\ No newline at end of file
+	// CopyFrom handles both same-skeleton (fast copy) and cross-skeleton (remap via pool) transparently
+	Result.LODPose.PrepareForLOD(TargetRefPose, Pose.LODPose.LODLevel, /*bSetRefPose=*/false, Pose.LODPose.IsAdditive());
+	Result.CopyFrom(Pose);
+}
diff --git a/Engine/Plugins/Experimental/UAF/UAF/Source/UAF/Private/Graph/RigUnit_AnimNextRemapPose.h b/Engine/Plugins/Experimental/UAF/UAF/Source/UAF/Private/Graph/RigUnit_AnimNextRemapPose.h
--- a/Engine/Plugins/Experimental/UAF/UAF/Source/UAF/Private/Graph/RigUnit_AnimNextRemapPose.h
+++ b/Engine/Plugins/Experimental/UAF/UAF/Source/UAF/Private/Graph/RigUnit_AnimNextRemapPose.h
@@ -5,14 +5,17 @@
 #include "AnimNextExecuteContext.h"
 #include "Graph/AnimNext_LODPose.h"
 #include "Graph/RigUnit_AnimNextBase.h"
-#include "RemapPoseData.h"
 #include "AnimNextStats.h"
 
 #include "RigUnit_AnimNextRemapPose.generated.h"
 
 DECLARE_CYCLE_STAT_EXTERN(TEXT("UAF RigUnit: Remap Pose"), STAT_AnimNext_RigUnit_RemapPose, STATGROUP_AnimNext, UAF_API);
 
-/** Remaps an anim graph pose from to another skeletal mesh component. */
+/** Remaps an anim graph pose from to another skeletal mesh component.
+ *  Note: Cross-skeleton remapping now happens automatically in FAnimNextGraphLODPose::CopyFrom
+ *  and FKeyframeState::CopyFrom via the FRemapPoseDataPool. This node is only needed for
+ *  explicit control over where remapping occurs in the system and might be deprecated later on.
+ */
 USTRUCT(meta=(DisplayName="Remap Pose", Category="Animation Graph", NodeColor="0, 1, 1", Keywords="Output,Pose,Port"))
 struct FRigUnit_AnimNextRemapPose : public FRigUnit_AnimNextBase
 {
@@ -35,6 +38,4 @@ struct FRigUnit_AnimNextRemapPose : public FRigUnit_AnimNextBase
 	UPROPERTY(EditAnywhere, DisplayName = "Execute", Category = "BeginExecution", meta = (Input, Output))
 	FAnimNextExecuteContext ExecuteContext;
 
-	UPROPERTY(meta = (Hidden))
-	FRemapPoseData RemapPoseData; // TODO: Store this per remap node or inside the data registry? We don't have ref counting ref poses and don't know when to remove these though.
 };
diff --git a/Engine/Plugins/Experimental/UAF/UAF/Source/UAF/Private/Graph/RigUnit_AnimNextRetargetPose.cpp b/Engine/Plugins/Experimental/UAF/UAF/Source/UAF/Private/Graph/RigUnit_AnimNextRetargetPose.cpp
--- a/Engine/Plugins/Experimental/UAF/UAF/Source/UAF/Private/Graph/RigUnit_AnimNextRetargetPose.cpp
+++ b/Engine/Plugins/Experimental/UAF/UAF/Source/UAF/Private/Graph/RigUnit_AnimNextRetargetPose.cpp
@@ -2,10 +2,10 @@
 
 #include "RigUnit_AnimNextRetargetPose.h"
 
-#include "GenerationTools.h"
 #include "Components/SkeletalMeshComponent.h"
 #include "Engine/SkeletalMesh.h"
 #include "AnimationRuntime.h"
+#include "RemapPoseDataPool.h"
 #include "RetargetingTools.h"
 #include "Animation/SkeletonRemappingRegistry.h"
 #include "Animation/Skeleton.h"
@@ -42,8 +42,6 @@ FRigUnit_AnimNextRetargetPose_Execute()
 	const UE::UAF::FReferencePose& TargetRefPose = TargetAnimGraphRefPose.ReferencePose.GetRef<UE::UAF::FReferencePose>();
 	UE::UAF::FLODPoseHeap& TargetPose = Result.LODPose;
 	const USkeletalMesh* TargetMesh = TargetRefPose.SkeletalMesh.Get();
-	FRemapPoseData& RemapPoseData = RetargetPoseData.RemapPoseData;
-
 	if (SourceMesh == TargetMesh)
 	{
 		TargetPose.PrepareForLOD(TargetRefPose, SourcePose.LODLevel, /*bSetRefPose=*/false, SourcePose.IsAdditive());
@@ -61,16 +59,12 @@ FRigUnit_AnimNextRetargetPose_Execute()
 	const bool bUseSourceRetargetModes = TargetSkeleton->GetUseRetargetModesFromCompatibleSkeleton();
 	const bool bDisableRetargeting = TargetPose.GetDisableRetargeting();
 	
-	// 1) Initialize or refresh cached source target bone mappings to avoid runtime lookups.
-	if (RemapPoseData.ShouldReinit(SourceRefPose, TargetRefPose))
-	{
-		RemapPoseData.Reinit(SourceRefPose, TargetRefPose);
-	}
+	// 1) Remap bones + attributes + curves in one call via the pool.
+	Result.LODPose.PrepareForLOD(TargetRefPose, SourcePose.LODLevel, /*bSetRefPose=*/false, SourcePose.IsAdditive());
+	Result.CopyFrom(Pose);
 
-	// 2) Copy curves and remap bone transforms and attributes to the target pose.
-	Result.Curves.CopyFrom(Pose.Curves);
-	RemapPoseData.RemapPose(SourcePose, TargetPose);
-	RemapPoseData.RemapAttributes(Pose.LODPose, Pose.Attributes, Result.LODPose, Result.Attributes);
+	// 2) Get remap data from the pool for retarget steps (O(1) lookup, already cached by CopyFrom).
+	const FRemapPoseData& RemapPoseData = UE::UAF::FRemapPoseDataPool::Get().GetRemapData(SourceRefPose, TargetRefPose);
 
 	// 3) Build or update per-bone retargeting pairs for current source/target LODs.
 	if (RetargetPoseData.PairsInitForSourcePoseLOD != SourcePose.LODLevel ||
diff --git a/Engine/Plugins/Experimental/UAF/UAF/Source/UAF/Private/Graph/RigUnit_AnimNextRetargetPose.h b/Engine/Plugins/Experimental/UAF/UAF/Source/UAF/Private/Graph/RigUnit_AnimNextRetargetPose.h
--- a/Engine/Plugins/Experimental/UAF/UAF/Source/UAF/Private/Graph/RigUnit_AnimNextRetargetPose.h
+++ b/Engine/Plugins/Experimental/UAF/UAF/Source/UAF/Private/Graph/RigUnit_AnimNextRetargetPose.h
@@ -5,7 +5,6 @@
 #include "AnimNextExecuteContext.h"
 #include "Graph/AnimNext_LODPose.h"
 #include "Graph/RigUnit_AnimNextBase.h"
-#include "RemapPoseData.h"
 #include "AnimEncoding.h"
 #include "AnimNextStats.h"
 
@@ -18,8 +17,6 @@ struct FRetargetPoseData
 {
 	GENERATED_BODY()
 
-	FRemapPoseData RemapPoseData; // TODO: Store this per remap node or inside the data registry? We don't have ref counting ref poses and don't know when to remove these though.
-
 	BoneTrackArray SkeletonPairs;
 	BoneTrackArray AnimScalePairs;
 	BoneTrackArray AnimRelativePairs;
diff --git a/Engine/Plugins/Experimental/UAF/UAF/Source/UAF/Private/RemapPoseData.cpp b/Engine/Plugins/Experimental/UAF/UAF/Source/UAF/Private/RemapPoseData.cpp
--- a/Engine/Plugins/Experimental/UAF/UAF/Source/UAF/Private/RemapPoseData.cpp
+++ b/Engine/Plugins/Experimental/UAF/UAF/Source/UAF/Private/RemapPoseData.cpp
@@ -1,4 +1,4 @@
-﻿// Copyright Epic Games, Inc. All Rights Reserved.
+// Copyright Epic Games, Inc. All Rights Reserved.
 
 #include "RemapPoseData.h"
 #include "AnimNextStats.h"
@@ -105,14 +105,14 @@ void FRemapPoseData::Reinit(const UE::UAF::FReferencePose& InSourceRefPose,
 	}
 }
 
-FTransform FRemapPoseData::RecursiveCalcModelspaceTransform(const UE::UAF::FLODPoseHeap& Pose, FBoneIndexType BoneIndex) const
+FTransform FRemapPoseData::RecursiveCalcModelspaceTransform(const UE::UAF::FLODPose& Pose, FBoneIndexType BoneIndex) const
 {
 	if (BoneIndex == (uint16)INDEX_NONE)
 	{
 		return FTransform::Identity;
 	}
 
-	const FTransform LocalTransform = Pose.LocalTransforms[BoneIndex];
+	const FTransform LocalTransform = Pose.LocalTransformsView[BoneIndex];
 
 	const TArrayView<const FBoneIndexType> BoneToParentBoneIndexMap = Pose.GetLODBoneIndexToParentLODBoneIndexMap();
 	FBoneIndexType ParentBoneIndex = BoneToParentBoneIndexMap[BoneIndex];
@@ -125,22 +125,25 @@ FTransform FRemapPoseData::RecursiveCalcModelspaceTransform(const UE::UAF::FLODP
 	return LocalTransform * ParentModelTransform;
 }
 
-// TODO: Source pose dictates the LOD level of the target at the moment. What if the LOD levels of source and target are not in sync? We probably need to use the target skel mesh component LOD level when preparing the target pose.
-void FRemapPoseData::RemapPose(const UE::UAF::FLODPoseHeap& SourcePose,
-	UE::UAF::FLODPoseHeap& OutTargetPose) const
+void FRemapPoseData::RemapPose(const UE::UAF::FLODPoseHeap& SourcePose, UE::UAF::FLODPoseHeap& OutTargetPose) const
 {
-	SCOPE_CYCLE_COUNTER(STAT_AnimNext_RemapPose_Mesh2Mesh);
-
-	const int32 SourceLODLevel = SourcePose.LODLevel;
 	const int32 TargetLODLevel = TargetRefPose->GetSourceLODLevel();
-
 	const bool bIsAdditive = SourcePose.IsAdditive();
 
 	// Always prepare for LOD as we may not be copying in all of the target bones and it is just cheaper to copy them all in
 	OutTargetPose.PrepareForLOD(*TargetRefPose, TargetLODLevel, /*bSetRefPose=*/true, bIsAdditive);
-
 	OutTargetPose.Flags = SourcePose.Flags;
 
+	RemapBones(SourcePose, OutTargetPose);
+}
+
+void FRemapPoseData::RemapBones(const UE::UAF::FLODPose& SourcePose, UE::UAF::FLODPose& OutTargetPose) const
+{
+	SCOPE_CYCLE_COUNTER(STAT_AnimNext_RemapPose_Mesh2Mesh);
+
+	const int32 SourceLODLevel = SourcePose.LODLevel;
+	const int32 TargetLODLevel = OutTargetPose.LODLevel;
+
 	const TArray<FRemappedBone>& SourceToTargetBoneIndexMap = SourceToTargetBoneIndexMapPerLOD[SourceLODLevel][TargetLODLevel].BoneIndexMap;
 	const int32 NumToRemap = SourceToTargetBoneIndexMap.Num();
 
@@ -183,8 +186,7 @@ void FRemapPoseData::RemapPose(const UE::UAF::FLODPoseHeap& SourcePose,
 		const int32 SourceBoneIndex = SourceToTargetBoneIndexMap[MappingElementIndex].SourceBoneIndex;
 		const int32 TargetBoneIndex = SourceToTargetBoneIndexMap[MappingElementIndex].TargetBoneIndex;
 
-		const auto& TargetTransform = SourcePose.LocalTransformsView[SourceBoneIndex];
-		OutTargetPose.LocalTransforms[TargetBoneIndex] = TargetTransform;
+		OutTargetPose.LocalTransformsView[TargetBoneIndex] = SourcePose.LocalTransformsView[SourceBoneIndex];
 	}
 #endif
 
@@ -197,14 +199,15 @@ void FRemapPoseData::RemapPose(const UE::UAF::FLODPoseHeap& SourcePose,
 	{
 		// TODO: Replace with the generic method once we have access to the modelspace transform via the LOD pose.
 		const FTransform RootModelTransform = RecursiveCalcModelspaceTransform(SourcePose, RootBoneOnSourceMapped);
-		OutTargetPose.LocalTransforms[UE::UAF::FReferencePose::RootBoneIndex] = RootModelTransform;
+		OutTargetPose.LocalTransformsView[UE::UAF::FReferencePose::RootBoneIndex] = RootModelTransform;
 	}
 }
 
+template<typename OutContainerType>
 void FRemapPoseData::RemapAttributes(const UE::UAF::FLODPose& SourceLODPose,
 	const UE::Anim::FHeapAttributeContainer& InAttributes,
 	const UE::UAF::FLODPose& TargetLODPose,
-	UE::Anim::FHeapAttributeContainer& OutAttributes) const
+	OutContainerType& OutAttributes) const
 {
 	const USkeletalMesh* SourceSkeletalMesh = SourceRefPose->SkeletalMesh.Get();
 	if (!SourceSkeletalMesh)
@@ -257,3 +260,7 @@ void FRemapPoseData::RemapAttributes(const UE::UAF::FLODPose& SourceLODPose,
 		}
 	}
 }
+
+// Explicit instantiations
+template UAF_API void FRemapPoseData::RemapAttributes(const UE::UAF::FLODPose&, const UE::Anim::FHeapAttributeContainer&, const UE::UAF::FLODPose&, UE::Anim::FHeapAttributeContainer&) const;
+template UAF_API void FRemapPoseData::RemapAttributes(const UE::UAF::FLODPose&, const UE::Anim::FHeapAttributeContainer&, const UE::UAF::FLODPose&, UE::Anim::FStackAttributeContainer&) const;
diff --git a/Engine/Plugins/Experimental/UAF/UAF/Source/UAF/Public/RemapPoseData.h b/Engine/Plugins/Experimental/UAF/UAF/Source/UAF/Public/RemapPoseData.h
--- a/Engine/Plugins/Experimental/UAF/UAF/Source/UAF/Public/RemapPoseData.h
+++ b/Engine/Plugins/Experimental/UAF/UAF/Source/UAF/Public/RemapPoseData.h
@@ -57,19 +57,34 @@ struct FRemapPoseData
 
 	/**
 	 * Remap a pose using the previously cached mapping. Call this at runtime.
+	 * Prepares the target pose internally (allocates, fills ref pose for unmapped bones).
 	 */
-	void RemapPose(const UE::UAF::FLODPoseHeap& SourcePose,
+	UAF_API void RemapPose(const UE::UAF::FLODPoseHeap& SourcePose,
 		UE::UAF::FLODPoseHeap& OutTargetPose) const;
 
+	/**
+	 * Remap bone transforms from source into an already-prepared target pose.
+	 * The target must already have memory allocated and ref pose set for unmapped bones.
+	 * Operates on FLODPose base, works with any allocator (heap or stack).
+	 */
+	UAF_API void RemapBones(const UE::UAF::FLODPose& SourcePose, UE::UAF::FLODPose& OutTargetPose) const;
+
+	/**
+	 * Remap attributes from source into an already-prepared target attribute container.
+	 * Bone indices are translated from source LOD pose space to target LOD pose space via the skeleton bone name.
+	 * Attributes whose source bone has no counterpart in the target skeleton are silently dropped.
+	 * Templated to support both heap (FHeapAttributeContainer) and stack (FStackAttributeContainer) output containers.
+	 */
+	template<typename OutContainerType>
 	void RemapAttributes(const UE::UAF::FLODPose& SourceLODPose,
 		const UE::Anim::FHeapAttributeContainer& InAttributes,
 		const UE::UAF::FLODPose& TargetLODPose,
-		UE::Anim::FHeapAttributeContainer& OutAttributes) const;
+		OutContainerType& OutAttributes) const;
 
 	/**
 	 * Temporary calculate modelspace transform, until the LODPose is able to do that.
 	 */
-	FTransform RecursiveCalcModelspaceTransform(const UE::UAF::FLODPoseHeap& Pose, FBoneIndexType BoneIndex) const;
+	FTransform RecursiveCalcModelspaceTransform(const UE::UAF::FLODPose& Pose, FBoneIndexType BoneIndex) const;
 
 	const UE::UAF::FReferencePose* SourceRefPose = nullptr;
 	const UE::UAF::FReferencePose* TargetRefPose = nullptr;
diff --git a/Engine/Plugins/Experimental/UAF/UAFAnimGraph/Source/UAFAnimGraph/Private/EvaluationVM/KeyframeState.cpp b/Engine/Plugins/Experimental/UAF/UAFAnimGraph/Source/UAFAnimGraph/Private/EvaluationVM/KeyframeState.cpp
new file mode 100644
--- /dev/null
+++ b/Engine/Plugins/Experimental/UAF/UAFAnimGraph/Source/UAFAnimGraph/Private/EvaluationVM/KeyframeState.cpp
@@ -0,0 +1,42 @@
+// Copyright Epic Games, Inc. All Rights Reserved.
+
+#include "EvaluationVM/KeyframeState.h"
+#include "Graph/AnimNext_LODPose.h"
+#include "RemapPoseDataPool.h"
+#include "AnimNextAnimGraphStats.h"
+
+DEFINE_STAT(STAT_AnimNext_Keyframe_CopyFrom);
+
+namespace UE::UAF
+{
+	void FKeyframeState::CopyFrom(const FAnimNextGraphLODPose& Source)
+	{
+		SCOPE_CYCLE_COUNTER(STAT_AnimNext_Keyframe_CopyFrom);
+
+		if (!Source.LODPose.IsValid())
+		{
+			return;
+		}
+
+		const FReferencePose* SourceRef = Source.LODPose.RefPose;
+		const FReferencePose* TargetRef = Pose.RefPose;
+
+		if (SourceRef != nullptr && TargetRef != nullptr && SourceRef != TargetRef)
+		{
+			// Cross-skeleton: Fill ref pose transforms for unmapped bones, then remap mapped ones.
+			Pose.SetRefPose(Source.LODPose.IsAdditive());
+			Pose.Flags = Source.LODPose.Flags;
+
+			const FRemapPoseData& RemapData = FRemapPoseDataPool::Get().GetRemapData(*SourceRef, *TargetRef);
+			RemapData.RemapBones(Source.LODPose, Pose);
+			RemapData.RemapAttributes(Source.LODPose, Source.Attributes, Pose, Attributes);
+		}
+		else
+		{
+			Pose.CopyFrom(Source.LODPose);
+			Attributes.CopyFrom(Source.Attributes);
+		}
+
+		Curves.CopyFrom(Source.Curves);
+	}
+}
diff --git a/Engine/Plugins/Experimental/UAF/UAFAnimGraph/Source/UAFAnimGraph/Private/EvaluationVM/Tasks/PushPose.cpp b/Engine/Plugins/Experimental/UAF/UAFAnimGraph/Source/UAFAnimGraph/Private/EvaluationVM/Tasks/PushPose.cpp
--- a/Engine/Plugins/Experimental/UAF/UAFAnimGraph/Source/UAFAnimGraph/Private/EvaluationVM/Tasks/PushPose.cpp
+++ b/Engine/Plugins/Experimental/UAF/UAFAnimGraph/Source/UAFAnimGraph/Private/EvaluationVM/Tasks/PushPose.cpp
@@ -40,10 +40,7 @@ void FAnimNextPushPoseTask::Execute(UE::UAF::FEvaluationVM& VM) const
 	}
 
 	FKeyframeState Keyframe = VM.MakeUninitializedKeyframe(GraphPose->LODPose.IsAdditive());
-
-	Keyframe.Pose.CopyFrom(GraphPose->LODPose);
-	Keyframe.Curves.CopyFrom(GraphPose->Curves);
-	Keyframe.Attributes.CopyFrom(GraphPose->Attributes);
-
+	// CopyFrom transparently remaps if the source skeleton differs from the target (cross-skeleton input).
+	Keyframe.CopyFrom(*GraphPose);
 	VM.PushValue(KEYFRAME_STACK_NAME, MakeUnique<FKeyframeState>(MoveTemp(Keyframe)));
 }
diff --git a/Engine/Plugins/Experimental/UAF/UAFAnimGraph/Source/UAFAnimGraph/Public/AnimNextAnimGraphStats.h b/Engine/Plugins/Experimental/UAF/UAFAnimGraph/Source/UAFAnimGraph/Public/AnimNextAnimGraphStats.h
--- a/Engine/Plugins/Experimental/UAF/UAFAnimGraph/Source/UAFAnimGraph/Public/AnimNextAnimGraphStats.h
+++ b/Engine/Plugins/Experimental/UAF/UAFAnimGraph/Source/UAFAnimGraph/Public/AnimNextAnimGraphStats.h
@@ -10,3 +10,4 @@ DECLARE_CYCLE_STAT_EXTERN(TEXT("UAF: Graph RigVM"), STAT_AnimNext_Graph_RigVM, S
 DECLARE_CYCLE_STAT_EXTERN(TEXT("UAF: Execute Evaluation Program"), STAT_AnimNext_EvaluationProgram_Execute, STATGROUP_AnimNext, UAFANIMGRAPH_API);
 
 DECLARE_CYCLE_STAT_EXTERN(TEXT("UAF: Call Function"), STAT_AnimNext_CallFunction, STATGROUP_AnimNext, UAFANIMGRAPH_API);
+DECLARE_CYCLE_STAT_EXTERN(TEXT("UAF: Keyframe CopyFrom"), STAT_AnimNext_Keyframe_CopyFrom, STATGROUP_AnimNext, UAFANIMGRAPH_API);
diff --git a/Engine/Plugins/Experimental/UAF/UAFAnimGraph/Source/UAFAnimGraph/Public/EvaluationVM/KeyframeState.h b/Engine/Plugins/Experimental/UAF/UAFAnimGraph/Source/UAFAnimGraph/Public/EvaluationVM/KeyframeState.h
--- a/Engine/Plugins/Experimental/UAF/UAFAnimGraph/Source/UAFAnimGraph/Public/EvaluationVM/KeyframeState.h
+++ b/Engine/Plugins/Experimental/UAF/UAFAnimGraph/Source/UAFAnimGraph/Public/EvaluationVM/KeyframeState.h
@@ -8,6 +8,10 @@
 #include "Animation/AttributesRuntime.h"
 #include "LODPose.h"
 
+#define UE_API UAFANIMGRAPH_API
+
+struct FAnimNextGraphLODPose;
+
 namespace UE::UAF
 {
 	/*
@@ -30,5 +34,10 @@ namespace UE::UAF
 
 		[[nodiscard]] bool IsValid() const { return Pose.IsValid(); }
 		[[nodiscard]] bool IsAdditive() const { return Pose.IsAdditive(); }
+
+		/** Copy from a graph pose, transparently handling cross-skeleton remapping. */
+		UE_API void CopyFrom(const FAnimNextGraphLODPose& Source);
 	};
 }
+
+#undef UE_API
