Add root motion attribute by default to abstract skeleton set binding assets

#Synced-CL 52346511

diff --git a/Engine/Plugins/Experimental/UAF/UAF/Source/UAF/Private/UAF/AbstractSkeleton/AbstractSkeletonSetBinding.cpp b/Engine/Plugins/Experimental/UAF/UAF/Source/UAF/Private/UAF/AbstractSkeleton/AbstractSkeletonSetBinding.cpp
--- a/Engine/Plugins/Experimental/UAF/UAF/Source/UAF/Private/UAF/AbstractSkeleton/AbstractSkeletonSetBinding.cpp
+++ b/Engine/Plugins/Experimental/UAF/UAF/Source/UAF/Private/UAF/AbstractSkeleton/AbstractSkeletonSetBinding.cpp
@@ -2,6 +2,7 @@
 
 #include "UAF/AbstractSkeleton/AbstractSkeletonSetBinding.h"
 
+#include "Animation/AnimRootMotionProvider.h"
 #include "UAF/AbstractSkeleton/AbstractSkeletonSetCollection.h"
 #include "UAF/Attributes/AttributeBindingDataCache.h"
 #include "UAF/Attributes/EngineAttributes.h"
@@ -13,6 +14,7 @@ void UAbstractSkeletonSetBinding::PostLoad()
 	Super::PostLoad();
 
 #if WITH_EDITOR
+	AddDefaultAttributes();
 	FilterOrphanedBindings();
 #endif
 
@@ -128,44 +130,56 @@ const TConstArrayView<FAbstractSkeleton_BoneBinding> UAbstractSkeletonSetBinding
 
 // Attribute Bindings
 
-bool UAbstractSkeletonSetBinding::AddAttributeToSet(const FAnimationAttributeIdentifier Attribute, const FName SetName)
+void UAbstractSkeletonSetBinding::AddDefaultAttributes()
 {
-	if (!SetCollection)
-	{
-		return false;
-	}
-
-	if (!SetCollection->HasSet(SetName) && SetName != NAME_None)
+	if (!ContainsAttribute(UE::Anim::IAnimRootMotionProvider::AttributeName))
 	{
-		return false; // Set does not exist
+		AddAttributeToSet(FAnimationAttributeIdentifier(UE::Anim::IAnimRootMotionProvider::AttributeName, INDEX_NONE, NAME_None, FTransformAnimationAttribute::StaticStruct()), NAME_None);
 	}
+}
 
+bool UAbstractSkeletonSetBinding::AddAttributeToSet(const FAnimationAttributeIdentifier Attribute, const FName SetName)
+{
 	if (IsAttributeInSet(Attribute))
 	{
-		return false; // Already exists in a set
+		return false; // Already exists in a named set
 	}
 
-	if (IsAttributeInSet(Attribute, NAME_None))
+	if (SetName != NAME_None)
 	{
-		FAbstractSkeleton_AttributeBinding* UnboundAttribute = AttributeBindings.FindByPredicate([Attribute, SetName](const FAbstractSkeleton_AttributeBinding& Binding)
-		{
-			return Binding.Attribute == Attribute && Binding.SetName == NAME_None;
-		});
-
-		if (ensure(UnboundAttribute))
+		if (!SetCollection || !SetCollection->HasSet(SetName))
 		{
-			Modify();
-
-			UnboundAttribute->SetName = SetName;
-			OnBindingsChanged();
-			return true;
+			return false; // Set does not exist
 		}
-		else
+
+		if (IsAttributeInSet(Attribute, NAME_None))
 		{
-			return false;
+			// Our attribute existed in the everything set, update it
+			FAbstractSkeleton_AttributeBinding* UnboundAttribute = AttributeBindings.FindByPredicate([Attribute, SetName](const FAbstractSkeleton_AttributeBinding& Binding)
+				{
+					return Binding.Attribute == Attribute && Binding.SetName == NAME_None;
+				});
+
+			if (ensure(UnboundAttribute))
+			{
+				Modify();
+
+				UnboundAttribute->SetName = SetName;
+				OnBindingsChanged();
+				return true;
+			}
+			else
+			{
+				return false;
+			}
 		}
 	}
 
+	if (IsAttributeInSet(Attribute, NAME_None))
+	{
+		return false;	// Already exists in the everything set
+	}
+
 	Modify();
 
 	FAbstractSkeleton_AttributeBinding& Binding = AttributeBindings.AddDefaulted_GetRef();
diff --git a/Engine/Plugins/Experimental/UAF/UAF/Source/UAF/Private/UAF/UAFSkeletonUserData.cpp b/Engine/Plugins/Experimental/UAF/UAF/Source/UAF/Private/UAF/UAFSkeletonUserData.cpp
--- a/Engine/Plugins/Experimental/UAF/UAF/Source/UAF/Private/UAF/UAFSkeletonUserData.cpp
+++ b/Engine/Plugins/Experimental/UAF/UAF/Source/UAF/Private/UAF/UAFSkeletonUserData.cpp
@@ -210,6 +210,7 @@ UUAFSkeletonUserData* UUAFSkeletonUserData::FromSkeleton(TNonNullPtr<USkeleton>
 			(void)UserData->SetBinding->AtomicallyClearInternalFlags(EInternalObjectFlags::Async);
 
 			UserData->SetBinding->SetSkeleton(Skeleton);
+			UserData->SetBinding->AddDefaultAttributes();
 
 			CachedData.Skeleton = Skeleton;
 			CachedData.UserData = TStrongObjectPtr<UUAFSkeletonUserData>(UserData);
diff --git a/Engine/Plugins/Experimental/UAF/UAF/Source/UAF/Public/UAF/AbstractSkeleton/AbstractSkeletonSetBinding.h b/Engine/Plugins/Experimental/UAF/UAF/Source/UAF/Public/UAF/AbstractSkeleton/AbstractSkeletonSetBinding.h
--- a/Engine/Plugins/Experimental/UAF/UAF/Source/UAF/Public/UAF/AbstractSkeleton/AbstractSkeletonSetBinding.h
+++ b/Engine/Plugins/Experimental/UAF/UAF/Source/UAF/Public/UAF/AbstractSkeleton/AbstractSkeletonSetBinding.h
@@ -76,6 +76,9 @@ public:
 		
 	// Attribute Bindings
 
+	// Adds all the default attributes that should always be present
+	UE_API void AddDefaultAttributes();
+
 	UE_API bool AddAttributeToSet(const FAnimationAttributeIdentifier Attribute, const FName SetName);
 
 	// Note this function does two different things depending on the type of attribute provided.
diff --git a/Engine/Plugins/Experimental/UAF/UAF/Source/UAFEditor/Private/UAF/AbstractSkeleton/Sets/SetBindingFactory.cpp b/Engine/Plugins/Experimental/UAF/UAF/Source/UAFEditor/Private/UAF/AbstractSkeleton/Sets/SetBindingFactory.cpp
--- a/Engine/Plugins/Experimental/UAF/UAF/Source/UAFEditor/Private/UAF/AbstractSkeleton/Sets/SetBindingFactory.cpp
+++ b/Engine/Plugins/Experimental/UAF/UAF/Source/UAFEditor/Private/UAF/AbstractSkeleton/Sets/SetBindingFactory.cpp
@@ -27,6 +27,7 @@ UObject* UAbstractSkeletonSetBindingFactory::FactoryCreateNew(UClass* Class, UOb
 	}
 
 	TObjectPtr<UAbstractSkeletonSetBinding> NewSetBinding = NewObject<UAbstractSkeletonSetBinding>(InParent, Class, Name, FlagsToUse);
+	NewSetBinding->AddDefaultAttributes();
 
 	return NewSetBinding;
 }
