Automatically add UAF system as tick prerequisite for output skeletal mesh component

#Synced-CL 52292702

diff --git a/Engine/Plugins/Experimental/UAF/UAF/Source/UAF/Private/Component/AnimNextComponent.cpp b/Engine/Plugins/Experimental/UAF/UAF/Source/UAF/Private/Component/AnimNextComponent.cpp
--- a/Engine/Plugins/Experimental/UAF/UAF/Source/UAF/Private/Component/AnimNextComponent.cpp
+++ b/Engine/Plugins/Experimental/UAF/UAF/Source/UAF/Private/Component/AnimNextComponent.cpp
@@ -755,6 +755,12 @@ void UUAFComponent::SetupInputOutput()
 		SystemReference.AddSystemEventDependency(SystemReference.GetFirstUserEventName(), FUAFWeakSystemReference(InputSystem), InputSystem.GetLastUserEventName(), UE::UAF::ESystemDependency::Prerequisite);
 	}
 
+	// Setup system as tick prerequisite for output skeletal mesh component, this should subsequently impact any follower components/manual tick dependencies setup up against the skeletal mesh component
+	if (bValidOutputComponent)
+	{
+		AddComponentSubsequent(CachedSkeletalMeshComponent, SystemReference.GetLastUserEventName());
+	}
+
 	SystemReference.QueueTask([WeakThis = TWeakObjectPtr<UUAFComponent>(this), bValidOutputComponent, InputBindings = MoveTemp(InputBindings)](const UE::UAF::FModuleTaskContext& InContext)
 	{
 		if (UUAFComponent* This = WeakThis.Get())
