Synchronize columns for UAF Browser across docked, drawer, and sidebar.
Set defaults for UAF browser columns, set default view to columns.
Show path column in UAF Browser by default.
---
Add Callbacks to SAssetPicker to save and load settings for synchronization across asset view instances.
Add Callback for when column visibility changes in SAssetView

#tests Editor
#Synced-CL 52424183

diff --git a/Engine/Plugins/Experimental/UAF/UAF/Source/UAFEditor/Private/Common/SUAFBrowser.cpp b/Engine/Plugins/Experimental/UAF/UAF/Source/UAFEditor/Private/Common/SUAFBrowser.cpp
--- a/Engine/Plugins/Experimental/UAF/UAF/Source/UAFEditor/Private/Common/SUAFBrowser.cpp
+++ b/Engine/Plugins/Experimental/UAF/UAF/Source/UAFEditor/Private/Common/SUAFBrowser.cpp
@@ -13,6 +13,8 @@
 #include "IContentBrowserSingleton.h"
 #include "IWorkspaceEditor.h"
 #include "IWorkspaceEditorModule.h"
+#include "Misc/ConfigCacheIni.h"
+#include "SAssetView.h"
 #include "Sound/SoundWave.h"
 #include "Settings/UAFEditorProjectSettings.h"
 #include "Settings/UAFEditorUserSettings.h"
@@ -31,11 +33,25 @@
 namespace UE::UAF::Editor
 {
 
+
+SUAFBrowser::~SUAFBrowser()
+{
+	if (UUAFEditorUserSettings* UAFSettings = GetMutableDefault<UUAFEditorUserSettings>())
+	{
+		UAFSettings->OnBrowserHiddenColumnsChanged.RemoveAll(this);
+	}
+}
+
 void SUAFBrowser::Construct(const FArguments& InArgs, const TSharedPtr<UE::Workspace::IWorkspaceEditor>& InHostingApp, ETabState::Type InTabState)
 {
 	WeakHostingApp = InHostingApp;
 	TabState = InTabState;
 
+	if (UUAFEditorUserSettings* UAFSettings = GetMutableDefault<UUAFEditorUserSettings>())
+	{
+		UAFSettings->OnBrowserHiddenColumnsChanged.AddSP(this, &SUAFBrowser::OnExternalColumnVisibilityChanged);
+	}
+
 	FMargin DrawerPadding = TabState == ETabState::SidebarTab || TabState == ETabState::StatusbarTab
 		? 4.0f
 		: 0.0f;
@@ -132,13 +148,27 @@ TSharedRef<SWidget> SUAFBrowser::Construct_GetAssetPreviewNotSupportedWidget()
 
 void SUAFBrowser::OnOverrideAssetPickerConfig(FAssetPickerConfig& Config)
 {
-	// @TODO: Meet with product design & choose defaults
+	static const FString ContentBrowserIniSection = TEXT("ContentBrowser");
+	static const FString UAFBrowserConfigName = TEXT("UAFBrowser");
+
 	Config.OnExtendAssetPickerTopBar.BindSP(this, &SUAFBrowser::OnExtendAssetPickerTopBar);
-	Config.InitialAssetViewType = EAssetViewType::Column;
 	Config.bAddFilterUI = true;
 	Config.bAllowDragging = true;
-	//Config.bShowPathInColumnView = true;
-	//Config.bSortByPathInColumnView = true;
+	Config.bShowPathInColumnView = true;
+	Config.SaveSettingsName = UAFBrowserConfigName;
+
+	// Apply user-configurable defaults from UUAFEditorUserSettings.
+	if (const UUAFEditorUserSettings* UAFSettings = GetDefault<UUAFEditorUserSettings>())
+	{
+		// If the user hasn't set any view type, default to column view
+		int32 ViewType_Discarded = INDEX_NONE;
+		if (!GConfig->GetInt(*ContentBrowserIniSection, *(UAFBrowserConfigName + TEXT(".CurrentViewType")), ViewType_Discarded, GEditorPerProjectIni))
+		{
+			Config.InitialAssetViewType = EAssetViewType::Column;
+		}
+
+		Config.HiddenColumnNames = UAFSettings->UAFBrowserHiddenColumns;
+	}
 
 	Config.OnGetAssetContextMenu.Unbind();
 	Config.OnAssetsActivated.Unbind();
@@ -155,6 +185,9 @@ void SUAFBrowser::OnOverrideAssetPickerConfig(FAssetPickerConfig& Config)
 
 	Config.OnAssetSelected.BindSP(this, &SUAFBrowser::OnAssetSelected);
 	Config.OnAssetsDragged.BindSP(this, &SUAFBrowser::OnAssetsDragged);
+	Config.OnColumnVisibilityChanged.BindSP(this, &SUAFBrowser::OnBrowserColumnVisibilityChanged);
+	Config.SaveAssetViewSettingsDelegates.Add(&SaveAssetViewSettingsDelegate);
+	Config.LoadAssetViewSettingsDelegates.Add(&LoadAssetViewSettingsDelegate);
 }
 
 void SUAFBrowser::OnExtendAssetPickerTopBar(TSharedRef<SHorizontalBox> InAssetPickerTopBar)
@@ -313,6 +346,26 @@ void SUAFBrowser::OnAssetsDragged(const TArray<FAssetData>& InAssetData)
 	}
 }
 
+void SUAFBrowser::OnBrowserColumnVisibilityChanged(const TArray<FString>& HiddenColumnIds)
+{
+	if (UUAFEditorUserSettings* UAFSettings = GetMutableDefault<UUAFEditorUserSettings>())
+	{
+		UAFSettings->UAFBrowserHiddenColumns = HiddenColumnIds;
+		UAFSettings->SaveConfig();
+
+		// Write to .ini
+		SaveAssetViewSettingsDelegate.ExecuteIfBound();
+
+		UAFSettings->OnBrowserHiddenColumnsChanged.Broadcast(HiddenColumnIds);
+	}
+}
+
+void SUAFBrowser::OnExternalColumnVisibilityChanged(const TArray<FString>& HiddenColumnIds)
+{
+	// Read from .ini
+	LoadAssetViewSettingsDelegate.ExecuteIfBound();
+}
+
 } // namespace UE::UAF::Editor
 
 #undef LOCTEXT_NAMESPACE
diff --git a/Engine/Plugins/Experimental/UAF/UAF/Source/UAFEditor/Private/Common/SUAFBrowser.h b/Engine/Plugins/Experimental/UAF/UAF/Source/UAFEditor/Private/Common/SUAFBrowser.h
--- a/Engine/Plugins/Experimental/UAF/UAF/Source/UAFEditor/Private/Common/SUAFBrowser.h
+++ b/Engine/Plugins/Experimental/UAF/UAF/Source/UAFEditor/Private/Common/SUAFBrowser.h
@@ -3,6 +3,7 @@
 #pragma once
 
 #include "CoreMinimal.h"
+#include "ContentBrowserDelegates.h"
 #include "Framework/Docking/TabManager.h"
 #include "Widgets/Docking/SDockTab.h"
 #include "Widgets/Layout/SSplitter.h"
@@ -42,6 +43,8 @@ public:
 
 public:
 
+	virtual ~SUAFBrowser() override;
+
 	void Construct(const FArguments& InArgs, const TSharedPtr<UE::Workspace::IWorkspaceEditor>& InHostingApp, ETabState::Type InTabState);
 
 private:
@@ -86,6 +89,12 @@ private:
 	/** Callback for when assets are dragged */
 	void OnAssetsDragged(const TArray<FAssetData>& InAssetData);
 
+	/** Callback to sync settings changes across UAF Browser instances */
+	void OnExternalColumnVisibilityChanged(const TArray<FString>& HiddenColumnIds);
+
+	/** Callback to write column visibility changes to .ini & UAF settings */
+	void OnBrowserColumnVisibilityChanged(const TArray<FString>& HiddenColumnIds);
+
 private:
 
 	/** App Hosting this Browser, to query settings against */
@@ -105,6 +114,12 @@ private:
 	 * update if swapping betweeen assets of the same type
 	 */
 	TWeakPtr<IUAFAssetPreview> AssetPreviewWidgetWeak = nullptr;
+
+	/** Delegate write asset browser config .ini  whenever changes occur. */
+	FSaveAssetViewSettingsDelegate SaveAssetViewSettingsDelegate;
+
+	/** Delegate fired to refresh asset browser whenever config .ini changes occur. */
+	FLoadAssetViewSettingsDelegate LoadAssetViewSettingsDelegate;
 };
 
 } // namespace UE::UAF::Editor
\ No newline at end of file
diff --git a/Engine/Plugins/Experimental/UAF/UAF/Source/UAFEditor/Private/Settings/UAFEditorUserSettings.cpp b/Engine/Plugins/Experimental/UAF/UAF/Source/UAFEditor/Private/Settings/UAFEditorUserSettings.cpp
--- a/Engine/Plugins/Experimental/UAF/UAF/Source/UAFEditor/Private/Settings/UAFEditorUserSettings.cpp
+++ b/Engine/Plugins/Experimental/UAF/UAF/Source/UAFEditor/Private/Settings/UAFEditorUserSettings.cpp
@@ -6,4 +6,25 @@
 
 UUAFEditorUserSettings::UUAFEditorUserSettings()
 {
+	// Initialize hidden columns. Asset browser doesn't support an allow-list  
+	// since the columns that asset(s) may have aren't known ahead of time.
+	UAFBrowserHiddenColumns =
+	{
+		"ItemDiskSize",
+		"bUseNormalizedRootMotionScale",
+		"Compressed Size (KB)",
+		"ImportFileFramerate",
+		"ImportResampleFramerate",
+		"Interpolation",
+		"Number of Frames",
+		"Number of Keys",
+		"NumberOfSampledKeys",
+		"ParentAsset",
+		"PreviewSkeletalMesh",
+		"RetargetSource",
+		"RetargetSourceAsset",
+		"Skeleton",
+		"Target Frame Rate",
+		"Class"
+	};
 }
diff --git a/Engine/Plugins/Experimental/UAF/UAF/Source/UAFEditor/Private/Settings/UAFEditorUserSettings.h b/Engine/Plugins/Experimental/UAF/UAF/Source/UAFEditor/Private/Settings/UAFEditorUserSettings.h
--- a/Engine/Plugins/Experimental/UAF/UAF/Source/UAFEditor/Private/Settings/UAFEditorUserSettings.h
+++ b/Engine/Plugins/Experimental/UAF/UAF/Source/UAFEditor/Private/Settings/UAFEditorUserSettings.h
@@ -28,4 +28,12 @@ public:
 	/** If true, when user starts PIE any UAF assets requiring compilation will automatically be compiled */
 	UPROPERTY(EditAnywhere, Category="UAF", Config)
 	bool bAutoCompileOutOfDateAssets = true;
+
+	/** Default columns to hide in the UAF Browser asset list. Applies to all browser instances (docked, sidebar, drawer). */
+	UPROPERTY(EditAnywhere, Category="UAF|Browser", Config)
+	TArray<FString> UAFBrowserHiddenColumns;
+
+	/** Broadcasts to synchronize column visibility across all browser instances. */
+	DECLARE_MULTICAST_DELEGATE_OneParam(FOnBrowserHiddenColumnsChanged, const TArray<FString>& /*HiddenColumnIds*/);
+	FOnBrowserHiddenColumnsChanged OnBrowserHiddenColumnsChanged;
 };
diff --git a/Engine/Source/Editor/ContentBrowser/Private/SAssetPicker.cpp b/Engine/Source/Editor/ContentBrowser/Private/SAssetPicker.cpp
--- a/Engine/Source/Editor/ContentBrowser/Private/SAssetPicker.cpp
+++ b/Engine/Source/Editor/ContentBrowser/Private/SAssetPicker.cpp
@@ -154,6 +154,22 @@ void SAssetPicker::Construct( const FArguments& InArgs )
 		}
 	}
 
+	for (auto DelegateIt = InArgs._AssetPickerConfig.SaveAssetViewSettingsDelegates.CreateConstIterator(); DelegateIt; ++DelegateIt)
+	{
+		if ((*DelegateIt) != nullptr)
+		{
+			(**DelegateIt) = FSaveAssetViewSettingsDelegate::CreateSP(this, &SAssetPicker::SaveSettings);
+		}
+	}
+
+	for (auto DelegateIt = InArgs._AssetPickerConfig.LoadAssetViewSettingsDelegates.CreateConstIterator(); DelegateIt; ++DelegateIt)
+	{
+		if ((*DelegateIt) != nullptr)
+		{
+			(**DelegateIt) = FLoadAssetViewSettingsDelegate::CreateSP(this, &SAssetPicker::LoadSettings);
+		}
+	}
+
 	TSharedRef<SVerticalBox> VerticalBox = SNew(SVerticalBox);
 
 	ChildSlot
@@ -386,6 +402,7 @@ void SAssetPicker::Construct( const FArguments& InArgs )
 		.HiddenColumnNames(InArgs._AssetPickerConfig.HiddenColumnNames)
 		.ListHiddenColumnNames(InArgs._AssetPickerConfig.ListHiddenColumnNames)
 		.CustomColumns(InArgs._AssetPickerConfig.CustomColumns)
+		.OnColumnVisibilityChanged(InArgs._AssetPickerConfig.OnColumnVisibilityChanged)
 		.OnSearchOptionsChanged(this, &SAssetPicker::HandleSearchSettingsChanged)
 		.InitialThumbnailSize(InArgs._AssetPickerConfig.InitialThumbnailSize)
 		.AssetViewOptionsProfile(InArgs._AssetPickerConfig.AssetViewOptionsProfile)
diff --git a/Engine/Source/Editor/ContentBrowser/Private/SAssetView.cpp b/Engine/Source/Editor/ContentBrowser/Private/SAssetView.cpp
--- a/Engine/Source/Editor/ContentBrowser/Private/SAssetView.cpp
+++ b/Engine/Source/Editor/ContentBrowser/Private/SAssetView.cpp
@@ -1429,6 +1429,7 @@ void SAssetView::Construct( const FArguments& InArgs )
 	HiddenColumnNames = DefaultHiddenColumnNames = InArgs._HiddenColumnNames;
 	ListHiddenColumnNames = DefaultListHiddenColumnNames = InArgs._ListHiddenColumnNames;
 	CustomColumns = InArgs._CustomColumns;
+	OnColumnVisibilityChanged = InArgs._OnColumnVisibilityChanged;
 	OnSearchOptionsChanged = InArgs._OnSearchOptionsChanged;
 	bShowPathViewFilters = InArgs._bShowPathViewFilters;
 	OnExtendAssetViewOptionsMenuContext = InArgs._OnExtendAssetViewOptionsMenuContext;
@@ -6718,6 +6719,8 @@ void SAssetView::OnHiddenColumnsChanged()
 
 		UAssetViewConfig::Get()->SaveEditorConfig();
 	}
+
+	OnColumnVisibilityChanged.ExecuteIfBound(*HiddenColumnsToUse);
 }
 
 bool SAssetView::ShouldColumnGenerateWidget(const FString ColumnName) const
diff --git a/Engine/Source/Editor/ContentBrowser/Public/ContentBrowserDelegates.h b/Engine/Source/Editor/ContentBrowser/Public/ContentBrowserDelegates.h
--- a/Engine/Source/Editor/ContentBrowser/Public/ContentBrowserDelegates.h
+++ b/Engine/Source/Editor/ContentBrowser/Public/ContentBrowserDelegates.h
@@ -63,6 +63,12 @@ DECLARE_DELEGATE_OneParam(FSyncToAssetsDelegate, const TArray< FAssetData >& /*A
 /** Called to force the asset view to refresh */
 DECLARE_DELEGATE_OneParam(FRefreshAssetViewDelegate, bool /*UpdateSources*/);
 
+/** Called to save the asset view settings (column widths, sort order, hidden columns, etc.) */
+DECLARE_DELEGATE(FSaveAssetViewSettingsDelegate);
+
+/** Called to load/restore the asset view settings */
+DECLARE_DELEGATE(FLoadAssetViewSettingsDelegate);
+
 /** Called to set a new filter for an existing asset picker */
 DECLARE_DELEGATE_OneParam(FSetARFilterDelegate, const FARFilter& /*NewFilter*/);
 
@@ -78,6 +84,9 @@ DECLARE_DELEGATE_TwoParams(FOnAssetsActivated, const TArray<FAssetData>& /*Activ
 /** Called when asset(s) are dragged in the asset view */
 DECLARE_DELEGATE_OneParam(FOnAssetsDragged, const TArray<FAssetData>& /*DraggedAssets*/);
 
+/** Called when the user changes which columns are hidden in the asset column/list view */
+DECLARE_DELEGATE_OneParam(FOnColumnVisibilityChanged, const TArray<FString>& /*HiddenColumnIds*/);
+
 /** Called when an asset is double clicked in the asset view */
 DECLARE_DELEGATE_OneParam(FOnAssetDoubleClicked, const FAssetData& /*AssetData*/);
 
diff --git a/Engine/Source/Editor/ContentBrowser/Public/IContentBrowserSingleton.h b/Engine/Source/Editor/ContentBrowser/Public/IContentBrowserSingleton.h
--- a/Engine/Source/Editor/ContentBrowser/Public/IContentBrowserSingleton.h
+++ b/Engine/Source/Editor/ContentBrowser/Public/IContentBrowserSingleton.h
@@ -189,6 +189,12 @@ struct FAssetPickerConfig
 	/** A pointer to an existing delegate that, when executed, will refresh the asset view. */
 	TArray<FRefreshAssetViewDelegate*> RefreshAssetViewDelegates;
 
+	/** A pointer to an existing delegate that, when executed, will save the asset view settings. */
+	TArray<FSaveAssetViewSettingsDelegate*> SaveAssetViewSettingsDelegates;
+
+	/** A pointer to an existing delegate that, when executed, will load/restore the asset view settings. */
+	TArray<FLoadAssetViewSettingsDelegate*> LoadAssetViewSettingsDelegates;
+
 	/** The asset registry filter to use to cull results */
 	FARFilter Filter;
 
@@ -246,6 +252,9 @@ struct FAssetPickerConfig
 	/** The delegate that fires when any number of assets are dragged */
 	FOnAssetsDragged OnAssetsDragged;
 
+	/** Fires when the user changes which columns are hidden in the column or list view. Useful for syncing column state. */
+	FOnColumnVisibilityChanged OnColumnVisibilityChanged;
+
 	/** The delegate that fires when an asset is right clicked and a context menu is requested */
 	FOnGetAssetContextMenu OnGetAssetContextMenu;
 
diff --git a/Engine/Source/Editor/ContentBrowser/Public/SAssetView.h b/Engine/Source/Editor/ContentBrowser/Public/SAssetView.h
--- a/Engine/Source/Editor/ContentBrowser/Public/SAssetView.h
+++ b/Engine/Source/Editor/ContentBrowser/Public/SAssetView.h
@@ -340,6 +340,9 @@ public:
 		/** Custom columns that can be use specific */
 		SLATE_EVENT( FOnSearchOptionChanged, OnSearchOptionsChanged)
 
+		/** Called when the user changes column visibility. */
+		SLATE_EVENT( FOnColumnVisibilityChanged, OnColumnVisibilityChanged )
+
 		/** The content browser that owns this view if any */
 		SLATE_ARGUMENT(TSharedPtr<SContentBrowser>, OwningContentBrowser)
 
@@ -1180,6 +1183,9 @@ private:
 	/** Called when a search option changes to notify that results should be rebuilt */
 	FOnSearchOptionChanged OnSearchOptionsChanged;
 
+	/** Called when the user changes column visibility */
+	FOnColumnVisibilityChanged OnColumnVisibilityChanged;
+
 	/** Called when opening view options menu */
 	FOnExtendAssetViewOptionsMenuContext OnExtendAssetViewOptionsMenuContext;
 
