Add function binding support to FBindableValue

Allows bindable values (FBindableFloat, FBindableBool, etc.) to bind to
parameterless RigVM functions. The function is executed at GetValue time
and its return value is used instead of the constant.

- FAnimNextFunctionReference: GUID-based function reference (survives renames)
- FUAFAssetInstance::CallFunction: cross-asset function resolution via host hierarchy
- Editor: tabbed binding picker (Variables | Functions) with context sensitivity,
return type filtering, and stale binding error display
- Float/double and enum/byte size coercion in ExecuteParameterlessFunction
Function binding

#Synced-CL 52311217

diff --git a/Engine/Plugins/Experimental/UAF/UAF/Source/UAF/Internal/AnimNextRigVMFunctionData.h b/Engine/Plugins/Experimental/UAF/UAF/Source/UAF/Internal/AnimNextRigVMFunctionData.h
--- a/Engine/Plugins/Experimental/UAF/UAF/Source/UAF/Internal/AnimNextRigVMFunctionData.h
+++ b/Engine/Plugins/Experimental/UAF/UAF/Source/UAF/Internal/AnimNextRigVMFunctionData.h
@@ -9,6 +9,10 @@ struct FAnimNextRigVMFunctionData
 {
 	GENERATED_BODY()
 
+	// Stable GUID identifying the function (survives renames)
+	UPROPERTY()
+	FGuid FunctionGuid;
+
 	// The name of the function
 	UPROPERTY()
 	FName Name;
diff --git a/Engine/Plugins/Experimental/UAF/UAF/Source/UAF/Private/AnimNextFunctionReference.cpp b/Engine/Plugins/Experimental/UAF/UAF/Source/UAF/Private/AnimNextFunctionReference.cpp
new file mode 100644
--- /dev/null
+++ b/Engine/Plugins/Experimental/UAF/UAF/Source/UAF/Private/AnimNextFunctionReference.cpp
@@ -0,0 +1,28 @@
+// Copyright Epic Games, Inc. All Rights Reserved.
+
+#include "AnimNextFunctionReference.h"
+#include "RigVMCore/RigVMGraphFunctionDefinition.h"
+
+#include UE_INLINE_GENERATED_CPP_BY_NAME(AnimNextFunctionReference)
+
+FAnimNextFunctionReference FAnimNextFunctionReference::FromHeader(const FRigVMGraphFunctionHeader& InHeader, const UObject* InAsset)
+{
+	FAnimNextFunctionReference Ref;
+	Ref.FunctionGuid = InHeader.Variant.Guid;
+#if WITH_EDITORONLY_DATA
+	// Note: prefix must match UE::UAF::UncookedOnly::FUtils::MakeFunctionWrapperEventName()
+	Ref.EventName = FName(TEXT("__InternalCall_") + InHeader.Name.ToString());
+#endif
+	Ref.Object = InAsset;
+	return Ref;
+}
+
+#if WITH_EDITORONLY_DATA
+FAnimNextFunctionReference FAnimNextFunctionReference::FromEventName(FName InEventName, const UObject* InAsset)
+{
+	FAnimNextFunctionReference Ref;
+	Ref.EventName = InEventName;
+	Ref.Object = InAsset;
+	return Ref;
+}
+#endif
diff --git a/Engine/Plugins/Experimental/UAF/UAF/Source/UAF/Private/AnimNextRigVMAsset.cpp b/Engine/Plugins/Experimental/UAF/UAF/Source/UAF/Private/AnimNextRigVMAsset.cpp
--- a/Engine/Plugins/Experimental/UAF/UAF/Source/UAF/Private/AnimNextRigVMAsset.cpp
+++ b/Engine/Plugins/Experimental/UAF/UAF/Source/UAF/Private/AnimNextRigVMAsset.cpp
@@ -161,6 +161,102 @@ TArray<FRigVMExternalVariable> UUAFRigVMAsset::GetExternalVariablesImpl(bool bFa
 }
 
 
+bool UUAFRigVMAsset::ExecuteParameterlessFunction(UE::UAF::FFunctionHandle InHandle, FRigVMExtendedExecuteContext& InContext, void* OutReturnValue, int32 ReturnValueSize)
+{
+	check(InHandle.IsValid());
+#if WITH_EDITOR
+	check(InHandle.IsValidForVM(VM->GetVMHash()));
+#endif
+
+	const FAnimNextRigVMFunctionData& Function = FunctionData[InHandle.FunctionIndex];
+
+	// Must have at least 1 arg (the return value output). Functions that read from system
+	// variables (bIsInputVariable) will have additional arg indices for those reads;
+	// the output is always the last arg.
+	if (!ensure(Function.ArgIndices.Num() >= 1))
+	{
+		return false;
+	}
+
+	const int32 ReturnVarIndex = Function.ArgIndices.Last();
+
+	// Execute the function. Guard the public context data to support re-entrant calls
+	// (e.g. when GetValue is called from an anim node during an ongoing VM execution).
+	// This matches the pattern used by the InvokeEntry opcode in RigVM.cpp.
+	{
+		FRigVMExecuteContext& ContextPublicData = InContext.GetPublicData<>();
+		TGuardValue<FRigVMExecuteContext> PublicDataGuard(ContextPublicData, ContextPublicData);
+		VM->ExecuteVM(InContext, Function.EventName);
+	}
+
+	// Copy the return value from the external variable to the caller's buffer.
+	// Handles float↔double coercion (RigVM stores floats as doubles internally).
+	const FProperty* ReturnProp = VM->GetExternalVariableDefs()[ReturnVarIndex].GetProperty();
+	check(ReturnProp);
+	void* SrcPtr = InContext.ExternalVariableRuntimeData[ReturnVarIndex].Memory;
+
+	const int32 ReturnPropSize = ReturnProp->GetElementSize();
+	if (ReturnPropSize == ReturnValueSize)
+	{
+		// Exact size match — direct copy
+		ReturnProp->CopyCompleteValue(OutReturnValue, SrcPtr);
+	}
+	else if (const FNumericProperty* NumProp = CastField<FNumericProperty>(ReturnProp))
+	{
+		if (NumProp->IsFloatingPoint())
+		{
+			// Float↔double coercion (RigVM stores floats as doubles internally)
+			const double Value = NumProp->GetFloatingPointPropertyValue(SrcPtr);
+			if (ReturnValueSize == sizeof(float))
+			{
+				*static_cast<float*>(OutReturnValue) = static_cast<float>(Value);
+			}
+			else if (ReturnValueSize == sizeof(double))
+			{
+				*static_cast<double*>(OutReturnValue) = Value;
+			}
+			else
+			{
+				UE_LOG(LogAnimation, Error,
+					TEXT("ExecuteParameterlessFunction: Unsupported floating point value size %zu calling function %s"),
+					ReturnValueSize, *Function.Name.ToString());
+				return false;
+			}
+		}
+		else if (NumProp->IsInteger() && ReturnPropSize <= ReturnValueSize)
+		{
+			// Integer promotion: source fits in destination (e.g. uint8 enum → int32).
+			// Zero-init destination then copy source bytes into the low end.
+			FMemory::Memzero(OutReturnValue, ReturnValueSize);
+			FMemory::Memcpy(OutReturnValue, SrcPtr, ReturnPropSize);
+		}
+		else
+		{
+			UE_LOG(LogAnimation, Error,
+				TEXT("ExecuteParameterlessFunction: Unsupported numerical property type %s calling function %s"),
+				*NumProp->GetCPPType(), *Function.Name.ToString());
+			return false;
+		}
+	}
+	else if (ReturnPropSize < ReturnValueSize
+		&& (CastField<FEnumProperty>(ReturnProp) || CastField<FByteProperty>(ReturnProp)))
+	{
+		// Enum/byte promotion: source fits in destination (e.g. uint8 enum → int32).
+		// Zero-init destination then copy source bytes into the low end.
+		FMemory::Memzero(OutReturnValue, ReturnValueSize);
+		FMemory::Memcpy(OutReturnValue, SrcPtr, ReturnPropSize);
+	}
+	else
+	{
+		UE_LOG(LogAnimation, Error,
+			TEXT("ExecuteParameterlessFunction: return type size mismatch for '%s' — expected %d bytes, got %d. Binding may be stale."),
+			*Function.Name.ToString(), ReturnValueSize, ReturnPropSize);
+		return false;
+	}
+
+	return true;
+}
+
 void UUAFRigVMAsset::CallFunctionHandle(UE::UAF::FFunctionHandle InHandle, FRigVMExtendedExecuteContext& InContext, FInstancedPropertyBag& InArgs)
 {
 #if WITH_EDITOR
@@ -210,6 +306,32 @@ void UUAFRigVMAsset::CallFunctionHandle(UE::UAF::FFunctionHandle InHandle, FRigV
 	}
 }
 
+UE::UAF::FFunctionHandle UUAFRigVMAsset::GetFunctionHandle(const FGuid& InFunctionGuid) const
+{
+	using namespace UE::UAF;
+
+	if (!InFunctionGuid.IsValid())
+	{
+		return FFunctionHandle();
+	}
+
+	for (int32 FunctionIndex = 0; FunctionIndex < FunctionData.Num(); ++FunctionIndex)
+	{
+		const FAnimNextRigVMFunctionData& Function = FunctionData[FunctionIndex];
+		if (Function.FunctionGuid == InFunctionGuid)
+		{
+			FFunctionHandle Handle;
+			Handle.FunctionIndex = FunctionIndex;
+#if WITH_EDITOR
+			Handle.VMHash = VM->GetVMHash();
+#endif
+			return Handle;
+		}
+	}
+
+	return FFunctionHandle();
+}
+
 UE::UAF::FFunctionHandle UUAFRigVMAsset::GetFunctionHandle(FName InEventName) const
 {
 	using namespace UE::UAF;
diff --git a/Engine/Plugins/Experimental/UAF/UAF/Source/UAF/Private/BindableValue/UAFBindableTypes.cpp b/Engine/Plugins/Experimental/UAF/UAF/Source/UAF/Private/BindableValue/UAFBindableTypes.cpp
--- a/Engine/Plugins/Experimental/UAF/UAF/Source/UAF/Private/BindableValue/UAFBindableTypes.cpp
+++ b/Engine/Plugins/Experimental/UAF/UAF/Source/UAF/Private/BindableValue/UAFBindableTypes.cpp
@@ -165,6 +165,24 @@ namespace UE::UAF::Private
 		return Result == EPropertyBagResult::Success;
 	}
 
+	/**
+	 * Executes a parameterless function binding and writes the return value into OutMem.
+	 * Returns true if successful, false if the function couldn't be called.
+	 */
+	static bool ResolveFunctionBindingToMem(
+		const FUAFPropertyBinding& Binding,
+		void*                      OutMem,
+		int32                      OutMemSize,
+		FUAFAssetInstance&         Instance)
+	{
+		if (Binding.SourceFunction.IsNone())
+		{
+			return false;
+		}
+
+		return Instance.CallFunction(Binding.SourceFunction, OutMem, OutMemSize);
+	}
+
 	/**
 	 * Shared resolution logic for all simple FBindableXxx types.
 	 * Handles Variable and SubProperty binding modes, falling back to ConstantValue.
@@ -181,29 +199,58 @@ namespace UE::UAF::Private
 			return ConstantValue;
 		}
 		const FUAFPropertyBinding* Binding = Bindable.GetBinding();
-		if (!Binding || Binding->SourceVariable.IsNone())
+		if (!Binding)
 		{
 			return ConstantValue;
 		}
 
-		if (Binding->SourceType == EUAFBindingSourceType::Variable)
-		{
-			T Out = ConstantValue;
-			Instance->GetVariable(Binding->SourceVariable, Out);
-			return Out;
-		}
-		if (Binding->SourceType == EUAFBindingSourceType::SubProperty)
+		switch (Binding->SourceType)
 		{
-			const FProperty* TargetProp = nullptr;
-			if constexpr (TIsArithmetic<T>::Value)
+		case EUAFBindingSourceType::Variable:
 			{
-				static const FProperty* CachedProp = BindableType::StaticStruct()->FindPropertyByName(GET_MEMBER_NAME_CHECKED(BindableType, ConstantValue));
-				TargetProp = CachedProp;
+				if (Binding->SourceVariable.IsNone())
+				{
+					return ConstantValue;
+				}
+			
+				T Out = ConstantValue;
+				Instance->GetVariable(Binding->SourceVariable, Out);
+				return Out;
 			}
-			T Out = ConstantValue;
-			ResolveSubPropertyToMem(*Binding, TargetProp, &Out, *Instance);
-			return Out;
+		case EUAFBindingSourceType::SubProperty:
+			{
+				const FProperty* TargetProp = nullptr;
+				if constexpr (TIsArithmetic<T>::Value)
+				{
+					static const FProperty* CachedProp = BindableType::StaticStruct()->FindPropertyByName(GET_MEMBER_NAME_CHECKED(BindableType, ConstantValue));
+					TargetProp = CachedProp;
+				}
+				// Zero-init: the leaf property may be smaller than T (e.g. uint8 sub-property → int32).
+				T Out{};
+				if (!ResolveSubPropertyToMem(*Binding, TargetProp, &Out, *Instance))
+				{
+					return ConstantValue;
+				}
+				return Out;
+			}
+		case EUAFBindingSourceType::Function:
+			{
+				if (Binding->SourceFunction.IsNone())
+				{
+					return ConstantValue;
+				}
+				// Zero-init: the function return type may be smaller than T (e.g. uint8 enum → int32).
+				T Out{};
+				if (!ResolveFunctionBindingToMem(*Binding, &Out, sizeof(T), *Instance))
+				{
+					return ConstantValue;
+				}
+				return Out;
+			}
+		default:
+			checkNoEntry();	// Unsupported binding type
 		}
+		
 		return ConstantValue;
 	}
 
@@ -274,40 +321,75 @@ int32 FBindableEnum::GetValue(FUAFAssetInstance* Instance) const
 		return ConstantValue;
 	}
 	const FUAFPropertyBinding& B = *Binding;
-	if (B.SourceVariable.IsNone())
-	{
-		return ConstantValue;
-	}
 
-	if (B.SourceType == EUAFBindingSourceType::Variable)
+	switch (B.SourceType)
 	{
-		const FAnimNextParamType ParamType = GetEnumClass()
-			? FAnimNextParamType(EPropertyBagPropertyType::Enum, EPropertyBagContainerType::None, GetEnumClass())
-			: FAnimNextParamType(EPropertyBagPropertyType::Byte);
-
-		// Determine the source enum's actual underlying byte size.
-		// FByteProperty (uint8 enum) has ElementSize 1; FEnumProperty wraps a numeric property.
-		int32 ReadSize = sizeof(int32);
-		const FProperty* SourceProp = B.SourceVariable.ResolveProperty();
-		if (const FEnumProperty* EnumProp = SourceProp ? CastField<FEnumProperty>(SourceProp) : nullptr)
-		{
-			if (!ensureMsgf(EnumProp->GetElementSize() <= sizeof(int32),
-				TEXT("FBindableEnum: enum '%s' has underlying size %u which exceeds int32; binding ignored"),
-				*EnumProp->GetName(), static_cast<uint32>(EnumProp->GetElementSize())))
+	case EUAFBindingSourceType::Variable:
+		{
+			if (B.SourceVariable.IsNone())
 			{
 				return ConstantValue;
 			}
-			ReadSize = static_cast<int32>(EnumProp->GetElementSize());
+
+			const FAnimNextParamType ParamType = GetEnumClass()
+				? FAnimNextParamType(EPropertyBagPropertyType::Enum, EPropertyBagContainerType::None, GetEnumClass())
+				: FAnimNextParamType(EPropertyBagPropertyType::Byte);
+
+			// Determine the source enum's actual underlying byte size.
+			// FByteProperty (uint8 enum) has ElementSize 1; FEnumProperty wraps a numeric property.
+			int32 ReadSize = sizeof(int32);
+			const FProperty* SourceProp = B.SourceVariable.ResolveProperty();
+			if (const FEnumProperty* EnumProp = SourceProp ? CastField<FEnumProperty>(SourceProp) : nullptr)
+			{
+				if (!ensureMsgf(EnumProp->GetElementSize() <= sizeof(int32),
+					TEXT("FBindableEnum: enum '%s' has underlying size %u which exceeds int32; binding ignored"),
+					*EnumProp->GetName(), static_cast<uint32>(EnumProp->GetElementSize())))
+				{
+					return ConstantValue;
+				}
+				ReadSize = static_cast<int32>(EnumProp->GetElementSize());
+			}
+			else if (const FByteProperty* ByteProp = SourceProp ? CastField<FByteProperty>(SourceProp) : nullptr)
+			{
+				ReadSize = static_cast<int32>(ByteProp->GetElementSize());
+			}
+
+			// Zero-init before writing — the source may be smaller than int32 (e.g. uint8 enum)
+			// and we must not leave upper bytes as garbage from ConstantValue.
+			int32 Out = 0;
+			UE::UAF::Private::ResolveVariableToMem(B.SourceVariable, ParamType, &Out, ReadSize, *Instance);
+			return Out;
 		}
-		else if (const FByteProperty* ByteProp = SourceProp ? CastField<FByteProperty>(SourceProp) : nullptr)
+	case EUAFBindingSourceType::SubProperty:
 		{
-			ReadSize = static_cast<int32>(ByteProp->GetElementSize());
+			if (B.SourceVariable.IsNone())
+			{
+				return ConstantValue;
+			}
+			int32 Out = 0;
+			if (!UE::UAF::Private::ResolveSubPropertyToMem(B, nullptr, &Out, *Instance))
+			{
+				return ConstantValue;
+			}
+			return Out;
 		}
-
-		int32 Out = ConstantValue;
-		UE::UAF::Private::ResolveVariableToMem(B.SourceVariable, ParamType, &Out, ReadSize, *Instance);
-		return Out;
+	case EUAFBindingSourceType::Function:
+		{
+			if (B.SourceFunction.IsNone())
+			{
+				return ConstantValue;
+			}
+			int32 Out = 0;
+			if (!UE::UAF::Private::ResolveFunctionBindingToMem(B, &Out, sizeof(int32), *Instance))
+			{
+				return ConstantValue;
+			}
+			return Out;
+		}
+	default:
+		checkNoEntry();
 	}
+
 	return ConstantValue;
 }
 
@@ -341,7 +423,9 @@ void FBindableStruct::GetValueToMem(
 	}
 
 	const FUAFPropertyBinding& B = *Binding;
-	if (B.SourceVariable.IsNone())
+
+	// Helper: fall back to ConstantValue or zero-init
+	auto CopyConstantOrInit = [this, ExpectedStruct, OutMem]()
 	{
 		if (ConstantValue.IsValid())
 		{
@@ -351,37 +435,60 @@ void FBindableStruct::GetValueToMem(
 		{
 			ExpectedStruct->InitializeStruct(OutMem);
 		}
-		return;
-	}
+	};
 
-	// Variable binding
-	if (B.SourceType == EUAFBindingSourceType::Variable)
+	switch (B.SourceType)
 	{
-		const FProperty*       SourceProp      = B.SourceVariable.ResolveProperty();
-		const FStructProperty* SourceStructProp = SourceProp ? CastField<FStructProperty>(SourceProp) : nullptr;
-		if (!SourceStructProp || !SourceStructProp->Struct)
+	case EUAFBindingSourceType::Variable:
 		{
-			return;
-		}
+			if (B.SourceVariable.IsNone())
+			{
+				CopyConstantOrInit();
+				return;
+			}
 
-		ensureMsgf(SourceStructProp->Struct == ExpectedStruct,
-			TEXT("FBindableStruct::GetValue type mismatch: expected %s, source variable is %s"),
-			*ExpectedStruct->GetName(), *SourceStructProp->Struct->GetName());
+			const FProperty*       SourceProp      = B.SourceVariable.ResolveProperty();
+			const FStructProperty* SourceStructProp = SourceProp ? CastField<FStructProperty>(SourceProp) : nullptr;
+			if (!SourceStructProp || !SourceStructProp->Struct)
+			{
+				return;
+			}
 
-		const FAnimNextParamType ParamType(
-			EPropertyBagPropertyType::Struct, EPropertyBagContainerType::None,
-			SourceStructProp->Struct);
+			ensureMsgf(SourceStructProp->Struct == ExpectedStruct,
+				TEXT("FBindableStruct::GetValue type mismatch: expected %s, source variable is %s"),
+				*ExpectedStruct->GetName(), *SourceStructProp->Struct->GetName());
 
-		Instance->GetVariable(
-			B.SourceVariable, ParamType,
-			TArrayView<uint8>(static_cast<uint8*>(OutMem), ExpectedStruct->GetStructureSize()));
-		return;
-	}
+			const FAnimNextParamType ParamType(
+				EPropertyBagPropertyType::Struct, EPropertyBagContainerType::None,
+				SourceStructProp->Struct);
 
-	// SubProperty binding
-	if (B.SourceType == EUAFBindingSourceType::SubProperty)
-	{
-		UE::UAF::Private::ResolveSubStructPropertyToMem(B, ExpectedStruct, OutMem, *Instance);
+			Instance->GetVariable(
+				B.SourceVariable, ParamType,
+				TArrayView<uint8>(static_cast<uint8*>(OutMem), ExpectedStruct->GetStructureSize()));
+			return;
+		}
+	case EUAFBindingSourceType::SubProperty:
+		{
+			if (B.SourceVariable.IsNone())
+			{
+				CopyConstantOrInit();
+				return;
+			}
+			UE::UAF::Private::ResolveSubStructPropertyToMem(B, ExpectedStruct, OutMem, *Instance);
+			return;
+		}
+	case EUAFBindingSourceType::Function:
+		{
+			if (!B.SourceFunction.IsNone() &&
+				UE::UAF::Private::ResolveFunctionBindingToMem(B, OutMem, ExpectedStruct->GetStructureSize(), *Instance))
+			{
+				return;
+			}
+			CopyConstantOrInit();
+			return;
+		}
+	default:
+		checkNoEntry();
 	}
 }
 
@@ -396,31 +503,56 @@ UObject* FBindableObject::GetValue(FUAFAssetInstance* Instance) const
 		return ConstantValue;
 	}
 	const FUAFPropertyBinding* OptBinding = GetBinding();
-	if (!OptBinding || OptBinding->SourceVariable.IsNone())
+	if (!OptBinding)
 	{
 		return ConstantValue;
 	}
 
-	if (OptBinding->SourceType == EUAFBindingSourceType::Variable)
+	switch (OptBinding->SourceType)
 	{
-		const UClass* ResolvedClass = ObjectClass ? ObjectClass.Get() : UObject::StaticClass();
-		const FAnimNextParamType ParamType(
-			EPropertyBagPropertyType::Object,
-			EPropertyBagContainerType::None,
-			ResolvedClass);
-		TObjectPtr<UObject> Out = ConstantValue;
-		UE::UAF::Private::ResolveVariableToMem(
-			OptBinding->SourceVariable, ParamType,
-			&Out, sizeof(TObjectPtr<UObject>), *Instance);
-		return Out;
-	}
-	if (OptBinding->SourceType == EUAFBindingSourceType::SubProperty)
-	{
-		TObjectPtr<UObject> Out = ConstantValue;
-		UE::UAF::Private::ResolveSubPropertyToMem(
-			*OptBinding, nullptr, &Out, *Instance);
-		return Out;
+	case EUAFBindingSourceType::Variable:
+		{
+			if (OptBinding->SourceVariable.IsNone())
+			{
+				return ConstantValue;
+			}
+			const UClass* ResolvedClass = ObjectClass ? ObjectClass.Get() : UObject::StaticClass();
+			const FAnimNextParamType ParamType(
+				EPropertyBagPropertyType::Object,
+				EPropertyBagContainerType::None,
+				ResolvedClass);
+			TObjectPtr<UObject> Out = ConstantValue;
+			UE::UAF::Private::ResolveVariableToMem(
+				OptBinding->SourceVariable, ParamType,
+				&Out, sizeof(TObjectPtr<UObject>), *Instance);
+			return Out;
+		}
+	case EUAFBindingSourceType::SubProperty:
+		{
+			if (OptBinding->SourceVariable.IsNone())
+			{
+				return ConstantValue;
+			}
+			TObjectPtr<UObject> Out = ConstantValue;
+			UE::UAF::Private::ResolveSubPropertyToMem(
+				*OptBinding, nullptr, &Out, *Instance);
+			return Out;
+		}
+	case EUAFBindingSourceType::Function:
+		{
+			if (OptBinding->SourceFunction.IsNone())
+			{
+				return ConstantValue;
+			}
+			TObjectPtr<UObject> Out = ConstantValue;
+			UE::UAF::Private::ResolveFunctionBindingToMem(
+				*OptBinding, &Out, sizeof(TObjectPtr<UObject>), *Instance);
+			return Out;
+		}
+	default:
+		checkNoEntry();
 	}
+
 	return ConstantValue;
 }
 
diff --git a/Engine/Plugins/Experimental/UAF/UAF/Source/UAF/Private/UAFAssetInstance.cpp b/Engine/Plugins/Experimental/UAF/UAF/Source/UAF/Private/UAFAssetInstance.cpp
--- a/Engine/Plugins/Experimental/UAF/UAF/Source/UAF/Private/UAFAssetInstance.cpp
+++ b/Engine/Plugins/Experimental/UAF/UAF/Source/UAF/Private/UAFAssetInstance.cpp
@@ -2,13 +2,15 @@
 
 #include "UAFAssetInstance.h"
 
+#include "AnimNextFunctionReference.h"
+#include "AnimNextFunctionHandle.h"
+#include "AnimNextRigVMAsset.h"
 #include "ObjectTrace.h"
 #include "Logging/StructuredLog.h"
-#include "Variables/UAFInstanceVariableContainer.h"
+#include "Script/UAFRigVMComponent.h"
 #include "Variables/VariableOverrides.h"
 #include "Variables/VariableOverridesCollection.h"
 #include "UAFAssetInstanceComponent.h"
-#include "Script/UAFScriptComponent.h"
 #include "Module/AnimNextModuleInstance.h"
 
 #include UE_INLINE_GENERATED_CPP_BY_NAME(UAFAssetInstance)
@@ -45,6 +47,54 @@ void FUAFAssetInstance::ReleaseComponents()
 	Components.Empty();
 }
 
+bool FUAFAssetInstance::CallFunction(const FAnimNextFunctionReference& InFunction, void* OutReturnValue, int32 ReturnValueSize)
+{
+	// Find the instance whose asset matches the function reference's owning asset.
+	// The bindable value may live on a child graph instance, but the function is
+	// typically defined on a parent system asset.
+	const UObject* FunctionOwnerAsset = InFunction.GetObject();
+	FUAFAssetInstance* TargetInstance = this;
+	if (FunctionOwnerAsset)
+	{
+		while (TargetInstance && TargetInstance->GetAsset<UUAFRigVMAsset>() != FunctionOwnerAsset)
+		{
+			TargetInstance = TargetInstance->GetHost();
+		}
+	}
+	if (!TargetInstance)
+	{
+		return false;
+	}
+
+	FUAFRigVMComponent* Component = TargetInstance->TryGetComponent<FUAFRigVMComponent>();
+	if (!Component)
+	{
+		return false;
+	}
+
+	UUAFRigVMAsset* TargetAsset = const_cast<UUAFRigVMAsset*>(TargetInstance->GetAsset<UUAFRigVMAsset>());
+	if (!TargetAsset)
+	{
+		return false;
+	}
+
+	// Try GUID lookup first (rename-robust), fall back to EventName in editor (backward compat)
+	UE::UAF::FFunctionHandle Handle = TargetAsset->GetFunctionHandle(InFunction.GetFunctionGuid());
+#if WITH_EDITORONLY_DATA
+	if (!Handle.IsValid())
+	{
+		Handle = TargetAsset->GetFunctionHandle(InFunction.GetEventName());
+	}
+#endif
+	if (!Handle.IsValid())
+	{
+		return false;
+	}
+
+	return TargetAsset->ExecuteParameterlessFunction(
+		Handle, Component->GetExtendedExecuteContext(), OutReturnValue, ReturnValueSize);
+}
+
 
 
 FUAFAssetInstance::FUAFAssetInstance(UScriptStruct* InScriptStruct)
diff --git a/Engine/Plugins/Experimental/UAF/UAF/Source/UAF/Public/AnimNextFunctionReference.h b/Engine/Plugins/Experimental/UAF/UAF/Source/UAF/Public/AnimNextFunctionReference.h
new file mode 100644
--- /dev/null
+++ b/Engine/Plugins/Experimental/UAF/UAF/Source/UAF/Public/AnimNextFunctionReference.h
@@ -0,0 +1,109 @@
+// Copyright Epic Games, Inc. All Rights Reserved.
+
+#pragma once
+
+#include "CoreMinimal.h"
+
+#include "AnimNextFunctionReference.generated.h"
+
+/**
+ * A reference to a specific parameterless RigVM function on a UUAFRigVMAsset.
+ *
+ * Analogous to FAnimNextVariableReference but for functions instead of variables.
+ * Stores the mangled event name used for VM execution and the owning asset for
+ * editor validation.
+ */
+struct FRigVMGraphFunctionHeader;
+
+USTRUCT(BlueprintType)
+struct FAnimNextFunctionReference
+{
+	GENERATED_BODY()
+
+	FAnimNextFunctionReference() = default;
+
+	/** Construct a reference from a function header and owning asset. Stores the GUID for rename robustness. */
+	UAF_API static FAnimNextFunctionReference FromHeader(const FRigVMGraphFunctionHeader& InHeader, const UObject* InAsset);
+
+	/** Construct a reference from an event name and owning asset (backward compat, no GUID). Editor only. */
+#if WITH_EDITORONLY_DATA
+	UAF_API static FAnimNextFunctionReference FromEventName(FName InEventName, const UObject* InAsset);
+#endif
+
+	/** Get the stable function GUID (survives renames). Invalid for old serialized data. */
+	const FGuid& GetFunctionGuid() const { return FunctionGuid; }
+
+#if WITH_EDITORONLY_DATA
+	/** Get the mangled event name (editor only, used for display fallback). */
+	FName GetEventName() const { return EventName; }
+#endif
+
+	/** Get the asset that owns the function. */
+	TObjectPtr<const UObject> GetObject() const { return Object; }
+
+	/** Check whether this can ever refer to a valid function. */
+	bool IsNone() const
+	{
+#if WITH_EDITORONLY_DATA
+		return !FunctionGuid.IsValid() && EventName.IsNone();
+#else
+		return !FunctionGuid.IsValid();
+#endif
+	}
+
+	/** Reset this reference to None. */
+	void Reset()
+	{
+		FunctionGuid.Invalidate();
+#if WITH_EDITORONLY_DATA
+		EventName = NAME_None;
+#endif
+		Object = nullptr;
+	}
+
+	bool operator==(const FAnimNextFunctionReference& Other) const
+	{
+		if (FunctionGuid.IsValid() && Other.FunctionGuid.IsValid())
+		{
+			return FunctionGuid == Other.FunctionGuid;
+		}
+#if WITH_EDITORONLY_DATA
+		return EventName == Other.EventName && Object == Other.Object;
+#else
+		return false;
+#endif
+	}
+
+	bool operator!=(const FAnimNextFunctionReference& Other) const
+	{
+		return !(*this == Other);
+	}
+
+	friend uint32 GetTypeHash(const FAnimNextFunctionReference& Key)
+	{
+		if (Key.FunctionGuid.IsValid())
+		{
+			return GetTypeHash(Key.FunctionGuid);
+		}
+#if WITH_EDITORONLY_DATA
+		return HashCombineFast(GetTypeHash(Key.EventName), GetTypeHash(Key.Object));
+#else
+		return 0;
+#endif
+	}
+
+private:
+	/** Stable GUID identifying the function (survives renames). Invalid for legacy data. */
+	UPROPERTY(EditAnywhere, Category = "Function")
+	FGuid FunctionGuid;
+
+#if WITH_EDITORONLY_DATA
+	/** The mangled event name (editor only, used for display fallback with legacy data). */
+	UPROPERTY(EditAnywhere, Category = "Function")
+	FName EventName;
+#endif
+
+	/** The asset that owns the function. */
+	UPROPERTY(EditAnywhere, Category = "Function")
+	TObjectPtr<const UObject> Object;
+};
diff --git a/Engine/Plugins/Experimental/UAF/UAF/Source/UAF/Public/AnimNextRigVMAsset.h b/Engine/Plugins/Experimental/UAF/UAF/Source/UAF/Public/AnimNextRigVMAsset.h
--- a/Engine/Plugins/Experimental/UAF/UAF/Source/UAF/Public/AnimNextRigVMAsset.h
+++ b/Engine/Plugins/Experimental/UAF/UAF/Source/UAF/Public/AnimNextRigVMAsset.h
@@ -100,7 +100,14 @@ public:
 	// Call a function handle using a property bag as arguments.
 	UE_API void CallFunctionHandle(UE::UAF::FFunctionHandle InHandle, FRigVMExtendedExecuteContext& InContext, FInstancedPropertyBag& InArgs);
 
-	// Get a function handle given its mangled event name
+	// Execute a parameterless function and copy its return value into OutReturnValue.
+	// The function must have exactly 1 arg (the return value output). Returns true on success.
+	UE_API bool ExecuteParameterlessFunction(UE::UAF::FFunctionHandle InHandle, FRigVMExtendedExecuteContext& InContext, void* OutReturnValue, int32 ReturnValueSize);
+
+	// Get a function handle given its stable GUID
+	UE_API UE::UAF::FFunctionHandle GetFunctionHandle(const FGuid& InFunctionGuid) const;
+
+	// Get a function handle given its mangled event name (backward compat)
 	UE_API UE::UAF::FFunctionHandle GetFunctionHandle(FName InEventName) const;
 
 #if WITH_EDITOR
diff --git a/Engine/Plugins/Experimental/UAF/UAF/Source/UAF/Public/BindableValue/UAFPropertyBinding.h b/Engine/Plugins/Experimental/UAF/UAF/Source/UAF/Public/BindableValue/UAFPropertyBinding.h
--- a/Engine/Plugins/Experimental/UAF/UAF/Source/UAF/Public/BindableValue/UAFPropertyBinding.h
+++ b/Engine/Plugins/Experimental/UAF/UAF/Source/UAF/Public/BindableValue/UAFPropertyBinding.h
@@ -3,6 +3,7 @@
 #pragma once
 
 #include "CoreMinimal.h"
+#include "AnimNextFunctionReference.h"
 #include "PropertyBindingPath.h"
 #include "Variables/AnimNextVariableReference.h"
 
@@ -20,7 +21,7 @@ enum class EUAFBindingSourceType : uint8
 	/** A dot-separated path into a struct-type UAF variable. */
 	SubProperty,
 
-	/** A 0-parameter RigVM function whose return value is used. (Reserved; not yet implemented.) */
+	/** A 0-parameter RigVM function whose return value is used. */
 	Function,
 };
 
@@ -47,6 +48,13 @@ struct FUAFPropertyBinding
 	UPROPERTY(EditAnywhere, Category = "Binding")
 	FPropertyBindingPath SubPropertyPath;
 
+	/**
+	 * Reference to the parameterless RigVM function whose return value is used.
+	 * Used only when SourceType is Function.
+	 */
+	UPROPERTY(EditAnywhere, Category = "Binding")
+	FAnimNextFunctionReference SourceFunction;
+
 	/** How to source the bound value. */
 	UPROPERTY(EditAnywhere, Category = "Binding")
 	EUAFBindingSourceType SourceType = EUAFBindingSourceType::Variable;
diff --git a/Engine/Plugins/Experimental/UAF/UAF/Source/UAF/Public/UAFAssetInstance.h b/Engine/Plugins/Experimental/UAF/UAF/Source/UAF/Public/UAFAssetInstance.h
--- a/Engine/Plugins/Experimental/UAF/UAF/Source/UAF/Public/UAFAssetInstance.h
+++ b/Engine/Plugins/Experimental/UAF/UAF/Source/UAF/Public/UAFAssetInstance.h
@@ -16,6 +16,7 @@
 #define UE_API UAF_API
 
 class UUAFSharedVariables;
+struct FAnimNextFunctionReference;
 struct FAnimNextModuleInjectionComponent;
 struct FUAFAssetInstanceComponent;
 struct FUAFScriptComponent;
@@ -161,6 +162,12 @@ struct FUAFAssetInstance
 		return Variables.GetAllVariablesOfType(FAnimNextParamType::GetType<ValueType>(), OutVariables);
 	}
 
+	// Execute a bound parameterless function and copy the return value into OutReturnValue.
+	// Resolves the function across the instance hierarchy by matching the function reference's
+	// owning asset, so functions on a parent System asset can be called from a child Graph instance.
+	// Returns true on success. Falls back to ConstantValue on failure.
+	UE_API bool CallFunction(const FAnimNextFunctionReference& InFunction, void* OutReturnValue, int32 ReturnValueSize);
+
 	// Returns a typed instance component, creating it lazily the first time it is queried.
 	// When querying for structs of a base type, existing derived types will be returned, rather than a new base struct created.
 	template<typename ComponentType>
diff --git a/Engine/Plugins/Experimental/UAF/UAF/Source/UAFEditor/Internal/Common/SRigVMFunctionPicker.h b/Engine/Plugins/Experimental/UAF/UAF/Source/UAFEditor/Internal/Common/SRigVMFunctionPicker.h
--- a/Engine/Plugins/Experimental/UAF/UAF/Source/UAFEditor/Internal/Common/SRigVMFunctionPicker.h
+++ b/Engine/Plugins/Experimental/UAF/UAF/Source/UAFEditor/Internal/Common/SRigVMFunctionPicker.h
@@ -23,6 +23,9 @@ namespace UE::UAF::Editor
 // The header passed when selecting 'None' will be invalid
 using FOnRigVMFunctionPicked = TDelegate<void(const FRigVMGraphFunctionHeader&)>;
 
+// Optional filter predicate: return true to include the function, false to exclude it
+using FOnFilterRigVMFunction = TDelegate<bool(const FRigVMGraphFunctionHeader&)>;
+
 // Picker widget the allows choosing a RigVM function
 class SRigVMFunctionPicker : public SCompoundWidget
 {
@@ -53,6 +56,19 @@ class SRigVMFunctionPicker : public SCompoundWidget
 	// Whether to show the 'None' item
 	SLATE_ARGUMENT(bool, AllowClear)
 
+	// When true, renders only the search box + tree view (no combo button wrapper).
+	// Use this when embedding the picker inside another popup or tab.
+	SLATE_ARGUMENT(bool, ContentOnly)
+
+	// Optional filter for narrowing which functions appear. Return true to include, false to exclude.
+	SLATE_EVENT(FOnFilterRigVMFunction, OnFilterFunction)
+
+	// Assets to filter by when context-sensitive. Empty = no context filtering available.
+	SLATE_ARGUMENT(TArray<TWeakObjectPtr<const UUAFRigVMAsset>>, FilterAssets)
+
+	// Whether context sensitivity is initially enabled (only relevant when FilterAssets is non-empty).
+	SLATE_ARGUMENT(bool, IsContextSensitive)
+
 	SLATE_END_ARGS()
 
 	UAFEDITOR_API void Construct(const FArguments& InArgs);
@@ -93,6 +109,9 @@ private:
 		{}
 
 		FRigVMGraphFunctionHeader FunctionHeader;
+		const FSlateBrush* ReturnTypeIcon = nullptr;
+		FSlateColor ReturnTypeColor = FLinearColor::Gray;
+		FText ReturnTypeName;
 	};
 
 	struct FAssetEntry : FEntry
@@ -132,6 +151,9 @@ private:
 	TArray<TSharedPtr<FEntry>> Entries;
 	TArray<TSharedPtr<FEntry>> FilteredEntries;
 	TSharedPtr<FTextFilterExpressionEvaluator> TextFilter;
+	FOnFilterRigVMFunction OnFilterFunction;
+	TArray<TWeakObjectPtr<const UUAFRigVMAsset>> FilterAssets;
+	bool bIsContextSensitive = true;
 	bool bAllowClear = false;
 	bool bAllowNew = false;
 };
diff --git a/Engine/Plugins/Experimental/UAF/UAF/Source/UAFEditor/Private/BindableValue/UAFBindableValueCustomization.cpp b/Engine/Plugins/Experimental/UAF/UAF/Source/UAFEditor/Private/BindableValue/UAFBindableValueCustomization.cpp
--- a/Engine/Plugins/Experimental/UAF/UAF/Source/UAFEditor/Private/BindableValue/UAFBindableValueCustomization.cpp
+++ b/Engine/Plugins/Experimental/UAF/UAF/Source/UAFEditor/Private/BindableValue/UAFBindableValueCustomization.cpp
@@ -12,7 +12,12 @@
 #include "UObject/UnrealType.h"
 #include "IPropertyUtilities.h"
 
+#include "AnimNextFunctionReference.h"
 #include "AnimNextRigVMAsset.h"
+#include "AnimNextRigVMAssetEditorData.h"
+#include "Common/SRigVMFunctionPicker.h"
+#include "Variables/SVariablePicker.h"
+#include "UncookedOnlyUtils.h"
 #include "Variables/SVariablePickerCombo.h"
 #include "Variables/VariablePickerArgs.h"
 #include "Variables/AnimNextVariableReference.h"
@@ -26,8 +31,14 @@
 #include "InstancedStructDetails.h"
 #include "Widgets/Images/SImage.h"
 #include "Widgets/Input/SButton.h"
+#include "Widgets/Input/SSegmentedControl.h"
 #include "Widgets/Layout/SBox.h"
+#include "Widgets/Layout/SWidgetSwitcher.h"
 #include "PropertyCustomizationHelpers.h"
+#include "Kismet2/BlueprintEditorUtils.h"
+#include "Param/ParamCompatibility.h"
+#include "Param/ParamUtils.h"
+#include "RigVMCore/RigVMGraphFunctionDefinition.h"
 
 #define LOCTEXT_NAMESPACE "UAFBindableValueCustomization"
 
@@ -156,6 +167,120 @@ FAnimNextParamType GetBindableParamType(const UScriptStruct* BindableStruct, TSh
 	return FAnimNextParamType();
 }
 
+/**
+ * Returns true if the function is bindable: 0 caller-provided inputs, exactly 1 output.
+ * (0 caller-provided inputs, at least 1 output).
+ * Input variable arguments (bIsInputVariable) are excluded — they are read
+ * from system variable storage by the wrapper event, not provided by the caller.
+ */
+bool IsFunctionBindable(const FRigVMGraphFunctionHeader& InHeader)
+{
+	int32 InputCount = 0;
+	int32 OutputCount = 0;
+	for (const FRigVMGraphFunctionArgument& Arg : InHeader.Arguments)
+	{
+		if (Arg.IsExecuteContext())
+		{
+			continue;
+		}
+		if (Arg.bIsInputVariable)
+		{
+			continue;
+		}
+		if (Arg.Direction == ERigVMPinDirection::Input || Arg.Direction == ERigVMPinDirection::IO)
+		{
+			InputCount++;
+		}
+		if (Arg.Direction == ERigVMPinDirection::Output || Arg.Direction == ERigVMPinDirection::IO)
+		{
+			OutputCount++;
+		}
+	}
+	return InputCount == 0 && OutputCount == 1;
+}
+
+/**
+ * Returns true if the function is bindable and its return type is compatible
+ * with InTargetType (supporting promotion/demotion with data loss, e.g. float↔double).
+ */
+bool IsFunctionCompatibleWithBindableType(const FRigVMGraphFunctionHeader& InHeader, const FAnimNextParamType& InTargetType)
+{
+	if (!IsFunctionBindable(InHeader))
+	{
+		return false;
+	}
+	if (!InTargetType.IsValid())
+	{
+		return true;
+	}
+	for (const FRigVMGraphFunctionArgument& Arg : InHeader.Arguments)
+	{
+		if (!Arg.IsExecuteContext() && !Arg.bIsInputVariable && Arg.Direction == ERigVMPinDirection::Output)
+		{
+			FRigVMTemplateArgumentType RigVMArgType(Arg.CPPType, Arg.CPPTypeObject.Get());
+			FAnimNextParamType ReturnType = FAnimNextParamType::FromRigVMTemplateArgument(RigVMArgType);
+			return UE::UAF::FParamUtils::GetCompatibility(InTargetType, ReturnType).IsCompatibleWithDataLoss();
+		}
+	}
+	return false;
+}
+
+enum class EFunctionBindingStatus { Valid, FunctionNotFound, ReturnTypeIncompatible };
+
+EFunctionBindingStatus ValidateFunctionBinding(
+	const FAnimNextFunctionReference& InFunction,
+	const FAnimNextParamType& InTargetType,
+	const UUAFRigVMAsset* InAsset,
+	FName* OutResolvedName = nullptr)
+{
+	if (!InAsset)
+	{
+		return EFunctionBindingStatus::FunctionNotFound;
+	}
+
+	const UUAFRigVMAssetEditorData* EditorData = UE::UAF::UncookedOnly::FUtils::GetEditorData(InAsset);
+	if (!EditorData)
+	{
+		return EFunctionBindingStatus::FunctionNotFound;
+	}
+
+	// Only validate against public functions — private functions are not compiled
+	// into FunctionData and cannot be called at runtime.
+	FRigVMGraphFunctionHeaderArray AllHeaders;
+	TMap<FAssetData, FRigVMGraphFunctionHeaderArray> PublicExports;
+	UE::UAF::UncookedOnly::FUtils::GetExportedFunctionsFromAssetRegistry(UE::UAF::AnimNextPublicGraphFunctionsExportsRegistryTag, PublicExports);
+	if (const FRigVMGraphFunctionHeaderArray* AssetPublicHeaders = PublicExports.Find(FAssetData(InAsset)))
+	{
+		AllHeaders.Headers.Append(AssetPublicHeaders->Headers);
+	}
+
+	// Match by GUID first (rename-robust), fall back to name if GUID is invalid (legacy data)
+	const FGuid& FunctionGuid = InFunction.GetFunctionGuid();
+	FString FuncName = InFunction.GetEventName().ToString();
+	FuncName.RemoveFromStart(TEXT("__InternalCall_"));
+
+	for (const FRigVMGraphFunctionHeader& Header : AllHeaders.Headers)
+	{
+		const bool bMatches = FunctionGuid.IsValid()
+			? (Header.Variant.Guid == FunctionGuid)
+			: (Header.Name == FName(*FuncName));
+
+		if (bMatches)
+		{
+			if (OutResolvedName)
+			{
+				*OutResolvedName = Header.Name;
+			}
+			if (IsFunctionCompatibleWithBindableType(Header, InTargetType))
+			{
+				return EFunctionBindingStatus::Valid;
+			}
+			return EFunctionBindingStatus::ReturnTypeIncompatible;
+		}
+	}
+	return EFunctionBindingStatus::FunctionNotFound;
+}
+
 } // anonymous namespace
 
 namespace UE::UAF::Editor
@@ -305,10 +430,18 @@ void FBindableValueCustomization::CustomizeHeader(
 
 			if (const UUAFRigVMAsset* Asset = Cast<UUAFRigVMAsset>(Object))
 			{
+				if (!CurrentRigVMAsset.IsValid())
+				{
+					CurrentRigVMAsset = Asset;
+				}
 				FilterAssets.AddUnique(Asset);
 			}
 			if (const UUAFRigVMAsset* OuterAsset = Object->GetTypedOuter<UUAFRigVMAsset>())
 			{
+				if (!CurrentRigVMAsset.IsValid())
+				{
+					CurrentRigVMAsset = OuterAsset;
+				}
 				FilterAssets.AddUnique(OuterAsset);
 			}
 		}
@@ -545,6 +678,26 @@ void FBindableValueCustomization::CustomizeHeader(
 		}
 	}
 
+	// Validate function binding once at widget creation time (not per-frame).
+	// Also resolve the current display name from GUID (handles function renames).
+	EFunctionBindingStatus CachedFunctionStatus = EFunctionBindingStatus::Valid;
+	FName CachedFunctionDisplayName;
+	PropertyHandle->EnumerateConstRawData(
+		[&CachedFunctionStatus, &CachedFunctionDisplayName, TargetParamType, WeakAsset = CurrentRigVMAsset](const void* RawData, const int32, const int32) -> bool
+		{
+			const FBindableValueBase* Base = static_cast<const FBindableValueBase*>(RawData);
+			if (const FUAFPropertyBinding* B = Base->GetBinding())
+			{
+				if (B->SourceType == EUAFBindingSourceType::Function && !B->SourceFunction.IsNone())
+				{
+					CachedFunctionStatus = ValidateFunctionBinding(
+						B->SourceFunction, TargetParamType, Cast<UUAFRigVMAsset>(WeakAsset.Get()),
+						&CachedFunctionDisplayName);
+				}
+			}
+			return false;
+		});
+
 	InHeaderRow
 	.NameContent()
 	[
@@ -570,54 +723,217 @@ void FBindableValueCustomization::CustomizeHeader(
 			]
 		]
 
-		// Variable picker combo.
+		// Binding picker combo (tabbed: Variables | Functions).
 		+ SHorizontalBox::Slot()
 		.AutoWidth()
 		.VAlign(VAlign_Center)
 		.Padding(4.f, 0.f, 0.f, 0.f)
 		[
-			SNew(SVariablePickerCombo)
-			.PickerArgs(PickerArgs)
-			.OnGetVariableReference_Lambda([this]() -> FAnimNextSoftVariableReference
+			SNew(SComboButton)
+			.ToolTipText_Lambda([Handle = PropertyHandle]() -> FText
 			{
-				FAnimNextSoftVariableReference Result;
-				PropertyHandle->EnumerateConstRawData(
+				FText Result;
+				Handle->EnumerateConstRawData(
 					[&Result](const void* RawData, const int32, const int32) -> bool
 					{
 						const FBindableValueBase* Base = static_cast<const FBindableValueBase*>(RawData);
-						if (const FUAFPropertyBinding* Binding = Base->GetBinding())
+						if (const FUAFPropertyBinding* B = Base->GetBinding())
 						{
-							Result = FAnimNextSoftVariableReference(Binding->SourceVariable);
+							if (B->SourceType == EUAFBindingSourceType::Function && !B->SourceFunction.IsNone())
+							{
+								Result = FText::FromName(B->SourceFunction.GetEventName());
+							}
+							else if (!B->SourceVariable.IsNone())
+							{
+								Result = FText::FromName(B->SourceVariable.GetName());
+							}
 						}
-						return false; // first instance only
+						return false;
 					});
 				return Result;
 			})
-			.OnGetVariableType_Lambda([TargetParamType]() -> FAnimNextParamType
-			{
-				return TargetParamType;
-			})
-			// Show the full path (e.g. "TestVector.X") for sub-property bindings.
-			.VariableName_Lambda([Handle = PropertyHandle]() -> FText
-			{
-				FText Result;
-				Handle->EnumerateConstRawData(
-					[&Result](const void* RawData, const int32, const int32) -> bool
+			.ButtonContent()
+			[
+				SNew(SHorizontalBox)
+				+ SHorizontalBox::Slot()
+				.AutoWidth()
+				.VAlign(VAlign_Center)
+				.HAlign(HAlign_Center)
+				.Padding(0.f, 2.f, 2.f, 2.f)
+				[
+					SNew(SImage)
+					.Image_Lambda([Handle = PropertyHandle, TargetParamType, CachedFunctionStatus]() -> const FSlateBrush*
 					{
-						const FBindableValueBase* Base = static_cast<const FBindableValueBase*>(RawData);
-						if (const FUAFPropertyBinding* B = Base->GetBinding())
+						bool bIsFunctionBinding = false;
+						Handle->EnumerateConstRawData(
+							[&bIsFunctionBinding](const void* RawData, const int32, const int32) -> bool
+							{
+								const FBindableValueBase* Base = static_cast<const FBindableValueBase*>(RawData);
+								if (const FUAFPropertyBinding* B = Base->GetBinding())
+								{
+									bIsFunctionBinding = B->SourceType == EUAFBindingSourceType::Function && !B->SourceFunction.IsNone();
+								}
+								return false;
+							});
+						if (bIsFunctionBinding)
 						{
-							FString DisplayName = B->SourceVariable.GetName().ToString();
-							if (B->SourceType == EUAFBindingSourceType::SubProperty
-								&& B->SubPropertyPath.NumSegments() > 0)
+							if (CachedFunctionStatus != EFunctionBindingStatus::Valid)
 							{
-								DisplayName += TEXT(".") + B->SubPropertyPath.ToString();
+								return FAppStyle::GetBrush(TEXT("Icons.WarningWithColor"));
 							}
-							Result = FText::FromString(DisplayName);
+							return FAppStyle::GetBrush(TEXT("GraphEditor.Function_16x"));
 						}
-						return false; // first instance only
-					});
-				return Result;
+						FEdGraphPinType PinType = UE::UAF::UncookedOnly::FUtils::GetPinTypeFromParamType(TargetParamType);
+						return FBlueprintEditorUtils::GetIconFromPin(PinType, true);
+					})
+					.ColorAndOpacity_Lambda([Handle = PropertyHandle, TargetParamType, CachedFunctionStatus]() -> FSlateColor
+					{
+						bool bIsFunctionBinding = false;
+						Handle->EnumerateConstRawData(
+							[&bIsFunctionBinding](const void* RawData, const int32, const int32) -> bool
+							{
+								const FBindableValueBase* Base = static_cast<const FBindableValueBase*>(RawData);
+								if (const FUAFPropertyBinding* B = Base->GetBinding())
+								{
+									bIsFunctionBinding = B->SourceType == EUAFBindingSourceType::Function && !B->SourceFunction.IsNone();
+								}
+								return false;
+							});
+						if (bIsFunctionBinding && CachedFunctionStatus != EFunctionBindingStatus::Valid)
+						{
+							return FLinearColor::Red;
+						}
+						FEdGraphPinType PinType = UE::UAF::UncookedOnly::FUtils::GetPinTypeFromParamType(TargetParamType);
+						return GetDefault<UEdGraphSchema_K2>()->GetPinTypeColor(PinType);
+					})
+				]
+				+ SHorizontalBox::Slot()
+				.AutoWidth()
+				.VAlign(VAlign_Center)
+				.HAlign(HAlign_Left)
+				.Padding(0.f, 2.f, 0.f, 2.f)
+				[
+					SNew(STextBlock)
+					.TextStyle(&FCoreStyle::Get().GetWidgetStyle<FTextBlockStyle>("SmallText"))
+					.Text_Lambda([Handle = PropertyHandle, CachedFunctionStatus, CachedFunctionDisplayName]() -> FText
+					{
+						FText Result;
+						Handle->EnumerateConstRawData(
+							[&Result, CachedFunctionStatus, CachedFunctionDisplayName](const void* RawData, const int32, const int32) -> bool
+							{
+								const FBindableValueBase* Base = static_cast<const FBindableValueBase*>(RawData);
+								if (const FUAFPropertyBinding* B = Base->GetBinding())
+								{
+									if (B->SourceType == EUAFBindingSourceType::Function && !B->SourceFunction.IsNone())
+									{
+										// Use the resolved name from GUID lookup (handles renames).
+										// Fall back to parsing EventName for legacy data without GUID.
+										FString Name;
+										if (!CachedFunctionDisplayName.IsNone())
+										{
+											Name = CachedFunctionDisplayName.ToString();
+										}
+										else
+										{
+											Name = B->SourceFunction.GetEventName().ToString();
+											Name.RemoveFromStart(TEXT("__InternalCall_"));
+										}
+
+										if (CachedFunctionStatus == EFunctionBindingStatus::FunctionNotFound)
+										{
+											Name += TEXT(" (Missing)");
+										}
+										else if (CachedFunctionStatus == EFunctionBindingStatus::ReturnTypeIncompatible)
+										{
+											Name += TEXT(" (Type Mismatch)");
+										}
+										Result = FText::FromString(Name);
+									}
+									else if (!B->SourceVariable.IsNone())
+									{
+										FString DisplayName = B->SourceVariable.GetName().ToString();
+										if (B->SourceType == EUAFBindingSourceType::SubProperty
+											&& B->SubPropertyPath.NumSegments() > 0)
+										{
+											DisplayName += TEXT(".") + B->SubPropertyPath.ToString();
+										}
+										Result = FText::FromString(DisplayName);
+									}
+								}
+								return false;
+							});
+						return Result;
+					})
+				]
+			]
+			.OnGetMenuContent_Lambda(
+				[this, PickerArgs, SetBindingViaText, TargetParamType]() -> TSharedRef<SWidget>
+			{
+				TSharedRef<SWidgetSwitcher> Switcher = SNew(SWidgetSwitcher);
+
+				// Tab 0: Variables (existing variable picker)
+				Switcher->AddSlot()
+				[
+					SNew(SVariablePicker).Args(PickerArgs)
+				];
+
+				// Tab 1: Functions (parameterless function picker, content-only = no nested combo button)
+				Switcher->AddSlot()
+				[
+					SNew(SRigVMFunctionPicker)
+					.CurrentAsset(FAssetData(const_cast<UUAFRigVMAsset*>(CurrentRigVMAsset.Get())))
+					.AllowNew(false)
+					.AllowClear(false)
+					.ContentOnly(true)
+					.FilterAssets(FilterAssets)
+					.IsContextSensitive(bIsContextSensitive)
+					.OnFilterFunction_Lambda([TargetParamType](const FRigVMGraphFunctionHeader& InHeader) -> bool
+					{
+						return IsFunctionCompatibleWithBindableType(InHeader, TargetParamType);
+					})
+					.OnRigVMFunctionPicked_Lambda(
+						[Handle = PropertyHandle, SetBindingViaText, WeakAsset = CurrentRigVMAsset]
+						(const FRigVMGraphFunctionHeader& InHeader)
+					{
+						if (!InHeader.IsValid())
+						{
+							return;
+						}
+
+						const UUAFRigVMAsset* Asset = WeakAsset.Get();
+						FUAFPropertyBinding NewBinding;
+						NewBinding.SourceType = EUAFBindingSourceType::Function;
+						NewBinding.SourceFunction = FAnimNextFunctionReference::FromHeader(InHeader, Asset);
+						SetBindingViaText(Handle, NewBinding);
+					})
+				];
+
+				return SNew(SBox)
+					.WidthOverride(400.f)
+					.HeightOverride(400.f)
+					[
+						SNew(SVerticalBox)
+						+ SVerticalBox::Slot()
+						.AutoHeight()
+						.Padding(4.f)
+						[
+							SNew(SSegmentedControl<int32>)
+							.Value(0)
+							.OnValueChanged_Lambda([Switcher](int32 Value)
+							{
+								Switcher->SetActiveWidgetIndex(Value);
+							})
+							+ SSegmentedControl<int32>::Slot(0)
+							.Text(LOCTEXT("VariablesTab", "Variables"))
+							+ SSegmentedControl<int32>::Slot(1)
+							.Text(LOCTEXT("FunctionsTab", "Functions"))
+						]
+						+ SVerticalBox::Slot()
+						.FillHeight(1.0f)
+						[
+							Switcher
+						]
+					];
 			})
 		]
 
diff --git a/Engine/Plugins/Experimental/UAF/UAF/Source/UAFEditor/Private/BindableValue/UAFBindableValueCustomization.h b/Engine/Plugins/Experimental/UAF/UAF/Source/UAFEditor/Private/BindableValue/UAFBindableValueCustomization.h
--- a/Engine/Plugins/Experimental/UAF/UAF/Source/UAFEditor/Private/BindableValue/UAFBindableValueCustomization.h
+++ b/Engine/Plugins/Experimental/UAF/UAF/Source/UAFEditor/Private/BindableValue/UAFBindableValueCustomization.h
@@ -40,6 +40,7 @@ private:
 
 	bool bIsContextSensitive = true;
 	TArray<TWeakObjectPtr<const UUAFRigVMAsset>> FilterAssets;
+	TWeakObjectPtr<const UUAFRigVMAsset> CurrentRigVMAsset;
 	FOnIsContextSensitive OnIsContextSensitiveDelegate;
 };
 
diff --git a/Engine/Plugins/Experimental/UAF/UAF/Source/UAFEditor/Private/Common/SRigVMFunctionPicker.cpp b/Engine/Plugins/Experimental/UAF/UAF/Source/UAFEditor/Private/Common/SRigVMFunctionPicker.cpp
--- a/Engine/Plugins/Experimental/UAF/UAF/Source/UAFEditor/Private/Common/SRigVMFunctionPicker.cpp
+++ b/Engine/Plugins/Experimental/UAF/UAF/Source/UAFEditor/Private/Common/SRigVMFunctionPicker.cpp
@@ -3,6 +3,8 @@
 #include "Common/SRigVMFunctionPicker.h"
 
 #include "AnimNextRigVMAsset.h"
+#include "Kismet2/BlueprintEditorUtils.h"
+#include "Param/ParamType.h"
 #include "DetailLayoutBuilder.h"
 #include "RigVMBlueprintGeneratedClass.h"
 #include "UncookedOnlyUtils.h"
@@ -12,6 +14,7 @@
 #include "Widgets/Input/SSearchBox.h"
 #include "Widgets/Views/STableRow.h"
 #include "AnimNextRigVMAssetEditorData.h"
+#include "Widgets/Input/SCheckBox.h"
 
 #define LOCTEXT_NAMESPACE "SRigVMFunctionPicker"
 
@@ -32,28 +35,53 @@ class SRigVMFunctionRowWidget : public STableRow<TSharedPtr<SRigVMFunctionPicker
 		HighlightText = InArgs._HighlightText;
 		check(Entry.IsValid());
 
+		TSharedRef<SHorizontalBox> RowContent = SNew(SHorizontalBox);
+
+		// Return type icon — only for function entries
+		if (Entry->Type == SRigVMFunctionPicker::EEntryType::Function)
+		{
+			TSharedPtr<SRigVMFunctionPicker::FFunctionEntry> FuncEntry = StaticCastSharedPtr<SRigVMFunctionPicker::FFunctionEntry>(Entry);
+			if (FuncEntry->ReturnTypeIcon)
+			{
+				RowContent->AddSlot()
+				.VAlign(VAlign_Center)
+				.AutoWidth()
+				.Padding(0.0f, 0.0f, 2.0f, 0.0f)
+				[
+					SNew(SImage)
+					.Image(FuncEntry->ReturnTypeIcon)
+					.ColorAndOpacity(FuncEntry->ReturnTypeColor)
+					.ToolTipText(FuncEntry->ReturnTypeName)
+				];
+			}
+		}
+
+		// Entry icon (function, asset, none, etc.)
+		RowContent->AddSlot()
+		.VAlign(VAlign_Center)
+		.AutoWidth()
+		[
+			SNew(SImage)
+			.Image(Entry->Icon)
+		];
+
+		// Entry name
+		RowContent->AddSlot()
+		.AutoWidth()
+		.Padding(4.0f, 2.0f)
+		.VAlign(VAlign_Center)
+		[
+			SNew(STextBlock)
+			.Text(Entry->Name)
+			.HighlightText(HighlightText)
+		];
+
 		STableRow<TSharedPtr<SRigVMFunctionPicker::FEntry>>::Construct(
 			STableRow<TSharedPtr<SRigVMFunctionPicker::FEntry>>::FArguments()
 			.ToolTipText(Entry->ToolTip)
 			.Content()
 			[
-				SNew(SHorizontalBox)
-				+SHorizontalBox::Slot()
-				.VAlign(VAlign_Center)
-				.AutoWidth()
-				[
-					SNew(SImage)
-					.Image(Entry->Icon)
-				]
-				+SHorizontalBox::Slot()
-				.AutoWidth()
-				.Padding(4.0f, 2.0f)
-				.VAlign(VAlign_Center)
-				[
-					SNew(STextBlock)
-					.Text(Entry->Name)
-					.HighlightText(HighlightText)
-				]
+				RowContent
 			],
 			InOwnerTable);
 	}
@@ -72,135 +100,185 @@ void SRigVMFunctionPicker::Construct(const FArguments& InArgs)
 	OnNewFunction = InArgs._OnNewFunction;
 	bAllowNew = InArgs._AllowNew;
 	bAllowClear = InArgs._AllowClear;
+	OnFilterFunction = InArgs._OnFilterFunction;
+	FilterAssets = InArgs._FilterAssets;
+	bIsContextSensitive = InArgs._IsContextSensitive;
 
 	TextFilter = MakeShared<FTextFilterExpressionEvaluator>(ETextFilterExpressionEvaluatorMode::Complex);
 
-	ChildSlot
-	[
-		SNew(SComboButton)
-		.ToolTipText(FunctionToolTip)
-		.ButtonContent()
+	// Inner content: search box + tree view (shared between combo and content-only modes)
+	TSharedRef<SWidget> InnerContent =
+		SNew(SVerticalBox)
+		+SVerticalBox::Slot()
+		.AutoHeight()
+		[
+			SAssignNew(SearchBox, SSearchBox)
+			.OnTextChanged_Lambda([this](const FText& InText)
+			{
+				FilterText = InText;
+				RequestRefreshEntries();
+			})
+		]
+		+SVerticalBox::Slot()
+		.FillHeight(1.0f)
 		[
-			SNew(SHorizontalBox)
-			+SHorizontalBox::Slot()
-			.AutoWidth()
-			.Padding(4.0f, 2.0f)
+			SAssignNew(TreeView, STreeView<TSharedPtr<FEntry>>)
+			.SelectionMode(ESelectionMode::Single)
+			.TreeItemsSource(&FilteredEntries)
+			.OnGenerateRow_Lambda([this](TSharedPtr<FEntry> InItem, const TSharedRef<STableViewBase>& InOwnerTable) -> TSharedRef<ITableRow>
+			{
+				return SNew(SRigVMFunctionRowWidget, InOwnerTable, InItem)
+					.HighlightText_Lambda([this]()
+					{
+						return FilterText;
+					});
+			})
+			.OnGetChildren_Lambda([this](TSharedPtr<FEntry> InItem, TArray<TSharedPtr<FEntry>>& OutChildren)
+			{
+				switch(InItem->Type)
+				{
+				case EEntryType::Asset:
+					{
+						TSharedPtr<FAssetEntry> AssetEntry = StaticCastSharedPtr<FAssetEntry>(InItem);
+						if(FilterText.IsEmpty())
+						{
+							OutChildren.Append(AssetEntry->Functions);
+						}
+						else
+						{
+							OutChildren.Append(AssetEntry->FilteredFunctions);
+						}
+						break;
+					}
+				default:
+					break;
+				}
+			})
+			.OnSelectionChanged_Lambda([this](TSharedPtr<FEntry> InItem, ESelectInfo::Type)
+			{
+				FSlateApplication::Get().DismissAllMenus();
+
+				if(!InItem.IsValid())
+				{
+					return;
+				}
+
+				switch(InItem->Type)
+				{
+				case EEntryType::None:
+					{
+						OnRigVMFunctionPicked.ExecuteIfBound(FRigVMGraphFunctionHeader());
+						break;
+					}
+				case EEntryType::Asset:
+					break;
+				case EEntryType::Function:
+					{
+						TSharedPtr<FFunctionEntry> FunctionEntry = StaticCastSharedPtr<FFunctionEntry>(InItem);
+						OnRigVMFunctionPicked.ExecuteIfBound(FunctionEntry->FunctionHeader);
+						break;
+					}
+				case EEntryType::NewFunction:
+					{
+						OnNewFunction.ExecuteIfBound();
+						break;
+					}
+				}
+			})
+		];
+
+	if (InArgs._ContentOnly)
+	{
+		// Content-only mode: render the search + tree directly (for embedding in tabs/popups)
+		RefreshEntries();
+
+		ChildSlot
+		[
+			SNew(SVerticalBox)
+			+ SVerticalBox::Slot()
+			.FillHeight(1.0f)
 			[
-				SNew(SImage)
-				.Image(FAppStyle::GetBrush(TEXT("GraphEditor.Function_16x")))
+				InnerContent
 			]
-			+SHorizontalBox::Slot()
-			.AutoWidth()
-			.VAlign(VAlign_Center)
+			+ SVerticalBox::Slot()
+			.AutoHeight()
+			.Padding(0.0f, 2.0f, 0.0f, 0.0f)
 			[
-				SNew(STextBlock)
-				.Text(FunctionName)
-				.Font(IDetailLayoutBuilder::GetDetailFont())
+				SNew(SCheckBox)
+				.ToolTipText(LOCTEXT("ContextSensitiveTooltip", "Only show functions from the current asset context"))
+				.Visibility_Lambda([this]() -> EVisibility
+				{
+					return FilterAssets.Num() > 0 ? EVisibility::Visible : EVisibility::Collapsed;
+				})
+				.IsChecked_Lambda([this]() -> ECheckBoxState
+				{
+					return bIsContextSensitive ? ECheckBoxState::Checked : ECheckBoxState::Unchecked;
+				})
+				.OnCheckStateChanged_Lambda([this](ECheckBoxState InState)
+				{
+					bIsContextSensitive = (InState == ECheckBoxState::Checked);
+					RefreshEntries();
+				})
+				[
+					SNew(STextBlock)
+					.Text(LOCTEXT("ContextSensitiveLabel", "Context Sensitive"))
+				]
 			]
-		]
-		.OnMenuOpenChanged_Lambda([this](bool bInOpen)
-		{
-			if(bInOpen)
+		];
+	}
+	else
+	{
+		// Combo button mode: wrap content in a dropdown
+		ChildSlot
+		[
+			SNew(SComboButton)
+			.ToolTipText(FunctionToolTip)
+			.ButtonContent()
+			[
+				SNew(SHorizontalBox)
+				+SHorizontalBox::Slot()
+				.AutoWidth()
+				.Padding(4.0f, 2.0f)
+				[
+					SNew(SImage)
+					.Image(FAppStyle::GetBrush(TEXT("GraphEditor.Function_16x")))
+				]
+				+SHorizontalBox::Slot()
+				.AutoWidth()
+				.VAlign(VAlign_Center)
+				[
+					SNew(STextBlock)
+					.Text(FunctionName)
+					.Font(IDetailLayoutBuilder::GetDetailFont())
+				]
+			]
+			.OnMenuOpenChanged_Lambda([this](bool bInOpen)
 			{
-				SearchBox->SetText(FText::GetEmpty());
-				RefreshEntries();
-				RegisterActiveTimer(0.0f, FWidgetActiveTimerDelegate::CreateLambda([this](double InCurrentTime, float InDeltaTime)
+				if(bInOpen)
 				{
-					FSlateApplication::Get().SetKeyboardFocus(SearchBox);
-					return EActiveTimerReturnType::Stop;
-				}));
-			}
-		})
-		.MenuContent()
-		[
-			SNew(SBox)
-			.WidthOverride(300.0f)
-			.HeightOverride(400.0f)
+					SearchBox->SetText(FText::GetEmpty());
+					RefreshEntries();
+					RegisterActiveTimer(0.0f, FWidgetActiveTimerDelegate::CreateLambda([this](double InCurrentTime, float InDeltaTime)
+					{
+						FSlateApplication::Get().SetKeyboardFocus(SearchBox);
+						return EActiveTimerReturnType::Stop;
+					}));
+				}
+			})
+			.MenuContent()
 			[
-				SNew(SBorder)
+				SNew(SBox)
+				.WidthOverride(300.0f)
+				.HeightOverride(400.0f)
 				[
-					SNew(SVerticalBox)
-					+SVerticalBox::Slot()
-					.AutoHeight()
-					[
-						SAssignNew(SearchBox, SSearchBox)
-						.OnTextChanged_Lambda([this](const FText& InText)
-						{
-							FilterText = InText;
-							RequestRefreshEntries();
-						})
-					]
-					+SVerticalBox::Slot()
-					.FillHeight(1.0f)
+					SNew(SBorder)
 					[
-						SAssignNew(TreeView, STreeView<TSharedPtr<FEntry>>)
-						.SelectionMode(ESelectionMode::Single)
-						.TreeItemsSource(&FilteredEntries)
-						.OnGenerateRow_Lambda([this](TSharedPtr<FEntry> InItem, const TSharedRef<STableViewBase>& InOwnerTable) -> TSharedRef<ITableRow>
-						{
-							return SNew(SRigVMFunctionRowWidget, InOwnerTable, InItem)
-								.HighlightText_Lambda([this]()
-								{
-									return FilterText;
-								});
-						})
-						.OnGetChildren_Lambda([this](TSharedPtr<FEntry> InItem, TArray<TSharedPtr<FEntry>>& OutChildren)
-						{
-							switch(InItem->Type)
-							{
-							case EEntryType::Asset:
-								{
-									TSharedPtr<FAssetEntry> AssetEntry = StaticCastSharedPtr<FAssetEntry>(InItem);
-									if(FilterText.IsEmpty())
-									{
-										OutChildren.Append(AssetEntry->Functions);
-									}
-									else
-									{
-										OutChildren.Append(AssetEntry->FilteredFunctions);
-									}
-									break;
-								}
-							default:
-								break;
-							}
-						})
-						.OnSelectionChanged_Lambda([this](TSharedPtr<FEntry> InItem, ESelectInfo::Type)
-						{
-							FSlateApplication::Get().DismissAllMenus();
-
-							if(!InItem.IsValid())
-							{
-								return;
-							}
-
-							switch(InItem->Type)
-							{
-							case EEntryType::None:
-								{
-									OnRigVMFunctionPicked.ExecuteIfBound(FRigVMGraphFunctionHeader());
-									break;
-								}
-							case EEntryType::Asset:
-								break;
-							case EEntryType::Function:
-								{
-									TSharedPtr<FFunctionEntry> FunctionEntry = StaticCastSharedPtr<FFunctionEntry>(InItem);
-									OnRigVMFunctionPicked.ExecuteIfBound(FunctionEntry->FunctionHeader);
-									break;
-								}
-							case EEntryType::NewFunction:
-								{
-									OnNewFunction.ExecuteIfBound();
-									break;
-								}
-							}
-						})
+						InnerContent
 					]
 				]
 			]
-		]
-	];
+		];
+	}
 }
 
 void SRigVMFunctionPicker::RequestRefreshEntries()
@@ -238,25 +316,39 @@ void SRigVMFunctionPicker::RefreshEntries()
 		FilteredEntries.Add(NewFunctionEntry);
 	}
 
+	// Only show public functions — private functions are not compiled into FunctionData
+	// and cannot be called via CallFunction at runtime.
 	TMap<FAssetData, FRigVMGraphFunctionHeaderArray> FunctionExports;
 
-	if(UUAFRigVMAsset* RigVMAsset = WeakCurrentAsset.Get())
-	{
-		UUAFRigVMAssetEditorData* EditorData = UncookedOnly::FUtils::GetEditorData<UUAFRigVMAssetEditorData>(RigVMAsset);
-		FRigVMGraphFunctionHeaderArray HeaderArray;
-		UncookedOnly::FUtils::GetAssetPrivateFunctions(EditorData, HeaderArray);
-		if(HeaderArray.Headers.Num() > 0)
-		{
-			FRigVMGraphFunctionHeaderArray& AssetArray = FunctionExports.FindOrAdd(CurrentAsset);
-			AssetArray.Headers.Append(HeaderArray.Headers);
-		}
-	}
-
 	UncookedOnly::FUtils::GetExportedFunctionsFromAssetRegistry(AnimNextPublicGraphFunctionsExportsRegistryTag, FunctionExports);
 	// TODO: Ideally we can filter functions by schema or execute context, but right now we dont expose the schema and function execute contexts are
 	// all FRigVMExecuteContext, rather than the 'most derived' context in the function.
 //	UncookedOnly::FUtils::GetExportedFunctionsFromAssetRegistry(ControlRigAssetPublicGraphFunctionsExportsRegistryTag, FunctionExports);
 
+	// Context sensitivity: remove assets not in FilterAssets
+	if (bIsContextSensitive && FilterAssets.Num() > 0)
+	{
+		for (auto It = FunctionExports.CreateIterator(); It; ++It)
+		{
+			bool bKeep = false;
+			for (const TWeakObjectPtr<const UUAFRigVMAsset>& WeakAsset : FilterAssets)
+			{
+				if (const UUAFRigVMAsset* Asset = WeakAsset.Get())
+				{
+					if (It->Key == FAssetData(const_cast<UUAFRigVMAsset*>(Asset)))
+					{
+						bKeep = true;
+						break;
+					}
+				}
+			}
+			if (!bKeep)
+			{
+				It.RemoveCurrent();
+			}
+		}
+	}
+
 	for(const TPair<FAssetData, FRigVMGraphFunctionHeaderArray>& AssetExport : FunctionExports)
 	{
 		TSharedPtr<FAssetEntry> AssetEntry = MakeShared<FAssetEntry>();
@@ -267,11 +359,32 @@ void SRigVMFunctionPicker::RefreshEntries()
 
 		for(const FRigVMGraphFunctionHeader& FunctionHeader : AssetExport.Value.Headers)
 		{
+			if (OnFilterFunction.IsBound() && !OnFilterFunction.Execute(FunctionHeader))
+			{
+				continue;
+			}
+
 			TSharedPtr<FFunctionEntry> FunctionEntry = MakeShared<FFunctionEntry>();
 			FunctionEntry->Name = FText::FromName(FunctionHeader.Name);
 			FunctionEntry->ToolTip = FunctionHeader.Description.Len() > 0 ? FText::FromString(FunctionHeader.Description) : FText::FromName(FunctionHeader.Name);
 			FunctionEntry->Icon = FAppStyle::GetBrush(TEXT("GraphEditor.Function_16x"));
 			FunctionEntry->FunctionHeader = FunctionHeader;
+
+			// Extract return type icon from the first output argument
+			for (const FRigVMGraphFunctionArgument& Arg : FunctionHeader.Arguments)
+			{
+				if (!Arg.IsExecuteContext() && Arg.Direction == ERigVMPinDirection::Output)
+				{
+					FRigVMTemplateArgumentType RigVMArgType(Arg.CPPType, Arg.CPPTypeObject.Get());
+					FEdGraphPinType PinType = UncookedOnly::FUtils::GetPinTypeFromParamType(
+						FAnimNextParamType::FromRigVMTemplateArgument(RigVMArgType));
+					FunctionEntry->ReturnTypeIcon = FBlueprintEditorUtils::GetIconFromPin(PinType, true);
+					FunctionEntry->ReturnTypeColor = GetDefault<UEdGraphSchema_K2>()->GetPinTypeColor(PinType);
+					FunctionEntry->ReturnTypeName = UEdGraphSchema_K2::TypeToText(PinType);
+					break;
+				}
+			}
+
 			AssetEntry->Functions.Add(FunctionEntry);
 		}
 	}
diff --git a/Engine/Plugins/Experimental/UAF/UAF/Source/UAFTestData/Public/UAFTestVars.h b/Engine/Plugins/Experimental/UAF/UAF/Source/UAFTestData/Public/UAFTestVars.h
--- a/Engine/Plugins/Experimental/UAF/UAF/Source/UAFTestData/Public/UAFTestVars.h
+++ b/Engine/Plugins/Experimental/UAF/UAF/Source/UAFTestData/Public/UAFTestVars.h
@@ -25,6 +25,25 @@ namespace EUAFTestByteEnum
 	};
 }
 
+/** Struct with adjacent small fields for testing SubProperty overread safety. */
+USTRUCT()
+struct FUAFPackedByteStruct
+{
+	GENERATED_BODY()
+
+	UPROPERTY()
+	uint8 Val1 = 0;
+
+	UPROPERTY()
+	uint8 Val2 = 0;
+
+	UPROPERTY()
+	EUAFTestEnum EnumVal = EUAFTestEnum::ValueA;
+
+	UPROPERTY()
+	float FloatVal = 0.0f;
+};
+
 /** Nested struct for FBindableStruct SubProperty binding tests. */
 USTRUCT()
 struct FUAFNestedTestStruct
@@ -93,6 +112,9 @@ struct FUAFTestVars
 
 	UPROPERTY()
 	TObjectPtr<UObject> ObjectVar = nullptr;
+
+	UPROPERTY()
+	FUAFPackedByteStruct PackedByteVar;
 };
 
 /** 10-property struct used by bulk-resolution performance benchmarks. */
diff --git a/Engine/Plugins/Experimental/UAF/UAF/Source/UAFUncookedOnly/Internal/UncookedOnlyUtils.h b/Engine/Plugins/Experimental/UAF/UAF/Source/UAFUncookedOnly/Internal/UncookedOnlyUtils.h
--- a/Engine/Plugins/Experimental/UAF/UAF/Source/UAFUncookedOnly/Internal/UncookedOnlyUtils.h
+++ b/Engine/Plugins/Experimental/UAF/UAF/Source/UAFUncookedOnly/Internal/UncookedOnlyUtils.h
@@ -76,6 +76,7 @@ struct FUtils
 
 	// Get the corresponding editor data from an asset (casts the editor data appropriately)
 	static UE_API UUAFRigVMAssetEditorData* GetEditorData(UUAFRigVMAsset* InAsset);
+	static UE_API const UUAFRigVMAssetEditorData* GetEditorData(const UUAFRigVMAsset* InAsset);
 
 	template<typename EditorDataType, typename AssetType>
 	static EditorDataType* GetEditorData(AssetType* InAsset)
diff --git a/Engine/Plugins/Experimental/UAF/UAF/Source/UAFUncookedOnly/Private/AnimNextRigVMAssetEditorData.cpp b/Engine/Plugins/Experimental/UAF/UAF/Source/UAFUncookedOnly/Private/AnimNextRigVMAssetEditorData.cpp
--- a/Engine/Plugins/Experimental/UAF/UAF/Source/UAFUncookedOnly/Private/AnimNextRigVMAssetEditorData.cpp
+++ b/Engine/Plugins/Experimental/UAF/UAF/Source/UAFUncookedOnly/Private/AnimNextRigVMAssetEditorData.cpp
@@ -1057,6 +1057,7 @@ void UUAFRigVMAssetEditorData::BuildFunctionWrapperEvents(FAnimNextRigVMAssetCom
 		bAddedWrapperEvent = true;
 
 		FAnimNextRigVMFunctionData FunctionData;
+		FunctionData.FunctionGuid = FunctionHeader.Variant.Guid;
 		FunctionData.Name = FunctionHeader.Name;
 		FunctionData.EventName = FName(*WrapperEventName);
 		FunctionData.ArgIndices = MoveTemp(ArgIndices);
diff --git a/Engine/Plugins/Experimental/UAF/UAF/Source/UAFUncookedOnly/Private/UncookedOnlyUtils.cpp b/Engine/Plugins/Experimental/UAF/UAF/Source/UAFUncookedOnly/Private/UncookedOnlyUtils.cpp
--- a/Engine/Plugins/Experimental/UAF/UAF/Source/UAFUncookedOnly/Private/UncookedOnlyUtils.cpp
+++ b/Engine/Plugins/Experimental/UAF/UAF/Source/UAFUncookedOnly/Private/UncookedOnlyUtils.cpp
@@ -931,6 +931,12 @@ UUAFRigVMAssetEditorData* FUtils::GetEditorData(UUAFRigVMAsset* InAsset)
 	return CastChecked<UUAFRigVMAssetEditorData>(InAsset->EditorData);
 }
 
+const UUAFRigVMAssetEditorData* FUtils::GetEditorData(const UUAFRigVMAsset* InAsset)
+{
+	check(InAsset);
+	return CastChecked<UUAFRigVMAssetEditorData>(InAsset->EditorData);
+}
+
 FAnimNextParamType FUtils::GetParamTypeFromPinType(const FEdGraphPinType& InPinType)
 {
 	FAnimNextParamType::EValueType ValueType = FAnimNextParamType::EValueType::None;
diff --git a/Engine/Plugins/Experimental/UAF/UAF/Tests/UAF/Private/BindableValueTests.cpp b/Engine/Plugins/Experimental/UAF/UAF/Tests/UAF/Private/BindableValueTests.cpp
--- a/Engine/Plugins/Experimental/UAF/UAF/Tests/UAF/Private/BindableValueTests.cpp
+++ b/Engine/Plugins/Experimental/UAF/UAF/Tests/UAF/Private/BindableValueTests.cpp
@@ -3,6 +3,8 @@
 #include "CoreMinimal.h"
 #include "TestHarness.h"
 
+#include "AnimNextFunctionReference.h"
+#include "RigVMCore/RigVMGraphFunctionDefinition.h"
 #include "BindableValue/UAFBindableTypes.h"
 #include "BindableValue/UAFPropertyBinding.h"
 #include "UAFAssetInstance.h"
@@ -11,6 +13,11 @@
 #include "UAFTestVars.h"
 
 #include <catch2/catch_test_macros.hpp>
+#include <catch2/catch_approx.hpp>
+#include <cmath>
+#include <limits>
+
+using Catch::Approx;
 
 // ---------------------------------------------------------------------------
 // Test fixture
@@ -73,6 +80,18 @@ namespace
 		return Binding;
 	}
 
+	/** Creates a FUAFPropertyBinding with Function source type using a dummy GUID. */
+	FUAFPropertyBinding MakeFunctionPropertyBinding()
+	{
+		FUAFPropertyBinding Binding;
+		Binding.SourceType = EUAFBindingSourceType::Function;
+		FRigVMGraphFunctionHeader DummyHeader;
+		DummyHeader.Name = TEXT("TestFunction");
+		DummyHeader.Variant.Guid = FGuid::NewGuid();
+		Binding.SourceFunction = FAnimNextFunctionReference::FromHeader(DummyHeader, nullptr);
+		return Binding;
+	}
+
 } // anonymous namespace
 
 // ---------------------------------------------------------------------------
@@ -416,6 +435,135 @@ TEST_CASE("FBindableEnum.BytePropertyVariableBindingResolvesValue", "[UAF][unit]
 	REQUIRE(Field.GetValue(&Owner) == (int32)EUAFTestByteEnum::ByteC);
 }
 
+TEST_CASE("FBindableEnum.VariableBinding_SmallEnum_DoesNotOverread", "[UAF][unit][MustPass]")
+{
+	// EnumVar is uint8 (1 byte) but FBindableEnum stores int32 (4 bytes).
+	// Verify the resolved value only reads the 1-byte source and zero-fills the rest,
+	// rather than reading 4 bytes from a 1-byte property (which would overread).
+	auto [Owner, Ref] = MakeOwnerAndRef("EnumVar");
+	REQUIRE(Owner.SetVariable(Ref, (uint8)EUAFTestEnum::ValueC) == EPropertyBagResult::Success);
+
+	// Pre-fill ConstantValue with a sentinel that would be visible if upper bytes leak through
+	FBindableEnum Field;
+	Field.SetConstantValue(0x7F7F7F7F);
+	Field.SetEnumClass(StaticEnum<EUAFTestEnum>());
+	Field.SetBinding(MakeVariablePropertyBinding(Ref));
+
+	const int32 Result = Field.GetValue(&Owner);
+	// ValueC == 2.  Upper 3 bytes must be zero, not leftover from ConstantValue.
+	REQUIRE(Result == (int32)EUAFTestEnum::ValueC);
+	REQUIRE((Result & 0xFFFFFF00) == 0);
+}
+
+TEST_CASE("FBindableInt32.SubPropertyBinding_SmallField_DoesNotOverread", "[UAF][unit][MustPass]")
+{
+	// PackedByteVar has {uint8 Val1, uint8 Val2} adjacent in memory.
+	// Binding a FBindableInt32 (4 bytes) to Val1 (1 byte) via SubProperty must NOT
+	// read Val2 into the upper bytes of the result.
+	auto [Owner, Ref] = MakeOwnerAndRef("PackedByteVar");
+
+	FUAFPackedByteStruct Packed;
+	Packed.Val1 = 0x0A;
+	Packed.Val2 = 0xFF; // sentinel — must NOT appear in result
+	REQUIRE(Owner.SetVariable(Ref, Packed) == EPropertyBagResult::Success);
+
+	FBindableInt32 Field;
+	Field.SetConstantValue(0x7F7F7F7F); // sentinel in ConstantValue too
+	Field.SetBinding(MakeSubPropertyBinding(Ref, TEXT("Val1")));
+
+	const int32 Result = Field.GetValue(&Owner);
+	REQUIRE(Result == 0x0A);
+	// Upper bytes must be zero — no overread from Val2 (0xFF) or ConstantValue (0x7F...)
+	REQUIRE((Result & 0xFFFFFF00) == 0);
+}
+
+TEST_CASE("FBindableByte.SubPropertyBinding_SmallField_DoesNotOverread", "[UAF][unit][MustPass]")
+{
+	// Same setup but binding FBindableByte (1 byte target) to Val1 (1 byte source).
+	// This should be an exact size match — verify Val2 doesn't leak.
+	auto [Owner, Ref] = MakeOwnerAndRef("PackedByteVar");
+
+	FUAFPackedByteStruct Packed;
+	Packed.Val1 = 42;
+	Packed.Val2 = 0xFF; // sentinel
+	REQUIRE(Owner.SetVariable(Ref, Packed) == EPropertyBagResult::Success);
+
+	FBindableByte Field;
+	Field.SetBinding(MakeSubPropertyBinding(Ref, TEXT("Val1")));
+
+	REQUIRE(Field.GetValue(&Owner) == 42);
+}
+
+TEST_CASE("FBindableEnum.SubPropertyBinding_SmallEnum_DoesNotOverread", "[UAF][unit][MustPass]")
+{
+	// PackedByteVar.EnumVal is uint8 enum (1 byte). FBindableEnum stores int32 (4 bytes).
+	// SubProperty binding must not read adjacent bytes.
+	auto [Owner, Ref] = MakeOwnerAndRef("PackedByteVar");
+
+	FUAFPackedByteStruct Packed;
+	Packed.EnumVal = EUAFTestEnum::ValueB;
+	Packed.FloatVal = -1.0f; // adjacent sentinel
+	REQUIRE(Owner.SetVariable(Ref, Packed) == EPropertyBagResult::Success);
+
+	FBindableEnum Field;
+	Field.SetConstantValue(0x7F7F7F7F);
+	Field.SetEnumClass(StaticEnum<EUAFTestEnum>());
+	Field.SetBinding(MakeSubPropertyBinding(Ref, TEXT("EnumVal")));
+
+	const int32 Result = Field.GetValue(&Owner);
+	REQUIRE(Result == (int32)EUAFTestEnum::ValueB);
+	REQUIRE((Result & 0xFFFFFF00) == 0);
+}
+
+TEST_CASE("FBindableInt32.VariableBinding_SmallByte_DoesNotOverread", "[UAF][unit][MustPass]")
+{
+	// ByteVal is uint8 (1 byte). FBindableInt32 stores int32 (4 bytes).
+	auto [Owner, Ref] = MakeOwnerAndRef("ByteVal");
+	REQUIRE(Owner.SetVariable(Ref, (uint8)0x42) == EPropertyBagResult::Success);
+
+	FBindableInt32 Field;
+	Field.SetConstantValue(0x7F7F7F7F);
+	Field.SetBinding(MakeVariablePropertyBinding(Ref));
+
+	const int32 Result = Field.GetValue(&Owner);
+	// GetVariable handles type conversion — should return 0x42 with clean upper bytes
+	REQUIRE(Result == 0x42);
+}
+
+TEST_CASE("FBindableDouble.VariableBinding_Float_DoesNotOverread", "[UAF][unit][MustPass]")
+{
+	// FloatVal is float (4 bytes). FBindableDouble stores double (8 bytes).
+	// The resolution should promote float→double cleanly.
+	auto [Owner, Ref] = MakeOwnerAndRef("FloatVal");
+	REQUIRE(Owner.SetVariable(Ref, 3.14f) == EPropertyBagResult::Success);
+
+	FBindableDouble Field;
+	Field.SetConstantValue(std::numeric_limits<double>::quiet_NaN()); // sentinel
+	Field.SetBinding(MakeVariablePropertyBinding(Ref));
+
+	const double Result = Field.GetValue(&Owner);
+	REQUIRE(Result == Approx(3.14).epsilon(0.001));
+	REQUIRE(!std::isnan(Result));
+}
+
+TEST_CASE("FBindableDouble.SubPropertyBinding_Float_DoesNotOverread", "[UAF][unit][MustPass]")
+{
+	// PackedByteVar.FloatVal is float (4 bytes). FBindableDouble stores double (8 bytes).
+	auto [Owner, Ref] = MakeOwnerAndRef("PackedByteVar");
+
+	FUAFPackedByteStruct Packed;
+	Packed.FloatVal = 2.718f;
+	REQUIRE(Owner.SetVariable(Ref, Packed) == EPropertyBagResult::Success);
+
+	FBindableDouble Field;
+	Field.SetConstantValue(std::numeric_limits<double>::quiet_NaN()); // sentinel
+	Field.SetBinding(MakeSubPropertyBinding(Ref, TEXT("FloatVal")));
+
+	const double Result = Field.GetValue(&Owner);
+	REQUIRE(Result == Approx(2.718).epsilon(0.001));
+	REQUIRE(!std::isnan(Result));
+}
+
 // ---------------------------------------------------------------------------
 // FBindableVector
 // ---------------------------------------------------------------------------
@@ -624,3 +772,127 @@ TEST_CASE("FBindableObject.VariableBindingResolvesObject", "[UAF][unit][MustPass
 
 	REQUIRE(Field.GetValue(&Owner) == TestObj);
 }
+
+// ---------------------------------------------------------------------------
+// Function binding fallback tests
+// ---------------------------------------------------------------------------
+// These tests verify that Function-type bindings gracefully fall back to
+// ConstantValue when no RigVM component is available (the normal case in
+// unit tests without a fully compiled asset).
+
+TEST_CASE("FBindableFloat.FunctionBindingFailure_ReturnsConstantValue_NotZero", "[UAF][unit][MustPass]")
+{
+	// When a Function binding has a valid (non-None) SourceFunction but resolution fails
+	// (e.g. no RigVM component), the result must be ConstantValue — not zero.
+	// This catches the bug where zero-init of Out leaked through on failed resolution.
+	FBindableFloat Field;
+	Field.SetConstantValue(999.0f);
+
+	// Construct a binding with a valid GUID (non-None) but no matching asset/component.
+	FUAFPropertyBinding FuncBinding;
+	FuncBinding.SourceType = EUAFBindingSourceType::Function;
+	// Manually set FunctionGuid to a non-zero value so IsNone() returns false.
+	// Use placement to access the private FunctionGuid field via the struct's UPROPERTY serialization.
+	// Simplest approach: use FromHeader with a fake header that has a valid GUID.
+	FRigVMGraphFunctionHeader FakeHeader;
+	FakeHeader.Name = TEXT("FakeFunction");
+	FakeHeader.Variant.Guid = FGuid::NewGuid();
+	FuncBinding.SourceFunction = FAnimNextFunctionReference::FromHeader(FakeHeader, nullptr);
+	Field.SetBinding(MoveTemp(FuncBinding));
+
+	REQUIRE(Field.HasBinding());
+
+	// Instance without a RigVM component — resolution will fail.
+	FUAFAssetInstance Empty;
+	// Must return ConstantValue (999.0), NOT zero.
+	REQUIRE(Field.GetValue(&Empty) == 999.0f);
+}
+
+TEST_CASE("FBindableBool.FunctionBindingWithNoneRefFallsBack", "[UAF][unit][MustPass]")
+{
+	FBindableBool Field;
+	Field.SetConstantValue(true);
+
+	// Function binding with an IsNone reference
+	FUAFPropertyBinding EmptyFunctionBinding;
+	EmptyFunctionBinding.SourceType = EUAFBindingSourceType::Function;
+	// SourceFunction intentionally default-constructed (IsNone() == true)
+	Field.SetBinding(MoveTemp(EmptyFunctionBinding));
+
+	REQUIRE(Field.HasBinding());
+
+	FUAFAssetInstance Empty;
+	REQUIRE(Field.GetValue(&Empty) == true);
+}
+
+
+TEST_CASE("FBindableBool.FunctionBindingFallsBackWithoutRigVM", "[UAF][unit][MustPass]")
+{
+	FBindableBool Field;
+	Field.SetConstantValue(true);
+	Field.SetBinding(MakeFunctionPropertyBinding());
+
+	REQUIRE(Field.HasBinding());
+
+	// Instance without a RigVM component -> falls back to ConstantValue
+	FUAFAssetInstance Empty;
+	REQUIRE(Field.GetValue(&Empty) == true);
+}
+
+TEST_CASE("FBindableFloat.FunctionBindingFallsBackWithoutRigVM", "[UAF][unit][MustPass]")
+{
+	FBindableFloat Field;
+	Field.SetConstantValue(42.0f);
+	Field.SetBinding(MakeFunctionPropertyBinding());
+
+	FUAFAssetInstance Empty;
+	REQUIRE(Field.GetValue(&Empty) == 42.0f);
+}
+
+TEST_CASE("FBindableFloat.FunctionBindingWithNullInstanceFallsBack", "[UAF][unit][MustPass]")
+{
+	FBindableFloat Field;
+	Field.SetConstantValue(7.5f);
+	Field.SetBinding(MakeFunctionPropertyBinding());
+
+	// null instance -> falls back to ConstantValue
+	REQUIRE(Field.GetValue(nullptr) == 7.5f);
+}
+
+TEST_CASE("FBindableEnum.FunctionBindingFallsBackWithoutRigVM", "[UAF][unit][MustPass]")
+{
+	FBindableEnum Field;
+	Field.SetConstantValue((int32)EUAFTestEnum::ValueB);
+	Field.SetEnumClass(StaticEnum<EUAFTestEnum>());
+	Field.SetBinding(MakeFunctionPropertyBinding());
+
+	FUAFAssetInstance Empty;
+	REQUIRE(Field.GetValue(&Empty) == (int32)EUAFTestEnum::ValueB);
+}
+
+TEST_CASE("FBindableStruct.FunctionBindingFallsBackWithoutRigVM", "[UAF][unit][MustPass]")
+{
+	FBindableStruct Field;
+	FInstancedStruct Tmp;
+	Tmp.InitializeAs<FVector>(FVector(1.0, 2.0, 3.0));
+	Field.SetConstantValue(MoveTemp(Tmp));
+	Field.SetBinding(MakeFunctionPropertyBinding());
+
+	// Without a RigVM component, falls back to ConstantValue
+	FVector Result = FVector::ZeroVector;
+	FUAFAssetInstance Empty;
+	Field.GetValue<FVector>(&Empty, Result);
+	REQUIRE(Result == FVector(1.0, 2.0, 3.0));
+}
+
+TEST_CASE("FBindableObject.FunctionBindingFallsBackWithoutRigVM", "[UAF][unit][MustPass]")
+{
+	FBindableObject Field;
+	UObject* TestObj = UObject::StaticClass()->GetDefaultObject();
+	Field.SetConstantValue(TestObj);
+	Field.SetBinding(MakeFunctionPropertyBinding());
+
+	FUAFAssetInstance Empty;
+	REQUIRE(Field.GetValue(&Empty) == TestObj);
+}
+
diff --git a/Engine/Plugins/Experimental/UAF/UAFTestSuites/Source/UAFCQTestSuite/Private/BindableValueFunctionTests.cpp b/Engine/Plugins/Experimental/UAF/UAFTestSuites/Source/UAFCQTestSuite/Private/BindableValueFunctionTests.cpp
new file mode 100644
--- /dev/null
+++ b/Engine/Plugins/Experimental/UAF/UAFTestSuites/Source/UAFCQTestSuite/Private/BindableValueFunctionTests.cpp
@@ -0,0 +1,389 @@
+// Copyright Epic Games, Inc. All Rights Reserved.
+
+#if WITH_DEV_AUTOMATION_TESTS
+#if WITH_EDITOR
+
+#include "UAFTestsUtilities.h"
+#include "UAFTestBlueprintLibrary.h"
+#include "CQTest.h"
+
+#include "AnimNextFunctionReference.h"
+#include "AnimNextRigVMAsset.h"
+#include "Param/ParamCompatibility.h"
+#include "Param/ParamUtils.h"
+#include "RigVMCore/RigVMGraphFunctionDefinition.h"
+#include "BindableValue/UAFBindableTypes.h"
+#include "BindableValue/UAFPropertyBinding.h"
+#include "Graph/AnimNextAnimationGraph.h"
+#include "Module/AnimNextModule.h"
+#include "Module/AnimNextModuleFactory.h"
+#include "Script/UAFRigVMComponent.h"
+#include "UAFAssetInstance.h"
+#include "UncookedOnlyUtils.h"
+
+// ---------------------------------------------------------------------------
+// Test fixture — friend of FUAFAssetInstance for accessing protected members
+// ---------------------------------------------------------------------------
+
+namespace UE::UAF::Tests
+{
+
+struct FUAFBindingTestFixture
+{
+	static void InitInstanceFromAsset(FUAFAssetInstance& InInstance, const UUAFRigVMAsset* InAsset)
+	{
+		InInstance.Asset = InAsset;
+		InInstance.InitializeVariables();
+		InInstance.CopyDefaultComponents();
+		InInstance.BindDefaultComponents();
+		// Ensure the RigVM component exists — it may not be in the asset's default
+		// components, so GetComponent lazily creates it and calls OnBindToInstance
+		// which sets up the VM, execution context, and external variable runtime data.
+		[[maybe_unused]] FUAFRigVMComponent& RigVMComponent = InInstance.GetComponent<FUAFRigVMComponent>();
+	}
+};
+
+} // namespace UE::UAF::Tests
+
+// ---------------------------------------------------------------------------
+// Helpers
+// ---------------------------------------------------------------------------
+
+namespace
+{
+
+using FFixture = UE::UAF::Tests::FUAFBindingTestFixture;
+
+UUAFSystem* CreateAssetWithParameterlessFunction(
+	const FString& InFunctionName,
+	const FString& InReturnCPPType,
+	const FString& InDefaultValue)
+{
+	UUAFSystem* Asset = CastChecked<UUAFSystem>(
+		UAFTestsUtilities::CreateFactoryObject(
+			NewObject<UUAFSystemFactory>(GetTransientPackage()),
+			UUAFSystem::StaticClass()));
+
+	URigVMLibraryNode* FuncNode = UAFTestsUtilities::AddFunctionNode(Asset, InFunctionName);
+	UAFTestsUtilities::AddPin(Asset, FuncNode, ERigVMPinDirection::Output,
+		TEXT("ReturnValue"), InReturnCPPType, NAME_None, InDefaultValue);
+
+	// Mark the function as public so BuildFunctionHeadersContext includes it and
+	// BuildFunctionWrapperEvents creates the proper wrapper entry + external variables.
+	UUAFTestBlueprintLibrary::MarkFunctionAsPublic(Asset, FName(InFunctionName));
+
+	// Recompile — now the function is public, BuildFunctionWrapperEvents will create
+	// wrapper events with external variables, populating FunctionData on the asset.
+	UUAFTestBlueprintLibrary::RecompileVM(Asset);
+
+	return Asset;
+}
+
+/** Creates a binding using EventName only (no GUID) — for backward compat testing. */
+FUAFPropertyBinding MakeFunctionBinding(FName InFunctionName, const UUAFRigVMAsset* InAsset)
+{
+	FUAFPropertyBinding Binding;
+	Binding.SourceType = EUAFBindingSourceType::Function;
+	Binding.SourceFunction = FAnimNextFunctionReference::FromEventName(
+		FName(*UE::UAF::UncookedOnly::FUtils::MakeFunctionWrapperEventName(InFunctionName)),
+		InAsset);
+	return Binding;
+}
+
+/** Creates a binding using GUID + EventName (rename-robust). */
+FUAFPropertyBinding MakeFunctionBindingWithGuid(FName InFunctionName, const UUAFRigVMAsset* InAsset)
+{
+	FUAFPropertyBinding Binding;
+	Binding.SourceType = EUAFBindingSourceType::Function;
+
+	FGuid FunctionGuid = UUAFTestBlueprintLibrary::GetFunctionGuid(InAsset, InFunctionName);
+
+	// Build the reference with GUID via a temporary header
+	FRigVMGraphFunctionHeader TempHeader;
+	TempHeader.Name = InFunctionName;
+	TempHeader.Variant.Guid = FunctionGuid;
+	Binding.SourceFunction = FAnimNextFunctionReference::FromHeader(TempHeader, InAsset);
+
+	return Binding;
+}
+
+/** Builds a minimal FRigVMGraphFunctionHeader with a single output argument of the given type. */
+FRigVMGraphFunctionHeader MakeTestFunctionHeader(FName InCPPType, UObject* InCPPTypeObject = nullptr)
+{
+	FRigVMGraphFunctionHeader Header;
+	Header.Name = TEXT("TestFunc");
+
+	// Add execute context argument (always present, should be ignored by filters)
+	FRigVMGraphFunctionArgument ExecArg;
+	ExecArg.Name = TEXT("ExecuteContext");
+	ExecArg.CPPType = TEXT("FAnimNextExecuteContext");
+	ExecArg.Direction = ERigVMPinDirection::IO;
+	Header.Arguments.Add(ExecArg);
+
+	// Add output argument (return value)
+	FRigVMGraphFunctionArgument OutputArg;
+	OutputArg.Name = TEXT("ReturnValue");
+	OutputArg.CPPType = InCPPType;
+	OutputArg.CPPTypeObject = InCPPTypeObject;
+	OutputArg.Direction = ERigVMPinDirection::Output;
+	Header.Arguments.Add(OutputArg);
+
+	return Header;
+}
+
+/** Tests return type compatibility using the same check as the function picker filter. */
+bool IsReturnTypeCompatible(const FAnimNextParamType& InTargetType, FName InReturnCPPType, UObject* InCPPTypeObject = nullptr)
+{
+	FRigVMTemplateArgumentType RigVMArgType(InReturnCPPType, InCPPTypeObject);
+	FAnimNextParamType ReturnType = FAnimNextParamType::FromRigVMTemplateArgument(RigVMArgType);
+	return UE::UAF::FParamUtils::GetCompatibility(InTargetType, ReturnType).IsCompatibleWithDataLoss();
+}
+
+} // anonymous namespace
+
+// ---------------------------------------------------------------------------
+// Tests
+// ---------------------------------------------------------------------------
+
+TEST_CLASS(FBindableValueFunctionTests, "Animation.UAF.Functional")
+{
+	TEST_METHOD(FBindableFloat_FunctionBinding_ResolvesReturnValue)
+	{
+		TestCommandBuilder
+			.Do(TEXT("Create asset, compile, and verify function binding resolves float"), [&]()
+			{
+				UUAFSystem* Asset = CreateAssetWithParameterlessFunction(
+					TEXT("GetFloat"), TEXT("float"), TEXT("42.0"));
+				ASSERT_THAT(IsNotNull(Asset));
+				ASSERT_THAT(IsNotNull(Asset->GetRigVM()));
+
+				FUAFAssetInstance Instance;
+				FFixture::InitInstanceFromAsset(Instance, Asset);
+
+				FBindableFloat Field;
+				Field.SetConstantValue(0.0f);
+				Field.SetBinding(MakeFunctionBinding(TEXT("GetFloat"), Asset));
+
+				const float Result = Field.GetValue(&Instance);
+				ASSERT_THAT(IsNear(Result, 42.0f, UE_KINDA_SMALL_NUMBER));
+			});
+	}
+
+	TEST_METHOD(FBindableBool_FunctionBinding_ResolvesReturnValue)
+	{
+		TestCommandBuilder
+			.Do(TEXT("Create asset, compile, and verify function binding resolves bool"), [&]()
+			{
+				UUAFSystem* Asset = CreateAssetWithParameterlessFunction(
+					TEXT("GetBool"), TEXT("bool"), TEXT("true"));
+				ASSERT_THAT(IsNotNull(Asset));
+
+				FUAFAssetInstance Instance;
+				FFixture::InitInstanceFromAsset(Instance, Asset);
+
+				FBindableBool Field;
+				Field.SetConstantValue(false);
+				Field.SetBinding(MakeFunctionBinding(TEXT("GetBool"), Asset));
+
+				ASSERT_THAT(IsTrue(Field.GetValue(&Instance)));
+			});
+	}
+
+	TEST_METHOD(FBindableInt32_FunctionBinding_ResolvesReturnValue)
+	{
+		TestCommandBuilder
+			.Do(TEXT("Create asset, compile, and verify function binding resolves int32"), [&]()
+			{
+				UUAFSystem* Asset = CreateAssetWithParameterlessFunction(
+					TEXT("GetInt"), TEXT("int32"), TEXT("99"));
+				ASSERT_THAT(IsNotNull(Asset));
+
+				FUAFAssetInstance Instance;
+				FFixture::InitInstanceFromAsset(Instance, Asset);
+
+				FBindableInt32 Field;
+				Field.SetConstantValue(0);
+				Field.SetBinding(MakeFunctionBinding(TEXT("GetInt"), Asset));
+
+				ASSERT_THAT(AreEqual(99, Field.GetValue(&Instance)));
+			});
+	}
+
+	TEST_METHOD(FBindableFloat_FunctionBinding_OverridesConstantValue)
+	{
+		TestCommandBuilder
+			.Do(TEXT("Verify function binding takes precedence over ConstantValue"), [&]()
+			{
+				UUAFSystem* Asset = CreateAssetWithParameterlessFunction(
+					TEXT("GetFloat"), TEXT("float"), TEXT("42.0"));
+				ASSERT_THAT(IsNotNull(Asset));
+
+				FUAFAssetInstance Instance;
+				FFixture::InitInstanceFromAsset(Instance, Asset);
+
+				FBindableFloat Field;
+				Field.SetConstantValue(999.0f);
+				Field.SetBinding(MakeFunctionBinding(TEXT("GetFloat"), Asset));
+
+				// Function result (42.0) should override ConstantValue (999.0)
+				ASSERT_THAT(IsNear(Field.GetValue(&Instance), 42.0f, UE_KINDA_SMALL_NUMBER));
+			});
+	}
+
+	TEST_METHOD(FBindableFloat_FunctionBinding_InvalidEventNameFallsBack)
+	{
+		TestCommandBuilder
+			.Do(TEXT("Verify invalid function event name falls back to ConstantValue"), [&]()
+			{
+				UUAFSystem* Asset = CreateAssetWithParameterlessFunction(
+					TEXT("GetFloat"), TEXT("float"), TEXT("42.0"));
+				ASSERT_THAT(IsNotNull(Asset));
+
+				FUAFAssetInstance Instance;
+				FFixture::InitInstanceFromAsset(Instance, Asset);
+
+				FBindableFloat Field;
+				Field.SetConstantValue(999.0f);
+
+				// Use a bogus event name that doesn't exist
+				FUAFPropertyBinding Binding;
+				Binding.SourceType = EUAFBindingSourceType::Function;
+				Binding.SourceFunction = FAnimNextFunctionReference::FromEventName(
+					TEXT("__InternalCall_NonExistentFunction"), Asset);
+				Field.SetBinding(MoveTemp(Binding));
+
+				// Should fall back to ConstantValue since function handle is invalid
+				ASSERT_THAT(IsNear(Field.GetValue(&Instance), 999.0f, UE_KINDA_SMALL_NUMBER));
+			});
+	}
+
+	// TODO: Add end-to-end test for functions with input variable arguments once the test
+	// infrastructure supports full instance initialization with internal wrapper variables.
+	// The runtime fix (ArgIndices.Last() + >= 1) and UI filter fix (skip bIsInputVariable)
+	// are validated by manual editor testing with real system functions.
+
+	// -----------------------------------------------------------------------
+	// GUID-based function reference tests
+	// -----------------------------------------------------------------------
+
+	TEST_METHOD(FBindableFloat_GuidBinding_ResolvesAfterEventNameChange)
+	{
+		TestCommandBuilder
+			.Do(TEXT("Verify GUID-based lookup works even when EventName is stale (simulates rename)"), [&]()
+			{
+				UUAFSystem* Asset = CreateAssetWithParameterlessFunction(
+					TEXT("GetFloat"), TEXT("float"), TEXT("42.0"));
+				ASSERT_THAT(IsNotNull(Asset));
+
+				FUAFAssetInstance Instance;
+				FFixture::InitInstanceFromAsset(Instance, Asset);
+
+				// Create a GUID-based binding, then corrupt the EventName to simulate rename
+				FUAFPropertyBinding Binding = MakeFunctionBindingWithGuid(TEXT("GetFloat"), Asset);
+				// Override the EventName with a stale/wrong name
+				FAnimNextFunctionReference StaleRef = FAnimNextFunctionReference::FromEventName(
+					TEXT("__InternalCall_OldName"), Asset);
+				// But keep the GUID from the original — build a reference with correct GUID + wrong EventName
+				FRigVMGraphFunctionHeader HeaderWithStaleEvent;
+				HeaderWithStaleEvent.Name = TEXT("OldName"); // wrong name
+				HeaderWithStaleEvent.Variant.Guid = UUAFTestBlueprintLibrary::GetFunctionGuid(Asset, TEXT("GetFloat"));
+				Binding.SourceFunction = FAnimNextFunctionReference::FromHeader(HeaderWithStaleEvent, Asset);
+
+				FBindableFloat Field;
+				Field.SetConstantValue(0.0f);
+				Field.SetBinding(MoveTemp(Binding));
+
+				// GUID lookup should find the function even though EventName is wrong
+				ASSERT_THAT(IsNear(Field.GetValue(&Instance), 42.0f, UE_KINDA_SMALL_NUMBER));
+			});
+	}
+
+	TEST_METHOD(FBindableFloat_EventNameFallback_WorksWithoutGuid)
+	{
+		TestCommandBuilder
+			.Do(TEXT("Verify EventName fallback works when GUID is invalid (backward compat)"), [&]()
+			{
+				UUAFSystem* Asset = CreateAssetWithParameterlessFunction(
+					TEXT("GetFloat"), TEXT("float"), TEXT("42.0"));
+				ASSERT_THAT(IsNotNull(Asset));
+
+				FUAFAssetInstance Instance;
+				FFixture::InitInstanceFromAsset(Instance, Asset);
+
+				// Create binding with EventName only (no GUID) — simulates old serialized data
+				FBindableFloat Field;
+				Field.SetConstantValue(0.0f);
+				Field.SetBinding(MakeFunctionBinding(TEXT("GetFloat"), Asset));
+
+				// Should resolve via EventName fallback
+				ASSERT_THAT(IsNear(Field.GetValue(&Instance), 42.0f, UE_KINDA_SMALL_NUMBER));
+			});
+	}
+
+	// -----------------------------------------------------------------------
+	// Return type compatibility tests
+	// -----------------------------------------------------------------------
+
+	TEST_METHOD(ReturnTypeCompatibility_FloatToFloat_IsCompatible)
+	{
+		TestCommandBuilder
+			.Do(TEXT("float return → FBindableFloat is compatible"), [&]()
+			{
+				const FAnimNextParamType FloatType = FAnimNextParamType::GetType<float>();
+				ASSERT_THAT(IsTrue(IsReturnTypeCompatible(FloatType, TEXT("float"))));
+			});
+	}
+
+	TEST_METHOD(ReturnTypeCompatibility_DoubleToFloat_IsCompatible)
+	{
+		TestCommandBuilder
+			.Do(TEXT("double return → FBindableFloat is compatible (demotion)"), [&]()
+			{
+				const FAnimNextParamType FloatType = FAnimNextParamType::GetType<float>();
+				ASSERT_THAT(IsTrue(IsReturnTypeCompatible(FloatType, TEXT("double"))));
+			});
+	}
+
+	TEST_METHOD(ReturnTypeCompatibility_FloatToDouble_IsCompatible)
+	{
+		TestCommandBuilder
+			.Do(TEXT("float return → FBindableDouble is compatible (promotion)"), [&]()
+			{
+				const FAnimNextParamType DoubleType = FAnimNextParamType::GetType<double>();
+				ASSERT_THAT(IsTrue(IsReturnTypeCompatible(DoubleType, TEXT("float"))));
+			});
+	}
+
+	TEST_METHOD(ReturnTypeCompatibility_BoolToBool_IsCompatible)
+	{
+		TestCommandBuilder
+			.Do(TEXT("bool return → FBindableBool is compatible"), [&]()
+			{
+				const FAnimNextParamType BoolType = FAnimNextParamType::GetType<bool>();
+				ASSERT_THAT(IsTrue(IsReturnTypeCompatible(BoolType, TEXT("bool"))));
+			});
+	}
+
+	TEST_METHOD(ReturnTypeCompatibility_QuatToFloat_IsIncompatible)
+	{
+		TestCommandBuilder
+			.Do(TEXT("FQuat return → FBindableFloat is incompatible"), [&]()
+			{
+				const FAnimNextParamType FloatType = FAnimNextParamType::GetType<float>();
+				ASSERT_THAT(IsFalse(IsReturnTypeCompatible(FloatType, TEXT("FQuat"), TBaseStructure<FQuat>::Get())));
+			});
+	}
+
+	TEST_METHOD(ReturnTypeCompatibility_VectorToFloat_IsIncompatible)
+	{
+		TestCommandBuilder
+			.Do(TEXT("FVector return → FBindableFloat is incompatible"), [&]()
+			{
+				const FAnimNextParamType FloatType = FAnimNextParamType::GetType<float>();
+				ASSERT_THAT(IsFalse(IsReturnTypeCompatible(FloatType, TEXT("FVector"), TBaseStructure<FVector>::Get())));
+			});
+	}
+};
+
+#endif // WITH_EDITOR
+#endif // WITH_DEV_AUTOMATION_TESTS
diff --git a/Engine/Plugins/Experimental/UAF/UAFTestSuites/Source/UAFCQTestSuite/Private/UAFTestBlueprintLibrary.h b/Engine/Plugins/Experimental/UAF/UAFTestSuites/Source/UAFCQTestSuite/Private/UAFTestBlueprintLibrary.h
--- a/Engine/Plugins/Experimental/UAF/UAFTestSuites/Source/UAFCQTestSuite/Private/UAFTestBlueprintLibrary.h
+++ b/Engine/Plugins/Experimental/UAF/UAFTestSuites/Source/UAFCQTestSuite/Private/UAFTestBlueprintLibrary.h
@@ -5,6 +5,7 @@
 #include "CoreMinimal.h"
 #include "Kismet/BlueprintFunctionLibrary.h"
 #include "AnimNextRigVMAssetEditorData.h"
+#include "AnimNextRigVMFunctionData.h"
 #include "AutomationTestModule.h"
 #include "UAFCompilationScope.h"
 #include "Factories/Factory.h"
@@ -261,6 +262,53 @@ public:
 		return Asset;
 	}
 	
+	// Test helper: mark a function as public so it gets compiled into FunctionData
+	static void MarkFunctionAsPublic(UUAFRigVMAsset* InAsset, FName InFunctionName)
+	{
+		if (InAsset)
+		{
+			UUAFRigVMAssetEditorData* EditorData = UE::UAF::UncookedOnly::FUtils::GetEditorData(InAsset);
+			FRigVMGraphFunctionData* FuncData = EditorData->GraphFunctionStore.FindFunctionByName(InFunctionName);
+			if (FuncData)
+			{
+				EditorData->GraphFunctionStore.MarkFunctionAsPublic(FuncData->Header.LibraryPointer, true);
+			}
+		}
+	}
+
+	// Test helper: get the FunctionGuid from a function's compiled FunctionData
+	static FGuid GetFunctionGuid(const UUAFRigVMAsset* InAsset, FName InFunctionName)
+	{
+		if (InAsset)
+		{
+			FName EventName = FName(TEXT("__InternalCall_") + InFunctionName.ToString());
+			for (const FAnimNextRigVMFunctionData& Data : InAsset->FunctionData)
+			{
+				if (Data.EventName == EventName)
+				{
+					return Data.FunctionGuid;
+				}
+			}
+		}
+		return FGuid();
+	}
+
+	// Test helper: register a function in the asset's FunctionData array so it can be called
+	// via GetFunctionHandle/ExecuteParameterlessFunction. This is needed because
+	// BuildFunctionWrapperEvents only processes programmatic function headers (from traits),
+	// not user-created graph functions.
+	static void AddTestFunctionData(UUAFRigVMAsset* InAsset, FName InFunctionName, FName InEventName, TArray<int32> InArgIndices)
+	{
+		if (InAsset)
+		{
+			FAnimNextRigVMFunctionData NewData;
+			NewData.Name = InFunctionName;
+			NewData.EventName = InEventName;
+			NewData.ArgIndices = MoveTemp(InArgIndices);
+			InAsset->FunctionData.Add(MoveTemp(NewData));
+		}
+	}
+
 	UFUNCTION(BlueprintPure, Category="UAFTest", meta=(ScriptMethod))
 	static EAnimNextRigVMAssetState GetCompilationState(const UUAFRigVMAsset* InAsset)
 	{
