Function binding follow up from post-submit feedback

#Synced-CL 52463674

diff --git a/Engine/Plugins/Experimental/UAF/UAF/Source/UAF/Internal/AnimNextRigVMFunctionData.h b/Engine/Plugins/Experimental/UAF/UAF/Source/UAF/Internal/AnimNextRigVMFunctionData.h
--- a/Engine/Plugins/Experimental/UAF/UAF/Source/UAF/Internal/AnimNextRigVMFunctionData.h
+++ b/Engine/Plugins/Experimental/UAF/UAF/Source/UAF/Internal/AnimNextRigVMFunctionData.h
@@ -9,7 +9,7 @@ struct FAnimNextRigVMFunctionData
 {
 	GENERATED_BODY()
 
-	// Stable GUID identifying the function (survives renames)
+	// Stable GUID identifying the function (survives renames). This corresponds to the RigVM function's GUID
 	UPROPERTY()
 	FGuid FunctionGuid;
 
diff --git a/Engine/Plugins/Experimental/UAF/UAF/Source/UAF/Private/AnimNextRigVMAsset.cpp b/Engine/Plugins/Experimental/UAF/UAF/Source/UAF/Private/AnimNextRigVMAsset.cpp
--- a/Engine/Plugins/Experimental/UAF/UAF/Source/UAF/Private/AnimNextRigVMAsset.cpp
+++ b/Engine/Plugins/Experimental/UAF/UAF/Source/UAF/Private/AnimNextRigVMAsset.cpp
@@ -190,7 +190,7 @@ bool UUAFRigVMAsset::ExecuteParameterlessFunction(UE::UAF::FFunctionHandle InHan
 	}
 
 	// Copy the return value from the external variable to the caller's buffer.
-	// Handles float↔double coercion (RigVM stores floats as doubles internally).
+	// Handles float<->double coercion (RigVM stores floats as doubles internally).
 	const FProperty* ReturnProp = VM->GetExternalVariableDefs()[ReturnVarIndex].GetProperty();
 	check(ReturnProp);
 	void* SrcPtr = InContext.ExternalVariableRuntimeData[ReturnVarIndex].Memory;
@@ -198,14 +198,14 @@ bool UUAFRigVMAsset::ExecuteParameterlessFunction(UE::UAF::FFunctionHandle InHan
 	const int32 ReturnPropSize = ReturnProp->GetElementSize();
 	if (ReturnPropSize == ReturnValueSize)
 	{
-		// Exact size match — direct copy
+		// Exact size match -- direct copy
 		ReturnProp->CopyCompleteValue(OutReturnValue, SrcPtr);
 	}
 	else if (const FNumericProperty* NumProp = CastField<FNumericProperty>(ReturnProp))
 	{
 		if (NumProp->IsFloatingPoint())
 		{
-			// Float↔double coercion (RigVM stores floats as doubles internally)
+			// Float<->double coercion (RigVM stores floats as doubles internally)
 			const double Value = NumProp->GetFloatingPointPropertyValue(SrcPtr);
 			if (ReturnValueSize == sizeof(float))
 			{
@@ -225,7 +225,7 @@ bool UUAFRigVMAsset::ExecuteParameterlessFunction(UE::UAF::FFunctionHandle InHan
 		}
 		else if (NumProp->IsInteger() && ReturnPropSize <= ReturnValueSize)
 		{
-			// Integer promotion: source fits in destination (e.g. uint8 enum → int32).
+			// Integer promotion: source fits in destination (e.g. uint8 enum -> int32).
 			// Zero-init destination then copy source bytes into the low end.
 			FMemory::Memzero(OutReturnValue, ReturnValueSize);
 			FMemory::Memcpy(OutReturnValue, SrcPtr, ReturnPropSize);
@@ -241,7 +241,7 @@ bool UUAFRigVMAsset::ExecuteParameterlessFunction(UE::UAF::FFunctionHandle InHan
 	else if (ReturnPropSize < ReturnValueSize
 		&& (CastField<FEnumProperty>(ReturnProp) || CastField<FByteProperty>(ReturnProp)))
 	{
-		// Enum/byte promotion: source fits in destination (e.g. uint8 enum → int32).
+		// Enum/byte promotion: source fits in destination (e.g. uint8 enum -> int32).
 		// Zero-init destination then copy source bytes into the low end.
 		FMemory::Memzero(OutReturnValue, ReturnValueSize);
 		FMemory::Memcpy(OutReturnValue, SrcPtr, ReturnPropSize);
@@ -249,7 +249,7 @@ bool UUAFRigVMAsset::ExecuteParameterlessFunction(UE::UAF::FFunctionHandle InHan
 	else
 	{
 		UE_LOG(LogAnimation, Error,
-			TEXT("ExecuteParameterlessFunction: return type size mismatch for '%s' — expected %d bytes, got %d. Binding may be stale."),
+			TEXT("ExecuteParameterlessFunction: return type size mismatch for '%s' -- expected %d bytes, got %d. Binding may be stale."),
 			*Function.Name.ToString(), ReturnValueSize, ReturnPropSize);
 		return false;
 	}
diff --git a/Engine/Plugins/Experimental/UAF/UAF/Source/UAF/Private/BindableValue/UAFBindableTypes.cpp b/Engine/Plugins/Experimental/UAF/UAF/Source/UAF/Private/BindableValue/UAFBindableTypes.cpp
--- a/Engine/Plugins/Experimental/UAF/UAF/Source/UAF/Private/BindableValue/UAFBindableTypes.cpp
+++ b/Engine/Plugins/Experimental/UAF/UAF/Source/UAF/Private/BindableValue/UAFBindableTypes.cpp
@@ -90,12 +90,12 @@ namespace UE::UAF::Private
 	}
 
 	static bool ResolveSubStructPropertyToMem(
-		const FUAFPropertyBinding& B,
+		const FUAFPropertyBinding& Binding,
 		const UScriptStruct*       ExpectedStruct,
 		void*                      OutMem,
 		FUAFAssetInstance&         Instance)
 	{
-		const FProperty*       SourceProp      = B.SourceVariable.ResolveProperty();
+		const FProperty*       SourceProp      = Binding.SourceVariable.ResolveProperty();
 		const FStructProperty* SourceStructProp = SourceProp ? CastField<FStructProperty>(SourceProp) : nullptr;
 		if (!SourceStructProp || !SourceStructProp->Struct)
 		{
@@ -112,7 +112,7 @@ namespace UE::UAF::Private
 			SourceStructProp->Struct);
 
 		const EPropertyBagResult GetResult = Instance.GetVariable(
-			B.SourceVariable, SourceParamType,
+			Binding.SourceVariable, SourceParamType,
 			TArrayView<uint8>(StructBuf.GetData(), StructSize));
 
 		if (GetResult != EPropertyBagResult::Success)
@@ -124,7 +124,7 @@ namespace UE::UAF::Private
 		TArray<FPropertyBindingPathIndirection, TInlineAllocator<16>> Indirections;
 		const FPropertyBindingDataView DataView(SourceStructProp->Struct, StructBuf.GetData());
 
-		if (!B.SubPropertyPath.ResolveIndirectionsWithValue(DataView, Indirections) || Indirections.IsEmpty())
+		if (!Binding.SubPropertyPath.ResolveIndirectionsWithValue(DataView, Indirections) || Indirections.IsEmpty())
 		{
 			SourceStructProp->Struct->DestroyStruct(StructBuf.GetData());
 			return false;
@@ -180,7 +180,7 @@ namespace UE::UAF::Private
 			return false;
 		}
 
-		return Instance.CallFunction(Binding.SourceFunction, OutMem, OutMemSize);
+		return Instance.ExecuteParameterlessFunction(Binding.SourceFunction, OutMem, OutMemSize);
 	}
 
 	/**
@@ -320,13 +320,13 @@ int32 FBindableEnum::GetValue(FUAFAssetInstance* Instance) const
 	{
 		return ConstantValue;
 	}
-	const FUAFPropertyBinding& B = *Binding;
+	const FUAFPropertyBinding& ResolvedBinding = *Binding;
 
-	switch (B.SourceType)
+	switch (ResolvedBinding.SourceType)
 	{
 	case EUAFBindingSourceType::Variable:
 		{
-			if (B.SourceVariable.IsNone())
+			if (ResolvedBinding.SourceVariable.IsNone())
 			{
 				return ConstantValue;
 			}
@@ -338,7 +338,7 @@ int32 FBindableEnum::GetValue(FUAFAssetInstance* Instance) const
 			// Determine the source enum's actual underlying byte size.
 			// FByteProperty (uint8 enum) has ElementSize 1; FEnumProperty wraps a numeric property.
 			int32 ReadSize = sizeof(int32);
-			const FProperty* SourceProp = B.SourceVariable.ResolveProperty();
+			const FProperty* SourceProp = ResolvedBinding.SourceVariable.ResolveProperty();
 			if (const FEnumProperty* EnumProp = SourceProp ? CastField<FEnumProperty>(SourceProp) : nullptr)
 			{
 				if (!ensureMsgf(EnumProp->GetElementSize() <= sizeof(int32),
@@ -357,17 +357,17 @@ int32 FBindableEnum::GetValue(FUAFAssetInstance* Instance) const
 			// Zero-init before writing — the source may be smaller than int32 (e.g. uint8 enum)
 			// and we must not leave upper bytes as garbage from ConstantValue.
 			int32 Out = 0;
-			UE::UAF::Private::ResolveVariableToMem(B.SourceVariable, ParamType, &Out, ReadSize, *Instance);
+			UE::UAF::Private::ResolveVariableToMem(ResolvedBinding.SourceVariable, ParamType, &Out, ReadSize, *Instance);
 			return Out;
 		}
 	case EUAFBindingSourceType::SubProperty:
 		{
-			if (B.SourceVariable.IsNone())
+			if (ResolvedBinding.SourceVariable.IsNone())
 			{
 				return ConstantValue;
 			}
 			int32 Out = 0;
-			if (!UE::UAF::Private::ResolveSubPropertyToMem(B, nullptr, &Out, *Instance))
+			if (!UE::UAF::Private::ResolveSubPropertyToMem(ResolvedBinding, nullptr, &Out, *Instance))
 			{
 				return ConstantValue;
 			}
@@ -375,12 +375,12 @@ int32 FBindableEnum::GetValue(FUAFAssetInstance* Instance) const
 		}
 	case EUAFBindingSourceType::Function:
 		{
-			if (B.SourceFunction.IsNone())
+			if (ResolvedBinding.SourceFunction.IsNone())
 			{
 				return ConstantValue;
 			}
 			int32 Out = 0;
-			if (!UE::UAF::Private::ResolveFunctionBindingToMem(B, &Out, sizeof(int32), *Instance))
+			if (!UE::UAF::Private::ResolveFunctionBindingToMem(ResolvedBinding, &Out, sizeof(int32), *Instance))
 			{
 				return ConstantValue;
 			}
@@ -422,7 +422,7 @@ void FBindableStruct::GetValueToMem(
 		return;
 	}
 
-	const FUAFPropertyBinding& B = *Binding;
+	const FUAFPropertyBinding& ResolvedBinding = *Binding;
 
 	// Helper: fall back to ConstantValue or zero-init
 	auto CopyConstantOrInit = [this, ExpectedStruct, OutMem]()
@@ -437,17 +437,17 @@ void FBindableStruct::GetValueToMem(
 		}
 	};
 
-	switch (B.SourceType)
+	switch (ResolvedBinding.SourceType)
 	{
 	case EUAFBindingSourceType::Variable:
 		{
-			if (B.SourceVariable.IsNone())
+			if (ResolvedBinding.SourceVariable.IsNone())
 			{
 				CopyConstantOrInit();
 				return;
 			}
 
-			const FProperty*       SourceProp      = B.SourceVariable.ResolveProperty();
+			const FProperty*       SourceProp      = ResolvedBinding.SourceVariable.ResolveProperty();
 			const FStructProperty* SourceStructProp = SourceProp ? CastField<FStructProperty>(SourceProp) : nullptr;
 			if (!SourceStructProp || !SourceStructProp->Struct)
 			{
@@ -463,24 +463,24 @@ void FBindableStruct::GetValueToMem(
 				SourceStructProp->Struct);
 
 			Instance->GetVariable(
-				B.SourceVariable, ParamType,
+				ResolvedBinding.SourceVariable, ParamType,
 				TArrayView<uint8>(static_cast<uint8*>(OutMem), ExpectedStruct->GetStructureSize()));
 			return;
 		}
 	case EUAFBindingSourceType::SubProperty:
 		{
-			if (B.SourceVariable.IsNone())
+			if (ResolvedBinding.SourceVariable.IsNone())
 			{
 				CopyConstantOrInit();
 				return;
 			}
-			UE::UAF::Private::ResolveSubStructPropertyToMem(B, ExpectedStruct, OutMem, *Instance);
+			UE::UAF::Private::ResolveSubStructPropertyToMem(ResolvedBinding, ExpectedStruct, OutMem, *Instance);
 			return;
 		}
 	case EUAFBindingSourceType::Function:
 		{
-			if (!B.SourceFunction.IsNone() &&
-				UE::UAF::Private::ResolveFunctionBindingToMem(B, OutMem, ExpectedStruct->GetStructureSize(), *Instance))
+			if (!ResolvedBinding.SourceFunction.IsNone() &&
+				UE::UAF::Private::ResolveFunctionBindingToMem(ResolvedBinding, OutMem, ExpectedStruct->GetStructureSize(), *Instance))
 			{
 				return;
 			}
diff --git a/Engine/Plugins/Experimental/UAF/UAF/Source/UAF/Private/UAFAssetInstance.cpp b/Engine/Plugins/Experimental/UAF/UAF/Source/UAF/Private/UAFAssetInstance.cpp
--- a/Engine/Plugins/Experimental/UAF/UAF/Source/UAF/Private/UAFAssetInstance.cpp
+++ b/Engine/Plugins/Experimental/UAF/UAF/Source/UAF/Private/UAFAssetInstance.cpp
@@ -47,7 +47,7 @@ void FUAFAssetInstance::ReleaseComponents()
 	Components.Empty();
 }
 
-bool FUAFAssetInstance::CallFunction(const FAnimNextFunctionReference& InFunction, void* OutReturnValue, int32 ReturnValueSize)
+bool FUAFAssetInstance::ExecuteParameterlessFunction(const FAnimNextFunctionReference& InFunction, void* OutReturnValue, int32 ReturnValueSize)
 {
 	// Find the instance whose asset matches the function reference's owning asset.
 	// The bindable value may live on a child graph instance, but the function is
diff --git a/Engine/Plugins/Experimental/UAF/UAF/Source/UAF/Public/UAFAssetInstance.h b/Engine/Plugins/Experimental/UAF/UAF/Source/UAF/Public/UAFAssetInstance.h
--- a/Engine/Plugins/Experimental/UAF/UAF/Source/UAF/Public/UAFAssetInstance.h
+++ b/Engine/Plugins/Experimental/UAF/UAF/Source/UAF/Public/UAFAssetInstance.h
@@ -165,8 +165,8 @@ struct FUAFAssetInstance
 	// Execute a bound parameterless function and copy the return value into OutReturnValue.
 	// Resolves the function across the instance hierarchy by matching the function reference's
 	// owning asset, so functions on a parent System asset can be called from a child Graph instance.
-	// Returns true on success. Falls back to ConstantValue on failure.
-	UE_API bool CallFunction(const FAnimNextFunctionReference& InFunction, void* OutReturnValue, int32 ReturnValueSize);
+	// Returns true on success
+	UE_API bool ExecuteParameterlessFunction(const FAnimNextFunctionReference& InFunction, void* OutReturnValue, int32 ReturnValueSize);
 
 	// Returns a typed instance component, creating it lazily the first time it is queried.
 	// When querying for structs of a base type, existing derived types will be returned, rather than a new base struct created.
diff --git a/Engine/Plugins/Experimental/UAF/UAF/Source/UAFEditor/Private/BindableValue/UAFBindableValueCustomization.cpp b/Engine/Plugins/Experimental/UAF/UAF/Source/UAFEditor/Private/BindableValue/UAFBindableValueCustomization.cpp
--- a/Engine/Plugins/Experimental/UAF/UAF/Source/UAFEditor/Private/BindableValue/UAFBindableValueCustomization.cpp
+++ b/Engine/Plugins/Experimental/UAF/UAF/Source/UAFEditor/Private/BindableValue/UAFBindableValueCustomization.cpp
@@ -187,11 +187,11 @@ bool IsFunctionBindable(const FRigVMGraphFunctionHeader& InHeader)
 		{
 			continue;
 		}
-		if (Arg.Direction == ERigVMPinDirection::Input || Arg.Direction == ERigVMPinDirection::IO)
+		if (Arg.Direction == ERigVMPinDirection::Input)
 		{
 			InputCount++;
 		}
-		if (Arg.Direction == ERigVMPinDirection::Output || Arg.Direction == ERigVMPinDirection::IO)
+		if (Arg.Direction == ERigVMPinDirection::Output)
 		{
 			OutputCount++;
 		}
@@ -686,12 +686,12 @@ void FBindableValueCustomization::CustomizeHeader(
 		[&CachedFunctionStatus, &CachedFunctionDisplayName, TargetParamType, WeakAsset = CurrentRigVMAsset](const void* RawData, const int32, const int32) -> bool
 		{
 			const FBindableValueBase* Base = static_cast<const FBindableValueBase*>(RawData);
-			if (const FUAFPropertyBinding* B = Base->GetBinding())
+			if (const FUAFPropertyBinding* Binding = Base->GetBinding())
 			{
-				if (B->SourceType == EUAFBindingSourceType::Function && !B->SourceFunction.IsNone())
+				if (Binding->SourceType == EUAFBindingSourceType::Function && !Binding->SourceFunction.IsNone())
 				{
 					CachedFunctionStatus = ValidateFunctionBinding(
-						B->SourceFunction, TargetParamType, Cast<UUAFRigVMAsset>(WeakAsset.Get()),
+						Binding->SourceFunction, TargetParamType, Cast<UUAFRigVMAsset>(WeakAsset.Get()),
 						&CachedFunctionDisplayName);
 				}
 			}
@@ -737,15 +737,15 @@ void FBindableValueCustomization::CustomizeHeader(
 					[&Result](const void* RawData, const int32, const int32) -> bool
 					{
 						const FBindableValueBase* Base = static_cast<const FBindableValueBase*>(RawData);
-						if (const FUAFPropertyBinding* B = Base->GetBinding())
+						if (const FUAFPropertyBinding* PropertyBinding = Base->GetBinding())
 						{
-							if (B->SourceType == EUAFBindingSourceType::Function && !B->SourceFunction.IsNone())
+							if (PropertyBinding->SourceType == EUAFBindingSourceType::Function && !PropertyBinding->SourceFunction.IsNone())
 							{
-								Result = FText::FromName(B->SourceFunction.GetEventName());
+								Result = FText::FromName(PropertyBinding->SourceFunction.GetEventName());
 							}
-							else if (!B->SourceVariable.IsNone())
+							else if (!PropertyBinding->SourceVariable.IsNone())
 							{
-								Result = FText::FromName(B->SourceVariable.GetName());
+								Result = FText::FromName(PropertyBinding->SourceVariable.GetName());
 							}
 						}
 						return false;
@@ -769,9 +769,9 @@ void FBindableValueCustomization::CustomizeHeader(
 							[&bIsFunctionBinding](const void* RawData, const int32, const int32) -> bool
 							{
 								const FBindableValueBase* Base = static_cast<const FBindableValueBase*>(RawData);
-								if (const FUAFPropertyBinding* B = Base->GetBinding())
+								if (const FUAFPropertyBinding* PropertyBinding = Base->GetBinding())
 								{
-									bIsFunctionBinding = B->SourceType == EUAFBindingSourceType::Function && !B->SourceFunction.IsNone();
+									bIsFunctionBinding = PropertyBinding->SourceType == EUAFBindingSourceType::Function && !PropertyBinding->SourceFunction.IsNone();
 								}
 								return false;
 							});
@@ -793,9 +793,9 @@ void FBindableValueCustomization::CustomizeHeader(
 							[&bIsFunctionBinding](const void* RawData, const int32, const int32) -> bool
 							{
 								const FBindableValueBase* Base = static_cast<const FBindableValueBase*>(RawData);
-								if (const FUAFPropertyBinding* B = Base->GetBinding())
+								if (const FUAFPropertyBinding* PropertyBinding = Base->GetBinding())
 								{
-									bIsFunctionBinding = B->SourceType == EUAFBindingSourceType::Function && !B->SourceFunction.IsNone();
+									bIsFunctionBinding = PropertyBinding->SourceType == EUAFBindingSourceType::Function && !PropertyBinding->SourceFunction.IsNone();
 								}
 								return false;
 							});
@@ -822,9 +822,9 @@ void FBindableValueCustomization::CustomizeHeader(
 							[&Result, CachedFunctionStatus, CachedFunctionDisplayName](const void* RawData, const int32, const int32) -> bool
 							{
 								const FBindableValueBase* Base = static_cast<const FBindableValueBase*>(RawData);
-								if (const FUAFPropertyBinding* B = Base->GetBinding())
+								if (const FUAFPropertyBinding* Binding = Base->GetBinding())
 								{
-									if (B->SourceType == EUAFBindingSourceType::Function && !B->SourceFunction.IsNone())
+									if (Binding->SourceType == EUAFBindingSourceType::Function && !Binding->SourceFunction.IsNone())
 									{
 										// Use the resolved name from GUID lookup (handles renames).
 										// Fall back to parsing EventName for legacy data without GUID.
@@ -835,7 +835,7 @@ void FBindableValueCustomization::CustomizeHeader(
 										}
 										else
 										{
-											Name = B->SourceFunction.GetEventName().ToString();
+											Name = Binding->SourceFunction.GetEventName().ToString();
 											Name.RemoveFromStart(TEXT("__InternalCall_"));
 										}
 
@@ -849,13 +849,13 @@ void FBindableValueCustomization::CustomizeHeader(
 										}
 										Result = FText::FromString(Name);
 									}
-									else if (!B->SourceVariable.IsNone())
+									else if (!Binding->SourceVariable.IsNone())
 									{
-										FString DisplayName = B->SourceVariable.GetName().ToString();
-										if (B->SourceType == EUAFBindingSourceType::SubProperty
-											&& B->SubPropertyPath.NumSegments() > 0)
+										FString DisplayName = Binding->SourceVariable.GetName().ToString();
+										if (Binding->SourceType == EUAFBindingSourceType::SubProperty
+											&& Binding->SubPropertyPath.NumSegments() > 0)
 										{
-											DisplayName += TEXT(".") + B->SubPropertyPath.ToString();
+											DisplayName += TEXT(".") + Binding->SubPropertyPath.ToString();
 										}
 										Result = FText::FromString(DisplayName);
 									}
