UAFBrowser: AddButton
- Add support for creating assets within UAF Browser
- Initial creation is in same location of current asset. More location options to come in future revisions.
- Expose section changing in tagged asset browser

#tests editor
#Synced-CL 52437274

diff --git a/Engine/Plugins/Experimental/UAF/UAF/Source/UAFEditor/Private/Common/SUAFBrowser.cpp b/Engine/Plugins/Experimental/UAF/UAF/Source/UAFEditor/Private/Common/SUAFBrowser.cpp
--- a/Engine/Plugins/Experimental/UAF/UAF/Source/UAFEditor/Private/Common/SUAFBrowser.cpp
+++ b/Engine/Plugins/Experimental/UAF/UAF/Source/UAFEditor/Private/Common/SUAFBrowser.cpp
@@ -4,25 +4,46 @@
 
 #include "AnimNextRigVMAsset.h"
 #include "AnimNextRigVMAssetEditorData.h"
+#include "AssetToolsModule.h"
 #include "Assets/TaggedAssetBrowserConfiguration.h"
+#include "ClassViewerFilter.h"
+#include "ClassViewerModule.h"
 #include "Common/AssetPreview/AnimSequenceAssetPreview.h"
 #include "Common/AssetPreview/IUAFAssetPreview.h"
 #include "Common/AssetPreview/UAFAssetPreviewFactorySubsystem.h"
+#include "Common/UAFBrowserMenuContext.h"
 #include "Common/UAFBrowserTabSummoner.h"
 #include "ContentBrowserModule.h"
+#include "Editor.h"
+#include "Engine/Blueprint.h"
+#include "Factories/BlueprintFactory.h"
+#include "Framework/Application/SlateApplication.h"
+#include "Kismet2/KismetEditorUtilities.h"
+#include "Kismet2/SClassPickerDialog.h"
+#include "IAssetTools.h"
 #include "IContentBrowserSingleton.h"
 #include "IWorkspaceEditor.h"
 #include "IWorkspaceEditorModule.h"
 #include "Misc/ConfigCacheIni.h"
+#include "Misc/PackageName.h"
+#include "NewAssetContextMenu.h"
 #include "SAssetView.h"
 #include "Sound/SoundWave.h"
 #include "Settings/UAFEditorProjectSettings.h"
 #include "Settings/UAFEditorUserSettings.h"
 #include "StatusBarSubsystem.h"
 #include "Subsystems/AssetEditorSubsystem.h"
+#include "TaggedAssetBrowserFilters/TaggedAssetBrowser_CommonFilters.h"
+#include "ToolMenus.h"
 #include "UncookedOnlyUtils.h"
+#include "Widgets/Input/SButton.h"
+#include "Widgets/Input/SEditableTextBox.h"
 #include "Widgets/Layout/SBorder.h"
 #include "Widgets/Layout/SBox.h"
+#include "SActionButton.h"
+#include "SPositiveActionButton.h"
+#include "UObject/ObjectMacros.h"
+#include "Widgets/SAssetMenuIcon.h"
 #include "Widgets/STaggedAssetBrowser.h"
 
 #include "Animation/AnimSequence.h"
@@ -32,7 +53,28 @@
 
 namespace UE::UAF::Editor
 {
+	
+FName SUAFBrowser::AddNewMenuName = FName(TEXT("UAFBrowser.AddNewContextMenu"));
+
+/** Class filter that restricts the class viewer to children of a set of allowed parent classes. */
+class FUAFBlueprintParentClassFilter : public IClassViewerFilter
+{
+public:
+	TSet<const UClass*> AllowedParentClasses;
+	EClassFlags DisallowedClassFlags = CLASS_Deprecated | CLASS_NewerVersionExists | CLASS_Hidden | CLASS_HideDropDown;
 
+	virtual bool IsClassAllowed(const FClassViewerInitializationOptions& InInitOptions, const UClass* InClass, TSharedRef<FClassViewerFilterFuncs> InFilterFuncs) override
+	{
+		return !InClass->HasAnyClassFlags(DisallowedClassFlags)
+			&& InFilterFuncs->IfInChildOfClassesSet(AllowedParentClasses, InClass) != EFilterReturn::Failed;
+	}
+
+	virtual bool IsUnloadedClassAllowed(const FClassViewerInitializationOptions& InInitOptions, const TSharedRef<const IUnloadedBlueprintData> InUnloadedClassData, TSharedRef<FClassViewerFilterFuncs> InFilterFuncs) override
+	{
+		return !InUnloadedClassData->HasAnyClassFlags(DisallowedClassFlags)
+			&& InFilterFuncs->IfInChildOfClassesSet(AllowedParentClasses, InUnloadedClassData) != EFilterReturn::Failed;
+	}
+};
 
 SUAFBrowser::~SUAFBrowser()
 {
@@ -47,6 +89,238 @@ void SUAFBrowser::Construct(const FArguments& InArgs, const TSharedPtr<UE::Works
 	WeakHostingApp = InHostingApp;
 	TabState = InTabState;
 
+	// Register the shared add-new ToolMenu once, subsequent browser instances reuse the same registration.
+	if (!UToolMenus::Get()->IsMenuRegistered(AddNewMenuName))
+	{
+		if (UToolMenu* Menu = UToolMenus::Get()->RegisterMenu(AddNewMenuName, NAME_None, EMultiBoxType::Menu))
+		{
+			// Flat promoted entries for the active section's filter classes, each triggering a name prompt before creation
+			Menu->AddDynamicSection(TEXT("UAFBrowserPrioritized"),
+				FNewToolMenuDelegate::CreateLambda([](UToolMenu* InMenu)
+				{
+					UUAFBrowserMenuContext* Ctx = InMenu->Context.FindContext<UUAFBrowserMenuContext>();
+					TSharedPtr<SUAFBrowser> Browser = Ctx ? Ctx->WeakOwningBrowser.Pin() : nullptr;
+					if (!Browser)
+					{
+						return;
+					}
+
+					TArray<UClass*> CurrentSectionClasses = Browser->GetCurrentSectionFilterClasses();
+					if (CurrentSectionClasses.IsEmpty())
+					{
+						return;
+					}
+
+					IAssetTools& AssetTools = FModuleManager::LoadModuleChecked<FAssetToolsModule>("AssetTools").Get();
+					TArray<UFactory*> AllFactories = AssetTools.GetNewAssetFactories();
+
+					FToolMenuSection& PrioritizedSection = InMenu->AddSection(
+						TEXT("UAFBrowserPrioritized"), LOCTEXT("SuggestedTypes", "Suggested Types"));
+
+					// Mirror the SFactoryMenuEntry layout from NewAssetContextMenu.cpp.
+					// Read the same CVar that ContentBrowserStyle.h uses for IsNewStyleEnabled().
+					const IConsoleVariable* CVar = IConsoleManager::Get().FindConsoleVariable(TEXT("ContentBrowser.EnableNewStyle"));
+					const bool bNewStyle = CVar && CVar->GetBool();
+					static constexpr float DefaultMenuIconSize = 14.0f;
+					static constexpr float VerticalEntryPadding = 4.0f;
+					const FVector2D IconContainerSize = bNewStyle ? FVector2D(24, 24) : FVector2D(32, 32);
+					const FVector2D IconSize = bNewStyle ? FVector2D(16, 16) : FVector2D(28, 28);
+					const FMargin IconSlotPadding = bNewStyle ? FMargin(2, 0, 3, 0) : FMargin(0, 0, 0, 1);
+					const FMargin LabelSlotPadding = bNewStyle ? FMargin(4, 0, 6, 0) : FMargin(4, 0, 4, 0);
+					const float VPadAdj = ((IconContainerSize.Y - DefaultMenuIconSize) / 2.0f) - VerticalEntryPadding;
+					const FMargin ChildSlotPadding = bNewStyle ? FMargin(0, -VPadAdj, 0, -VPadAdj) : FMargin(0);
+
+					// Weak pointer extracted at section-build time, captured by click-time lambdas below.
+					TWeakPtr<SUAFBrowser> WeakBrowser = Ctx->WeakOwningBrowser;
+
+					// Helper that builds and registers one menu entry for the prioritized section.
+					auto AddSectionEntry = [&](UClass* EntryIconClass
+						, const FName& EntryIconOverride
+						, const FText& EntryDisplayName
+						, TFunction<void(TSharedPtr<SUAFBrowser>)> EntryOnClicked)
+					{
+						TSharedRef<SWidget> EntryWidget =
+							SNew(SBox)
+							.Padding(ChildSlotPadding)
+							[
+								SNew(SHorizontalBox)
+								.ToolTipText(EntryDisplayName)
+								+ SHorizontalBox::Slot()
+								.AutoWidth()
+								.VAlign(VAlign_Center)
+								.Padding(IconSlotPadding)
+								[
+									SNew(SAssetMenuIcon, EntryIconClass, EntryIconOverride)
+									.IconContainerSize(IconContainerSize)
+									.IconSize(IconSize)
+								]
+								+ SHorizontalBox::Slot()
+								.VAlign(VAlign_Center)
+								.Padding(LabelSlotPadding)
+								[
+									SNew(STextBlock)
+									.Font(FAppStyle::GetFontStyle("LevelViewportContextMenu.AssetLabel.Text.Font"))
+									.Text(EntryDisplayName)
+								]
+							];
+
+						PrioritizedSection.AddEntry(
+							FToolMenuEntry::InitMenuEntry(
+								FName(*EntryDisplayName.ToString()),
+								FUIAction(FExecuteAction::CreateLambda([WeakBrowser, EntryOnClicked]()
+								{
+									TSharedPtr<SUAFBrowser> PinnedBrowser = WeakBrowser.Pin();
+									if (PinnedBrowser)
+									{
+										EntryOnClicked(PinnedBrowser);
+									}
+								})),
+								EntryWidget
+							)
+						);
+					};
+
+					for (UClass* TargetClass : CurrentSectionClasses)
+					{
+						// Collect ALL factories whose supported class is a child of TargetClass.
+						// There can be multiple (e.g. several child-class factories for one base class).
+						TArray<UFactory*> MatchedFactories;
+						for (UFactory* Factory : AllFactories)
+						{
+							UClass* Supported = Factory->GetSupportedClass();
+							if (Supported && Supported->IsChildOf(TargetClass))
+							{
+								MatchedFactories.Add(Factory);
+							}
+						}
+
+						if (!MatchedFactories.IsEmpty())
+						{
+							// Factory path: one entry per factory.
+							for (UFactory* Factory : MatchedFactories)
+							{
+								UClass* IconClass = Factory->GetSupportedClass();
+								const FText DisplayName = Factory->GetDisplayName();
+								const FName IconOverride = bNewStyle ? Factory->GetNewAssetIconOverride() : Factory->GetNewAssetThumbnailOverride();
+								TWeakObjectPtr<UClass> WeakFactoryClass(Factory->GetClass());
+
+								AddSectionEntry(IconClass, IconOverride, DisplayName,
+									[WeakFactoryClass, DisplayName](TSharedPtr<SUAFBrowser> Browser)
+									{
+										UClass* FactoryClass = WeakFactoryClass.Get();
+										if (FactoryClass)
+										{
+											Browser->CreateAssetWithNameDialog(FactoryClass, DisplayName);
+										}
+									});
+							}
+						}
+						else if (FKismetEditorUtilities::CanCreateBlueprintOfClass(TargetClass))
+						{
+							// Blueprint path: show a class picker filtered to children of TargetClass,
+							// with TargetClass pre-highlighted as the common/default choice.
+							const FText DisplayName = TargetClass->GetDisplayNameText();
+							TWeakObjectPtr<UClass> WeakParentClass(TargetClass);
+
+							AddSectionEntry(UBlueprint::StaticClass(), NAME_None, DisplayName,
+								[WeakParentClass, DisplayName](TSharedPtr<SUAFBrowser> Browser)
+								{
+									UClass* DefaultParentClass = WeakParentClass.Get();
+									if (!DefaultParentClass)
+									{
+										return;
+									}
+
+									FClassViewerInitializationOptions Options;
+									Options.Mode = EClassViewerMode::ClassPicker;
+									Options.bIsBlueprintBaseOnly = true;
+									Options.bShowNoneOption = false;
+									Options.ExtraPickerCommonClasses.Add(DefaultParentClass);
+
+									TSharedPtr<FUAFBlueprintParentClassFilter> Filter = MakeShared<FUAFBlueprintParentClassFilter>();
+									Filter->AllowedParentClasses.Add(DefaultParentClass);
+									Options.ClassFilters.Add(Filter.ToSharedRef());
+
+									const FText TitleText = FText::Format(
+										LOCTEXT("PickBlueprintParentClass", "Pick Parent Class for New {0} Blueprint"),
+										DisplayName);
+
+									UClass* ChosenClass = nullptr;
+									const bool bPicked = SClassPickerDialog::PickClass(TitleText, Options, ChosenClass, DefaultParentClass);
+									if (!bPicked)
+									{
+										return;
+									}
+									if (!ChosenClass)
+									{
+										ChosenClass = DefaultParentClass;
+									}
+
+									// Creating a temp factory to spawn assets is normal.
+									UBlueprintFactory* BPFactory = NewObject<UBlueprintFactory>(GetTransientPackage());
+									BPFactory->ParentClass = ChosenClass;
+									BPFactory->bSkipClassPicker = true;
+									Browser->CreateAssetWithNameDialog(BPFactory, DisplayName);
+								});
+						}
+						// else: import-only class — no Add entry
+					}
+				}));
+
+			// Fallback: full factory list (default content-browser behavior) shown only when the
+			// active section has no class filter — i.e. UAFBrowserPrioritized contributed nothing.
+			Menu->AddDynamicSection(TEXT("UAFBrowserAllTypes"),
+				FNewToolMenuDelegate::CreateLambda([](UToolMenu* InMenu)
+				{
+					UUAFBrowserMenuContext* Ctx = InMenu->Context.FindContext<UUAFBrowserMenuContext>();
+					TSharedPtr<SUAFBrowser> Browser = Ctx ? Ctx->WeakOwningBrowser.Pin() : nullptr;
+					if (!Browser)
+					{
+						return;
+					}
+
+					// Yield to the prioritized section whenever it has items to show
+					if (!Browser->GetCurrentSectionFilterClasses().IsEmpty())
+					{
+						return;
+					}
+
+					TWeakPtr<SUAFBrowser> WeakBrowser = Ctx->WeakOwningBrowser;
+
+					FNewAssetContextMenu::FOnNewAssetRequested OnNewAsset =
+						FNewAssetContextMenu::FOnNewAssetRequested::CreateLambda(
+							[WeakBrowser](const FName& /*Path*/, TWeakObjectPtr<UClass> FactoryClass)
+							{
+								TSharedPtr<SUAFBrowser> PinnedBrowser = WeakBrowser.Pin();
+								UClass* FactoryClassPtr = FactoryClass.Get();
+								if (!PinnedBrowser || !FactoryClassPtr)
+								{
+									return;
+								}
+								UFactory* TempFactory = NewObject<UFactory>(GetTransientPackage(), FactoryClassPtr);
+								PinnedBrowser->CreateAssetWithNameDialog(FactoryClassPtr, TempFactory->GetDisplayName());
+							});
+
+					FNewAssetContextMenu::FOnImportAssetRequested OnImport =
+						FNewAssetContextMenu::FOnImportAssetRequested::CreateLambda(
+							[WeakBrowser](const FName& /*Path*/)
+							{
+								TSharedPtr<SUAFBrowser> PinnedBrowser = WeakBrowser.Pin();
+								if (!PinnedBrowser)
+								{
+									return;
+								}
+								const FString NewAssetPath = PinnedBrowser->GetNewAssetPath();
+								FModuleManager::LoadModuleChecked<FAssetToolsModule>("AssetTools")
+									.Get().ImportAssetsWithDialog(NewAssetPath);
+							});
+
+					TArray<FName> Paths = { FName(*Browser->GetNewAssetPath()) };
+					FNewAssetContextMenu::MakeContextMenu(InMenu, Paths, OnImport, OnNewAsset);
+				}));
+		}
+	}
+
 	if (UUAFEditorUserSettings* UAFSettings = GetMutableDefault<UUAFEditorUserSettings>())
 	{
 		UAFSettings->OnBrowserHiddenColumnsChanged.AddSP(this, &SUAFBrowser::OnExternalColumnVisibilityChanged);
@@ -115,11 +389,20 @@ TSharedRef<SWidget> SUAFBrowser::Construct_GetTaggedAssetBrowser()
 	{
 		// We intentionally do not restrict AvailableClasses as this will filter types we may not know about from beind displayed
 		// Ex: Even in contexts such as filtering by tag
-		return SNew(STaggedAssetBrowser, *CategoryConfiguration)
+		TSharedRef<STaggedAssetBrowser> TaggedBrowser = SNew(STaggedAssetBrowser, *CategoryConfiguration)
 			.AvailableClasses({})
 			.DefaultDetailsTabConfiguration(DefaultDetailsTabConfiguration)
 			.InterfaceOverrideProfiles(OverrideProfiles)
-			.OnOverrideAssetPickerConfig(STaggedAssetBrowser::FOnOverrideAssetPickerConfig::CreateSP(this, &SUAFBrowser::OnOverrideAssetPickerConfig));
+			.OnOverrideAssetPickerConfig(STaggedAssetBrowser::FOnOverrideAssetPickerConfig::CreateSP(this, &SUAFBrowser::OnOverrideAssetPickerConfig))
+			.OnActiveSectionChanged(STaggedAssetBrowser::FOnActiveSectionChanged::CreateSP(this, &SUAFBrowser::OnActiveSectionChanged))
+			.OnPrimaryFilterSelectionChanged(STaggedAssetBrowser::FOnPrimaryFilterSelectionChanged::CreateSP(this, &SUAFBrowser::OnPrimaryFilterSelectionChanged));
+
+		TaggedAssetBrowserWidgetWeak = TaggedBrowser;
+
+		// @TODO: Implement below, not used at the moment.
+		// Initialize current section from the browser's initial state so the Add button is correct before any section switch
+		CurrentActiveSection = TaggedBrowser->GetActiveSection();
+		return TaggedBrowser;
 	}
 
 	return SNullWidget::NullWidget;
@@ -192,20 +475,34 @@ void SUAFBrowser::OnOverrideAssetPickerConfig(FAssetPickerConfig& Config)
 
 void SUAFBrowser::OnExtendAssetPickerTopBar(TSharedRef<SHorizontalBox> InAssetPickerTopBar)
 {
+	// Insert add/import button(s) at position 0 (before the search bar and filter UI)
+	SAssignNew(AddImportButtonContainer, SHorizontalBox);
+	AddImportButtonContainer->AddSlot()
+		.AutoWidth()
+		[
+			BuildAddImportButtons()
+		];
+
+	InAssetPickerTopBar->InsertSlot(0)
+		.AutoWidth()
+		[
+			AddImportButtonContainer.ToSharedRef()
+		];
+
 	if (TabState == ETabState::SidebarTab || TabState == ETabState::StatusbarTab)
 	{
-		// Note: Value from `FAssetPickerConfig::OnExtendAssetPickerTopBar`. 
+		// Note: Value from `FAssetPickerConfig::OnExtendAssetPickerTopBar`.
 		// We could make it expose those indicies, but given expected usage that feels over-engineered for now.
 
 		const int ShowDeveloperContentIndex = 2;
 
-		InAssetPickerTopBar->InsertSlot(ShowDeveloperContentIndex + 1)
+		InAssetPickerTopBar->InsertSlot(ShowDeveloperContentIndex + 2)
 		.AutoWidth()
 		[
 			OnExtendAssetPickerTopBar_GetDrawerDockButton()
 		];
 
-		InAssetPickerTopBar->InsertSlot(ShowDeveloperContentIndex + 2)
+		InAssetPickerTopBar->InsertSlot(ShowDeveloperContentIndex + 3)
 		.AutoWidth()
 		[
 			OnExtendAssetPickerTopBar_GetToggleSwapHotkeyButton()
@@ -366,6 +663,295 @@ void SUAFBrowser::OnExternalColumnVisibilityChanged(const TArray<FString>& Hidde
 	LoadAssetViewSettingsDelegate.ExecuteIfBound();
 }
 
+void SUAFBrowser::OnActiveSectionChanged(const UTaggedAssetBrowserSection* NewSection)
+{
+	CurrentActiveSection = NewSection;
+
+	if (AddImportButtonContainer.IsValid())
+	{
+		AddImportButtonContainer->ClearChildren();
+		AddImportButtonContainer->AddSlot()
+			.AutoWidth()
+			[
+				BuildAddImportButtons()
+			];
+	}
+}
+
+void SUAFBrowser::OnPrimaryFilterSelectionChanged()
+{
+	if (AddImportButtonContainer.IsValid())
+	{
+		AddImportButtonContainer->ClearChildren();
+		AddImportButtonContainer->AddSlot()
+			.AutoWidth()
+			[
+				BuildAddImportButtons()
+			];
+	}
+}
+
+TSharedRef<SWidget> SUAFBrowser::BuildAddImportButtons()
+{
+	TArray<UClass*> SectionClasses = GetCurrentSectionFilterClasses();
+	const EUAFAddButtonMode Mode = GetAddButtonMode(SectionClasses);
+
+	TSharedRef<SHorizontalBox> Box = SNew(SHorizontalBox);
+
+	if (Mode != EUAFAddButtonMode::ImportOnly)
+	{
+		Box->AddSlot()
+			.AutoWidth()
+			[
+				SNew(SPositiveActionButton)
+				.Text(LOCTEXT("Add", "Add"))
+				.OnGetMenuContent(this, &SUAFBrowser::BuildAddMenu)
+			];
+	}
+
+	if (Mode == EUAFAddButtonMode::ImportOnly || Mode == EUAFAddButtonMode::Both)
+	{
+		Box->AddSlot()
+			.AutoWidth()
+			[
+				SNew(SActionButton)
+				.Text(LOCTEXT("Import", "Import"))
+				.Icon(FAppStyle::Get().GetBrush("LevelEditor.ImportContent"))
+				.OnClicked(this, &SUAFBrowser::OnImportButtonClicked)
+			];
+	}
+
+	return Box;
+}
+
+TSharedRef<SWidget> SUAFBrowser::BuildAddMenu()
+{
+	// Lifetime handled by ToolMenus system
+	UUAFBrowserMenuContext* Ctx = NewObject<UUAFBrowserMenuContext>();
+	Ctx->WeakOwningBrowser = SharedThis(this).ToWeakPtr();
+
+	FToolMenuContext ToolMenuContext;
+	ToolMenuContext.AddObject(Ctx);
+	return UToolMenus::Get()->GenerateWidget(AddNewMenuName, ToolMenuContext);
+}
+
+FReply SUAFBrowser::OnImportButtonClicked()
+{
+	const FString NewAssetPath = GetNewAssetPath();
+	FModuleManager::LoadModuleChecked<FAssetToolsModule>("AssetTools")
+		.Get().ImportAssetsWithDialog(NewAssetPath);
+	return FReply::Handled();
+}
+
+FString SUAFBrowser::GetNewAssetPath() const
+{
+	if (TSharedPtr<UE::Workspace::IWorkspaceEditor> Editor = WeakHostingApp.Pin())
+	{
+		if (UObject* FocusedAsset = Editor->GetFocusedAsset())
+		{
+			return FPackageName::GetLongPackagePath(FocusedAsset->GetOutermost()->GetName());
+		}
+		// Fallback: workspace asset's folder
+		if (UObject* WorkspaceAsset = Editor->GetWorkspaceAsset())
+		{
+			return FPackageName::GetLongPackagePath(WorkspaceAsset->GetOutermost()->GetName());
+		}
+	}
+	return TEXT("/Game");
+}
+
+TArray<UClass*> SUAFBrowser::GetCurrentSectionFilterClasses() const
+{
+	TArray<UClass*> Result;
+	TSharedPtr<STaggedAssetBrowser> Browser = TaggedAssetBrowserWidgetWeak.Pin();
+	if (!Browser)
+	{
+		return Result;
+	}
+
+	// Prefer the currently selected primary filters; fall back to all filters in the section
+	// so that GetAddButtonMode() returns the correct mode even before the user picks a filter.
+	TArray<const UTaggedAssetBrowserFilterBase*> Filters = Browser->GetSelectedPrimaryFilters();
+	if (Filters.IsEmpty())
+	{
+		Filters = Browser->GetAllPrimaryFilters();
+	}
+
+	for (const UTaggedAssetBrowserFilterBase* Filter : Filters)
+	{
+		if (const UTaggedAssetBrowserFilter_Class* ClassFilter = Cast<UTaggedAssetBrowserFilter_Class>(Filter))
+		{
+			for (const TObjectPtr<UClass>& Class : ClassFilter->Classes)
+			{
+				if (Class)
+				{
+					Result.AddUnique(Class);
+				}
+			}
+		}
+	}
+	return Result;
+}
+
+SUAFBrowser::EUAFAddButtonMode SUAFBrowser::GetAddButtonMode(const TArray<UClass*>& SectionClasses) const
+{
+	if (SectionClasses.IsEmpty())
+	{
+		return EUAFAddButtonMode::Default;
+	}
+
+	IAssetTools& AssetTools = FModuleManager::LoadModuleChecked<FAssetToolsModule>("AssetTools").Get();
+	TArray<UFactory*> AllFactories = AssetTools.GetNewAssetFactories();
+
+	bool bAnyCreatable = false;
+	bool bAnyImportOnly = false;
+
+	for (UClass* Class : SectionClasses)
+	{
+		// Priority 1: a dedicated creatable factory exists
+		const bool bHasFactory = AllFactories.ContainsByPredicate([Class](UFactory* Factory)
+		{
+			UClass* Supported = Factory->GetSupportedClass();
+			return Supported && (Supported == Class || Supported->IsChildOf(Class));
+		});
+
+		// Priority 2: no factory but the class is blueprintable (child Blueprint can be created)
+		const bool bIsBlueprintable = !bHasFactory && FKismetEditorUtilities::CanCreateBlueprintOfClass(Class);
+
+		if (bHasFactory || bIsBlueprintable)
+		{
+			bAnyCreatable = true;
+		}
+		else
+		{
+			bAnyImportOnly = true;
+		}
+	}
+
+	if (bAnyCreatable && bAnyImportOnly)
+	{
+		return EUAFAddButtonMode::Both;
+	}
+	if (bAnyCreatable)
+	{
+		return EUAFAddButtonMode::AddOnly;
+	}
+	return EUAFAddButtonMode::ImportOnly;
+}
+
+void SUAFBrowser::CreateAssetWithNameDialog(UClass* InFactoryClass, const FText& AssetTypeName)
+{
+	// Creating a temp factory to spawn assets is normal.
+	UFactory* Factory = NewObject<UFactory>(GetTransientPackage(), InFactoryClass);
+	CreateAssetWithNameDialog(Factory, AssetTypeName);
+}
+
+void SUAFBrowser::CreateAssetWithNameDialog(UFactory* InFactory, const FText& AssetTypeName)
+{
+	UFactory* Factory = InFactory;
+	const FString NewAssetPath = GetNewAssetPath();
+
+	// Allow the factory to configure its properties (e.g. Blueprint class picker) before we ask for a name.
+	// Matches the behavior of IAssetTools::CreateAssetWithDialog.
+	FEditorDelegates::OnConfigureNewAssetProperties.Broadcast(Factory);
+	if (!Factory->ConfigureProperties())
+	{
+		return;
+	}
+
+	IAssetTools& AssetTools = FModuleManager::LoadModuleChecked<FAssetToolsModule>("AssetTools").Get();
+	FString DefaultName;
+	FString PackageName;
+	AssetTools.CreateUniqueAssetName(NewAssetPath / Factory->GetDefaultNewAssetName(), TEXT(""), PackageName, DefaultName);
+
+	// Local state written by dialog callbacks and read after AddModalWindow returns
+	FString ChosenName = DefaultName;
+	bool bConfirmed = false;
+
+	// V1 dialog, just ask for asset name. Later on can add relative folder / etc.
+	TSharedRef<SWindow> DialogWindow = SNew(SWindow)
+		.Title(LOCTEXT("EnterNewAssetNameTitle", "Enter New Asset Name"))
+		.ClientSize(FVector2D(400.f, 110.f))
+		.SizingRule(ESizingRule::FixedSize)
+		.SupportsMaximize(false)
+		.SupportsMinimize(false);
+
+	TSharedPtr<SEditableTextBox> TextBox;
+
+	// Confirm: read text, set flag, close window. AddModalWindow is synchronous so all
+	// locals above remain alive until after it returns.
+	auto DoConfirm = [&]()
+	{
+		ChosenName = TextBox->GetText().ToString();
+		bConfirmed = true;
+		DialogWindow->RequestDestroyWindow();
+	};
+
+	DialogWindow->SetContent(
+		SNew(SBox)
+		.Padding(16.f)
+		[
+			SNew(SVerticalBox)
+			+ SVerticalBox::Slot()
+			.AutoHeight()
+			.Padding(0.f, 0.f, 0.f, 8.f)
+			[
+				SAssignNew(TextBox, SEditableTextBox)
+				.Text(FText::FromString(DefaultName))
+				.SelectAllTextWhenFocused(true)
+				.OnTextCommitted(FOnTextCommitted::CreateLambda([&DoConfirm](const FText&, ETextCommit::Type CommitType)
+				{
+					if (CommitType == ETextCommit::OnEnter)
+					{
+						DoConfirm();
+					}
+				}))
+			]
+			+ SVerticalBox::Slot()
+			.AutoHeight()
+			.HAlign(HAlign_Right)
+			[
+				SNew(SHorizontalBox)
+				+ SHorizontalBox::Slot()
+				.AutoWidth()
+				.Padding(0.f, 0.f, 4.f, 0.f)
+				[
+					SNew(SButton)
+					.Text(LOCTEXT("CreateButton", "Create"))
+					.OnClicked(FOnClicked::CreateLambda([&DoConfirm]()
+					{
+						DoConfirm();
+						return FReply::Handled();
+					}))
+				]
+				+ SHorizontalBox::Slot()
+				.AutoWidth()
+				[
+					SNew(SButton)
+					.Text(LOCTEXT("CancelButton", "Cancel"))
+					.OnClicked(FOnClicked::CreateLambda([DialogWindow]()
+					{
+						DialogWindow->RequestDestroyWindow();
+						return FReply::Handled();
+					}))
+				]
+			]
+		]
+	);
+
+	// While a temp factory is normal, we're about to spawn a non engine blocking modal, so pin the factory till we exit.
+	TStrongObjectPtr<UFactory> FactoryPinned = TStrongObjectPtr<UFactory>(Factory);
+	TSharedPtr<SWindow> ParentWindow = FSlateApplication::Get().GetActiveTopLevelWindow();
+	
+	// Does not return until modal is closed.
+	FSlateApplication::Get().AddModalWindow(DialogWindow, ParentWindow);
+
+	if (bConfirmed && !ChosenName.IsEmpty())
+	{
+		AssetTools.CreateAsset(ChosenName, NewAssetPath, Factory->GetSupportedClass(), Factory);
+	}
+}
+
 } // namespace UE::UAF::Editor
 
 #undef LOCTEXT_NAMESPACE
diff --git a/Engine/Plugins/Experimental/UAF/UAF/Source/UAFEditor/Private/Common/SUAFBrowser.h b/Engine/Plugins/Experimental/UAF/UAF/Source/UAFEditor/Private/Common/SUAFBrowser.h
--- a/Engine/Plugins/Experimental/UAF/UAF/Source/UAFEditor/Private/Common/SUAFBrowser.h
+++ b/Engine/Plugins/Experimental/UAF/UAF/Source/UAFEditor/Private/Common/SUAFBrowser.h
@@ -19,6 +19,8 @@ namespace UE::UAF::Editor
 	class IUAFAssetPreview;
 };
 
+class STaggedAssetBrowser;
+class UTaggedAssetBrowserSection;
 struct FAssetPickerConfig;
 class SHorizontalBox;
 
@@ -61,6 +63,20 @@ private:
 	/** Builds fallack when preview not supported by type in UAF browser Preview. */
 	TSharedRef<SWidget> Construct_GetAssetPreviewNotSupportedWidget();
 
+private:
+
+	/** Used for UToolMenu registration of add menu. */
+	static FName AddNewMenuName;
+
+	/** How AddButton should appear based on filter(s) active. */
+	enum class EUAFAddButtonMode 
+	{ 
+		AddOnly, 
+		ImportOnly,
+		Both, 
+		Default 
+	};
+
 private:
 
 	/** Callback to override asset picker configs for our tagged asset browser */
@@ -81,6 +97,36 @@ private:
 	/** Callback for when dock to drawer button is pressed. Hides all drawers and invokes docked browser tab. */
 	FReply OnDockDrawerButtonClicked();
 
+	/** Callback for when the active tagged asset browser section changes. Rebuilds add/import buttons. */
+	void OnActiveSectionChanged(const UTaggedAssetBrowserSection* NewSection);
+
+	/** Callback for when the primary filter tree selection changes. Rebuilds add/import buttons. */
+	void OnPrimaryFilterSelectionChanged();
+
+	/** Builds the add/import button(s) appropriate for the current active section. */
+	TSharedRef<SWidget> BuildAddImportButtons();
+
+	/** Generates the add menu widget for the per-instance ToolMenu. */
+	TSharedRef<SWidget> BuildAddMenu();
+
+	/** Shows a modal dialog to name the new asset, then creates it if confirmed. */
+	void CreateAssetWithNameDialog(UClass* InFactoryClass, const FText& AssetTypeName);
+
+	/** Overload that accepts a pre-configured factory (e.g. UBlueprintFactory with ParentClass already set). */
+	void CreateAssetWithNameDialog(UFactory* InFactory, const FText& AssetTypeName);
+
+	/** Opens the OS import dialog targeting the current asset folder. */
+	FReply OnImportButtonClicked();
+
+	/** Returns the package folder path to use for asset creation/import. */
+	FString GetNewAssetPath() const;
+
+	/** Collects UClass pointers from the currently selected primary filter items in the tagged asset browser. */
+	TArray<UClass*> GetCurrentSectionFilterClasses() const;
+
+	/** Determines which button(s) to show based on which section classes have creatable factories. */
+	EUAFAddButtonMode GetAddButtonMode(const TArray<UClass*>& SectionClasses) const;
+
 private:
 
 	/** Callback for when an asset is selected. */
@@ -100,6 +146,15 @@ private:
 	/** App Hosting this Browser, to query settings against */
 	TWeakPtr<UE::Workspace::IWorkspaceEditor> WeakHostingApp = nullptr;
 
+	/** The tagged asset browser widget, kept so we can query GetActiveSection() directly. */
+	TWeakPtr<STaggedAssetBrowser> TaggedAssetBrowserWidgetWeak;
+
+	/** The currently active section, updated via OnActiveSectionChanged. */
+	TWeakObjectPtr<const UTaggedAssetBrowserSection> CurrentActiveSection;
+
+	/** Container for the add/import button(s). Rebuilt reactively on section change. */
+	TSharedPtr<SHorizontalBox> AddImportButtonContainer;
+
 	/** Current Tab host type, so we can resturucture layout based on how this tab is being hosted (Sidebar / statusbar / etc). */
 	ETabState::Type TabState = ETabState::ClosedTab;
 
diff --git a/Engine/Plugins/Experimental/UAF/UAF/Source/UAFEditor/Private/Common/UAFBrowserMenuContext.h b/Engine/Plugins/Experimental/UAF/UAF/Source/UAFEditor/Private/Common/UAFBrowserMenuContext.h
new file mode 100644
--- /dev/null
+++ b/Engine/Plugins/Experimental/UAF/UAF/Source/UAFEditor/Private/Common/UAFBrowserMenuContext.h
@@ -0,0 +1,18 @@
+// Copyright Epic Games, Inc. All Rights Reserved.
+
+#pragma once
+
+#include "UAFBrowserMenuContext.generated.h"
+
+namespace UE::UAF::Editor { class SUAFBrowser; }
+
+/** Context object passed to per-instance UAF Browser add-new ToolMenu sections. */
+UCLASS()
+class UUAFBrowserMenuContext : public UObject
+{
+	GENERATED_BODY()
+
+public:
+	/** Browser that owns this context object. */
+	TWeakPtr<UE::UAF::Editor::SUAFBrowser> WeakOwningBrowser;
+};
diff --git a/Engine/Plugins/Experimental/UAF/UAF/Source/UAFEditor/UAFEditor.Build.cs b/Engine/Plugins/Experimental/UAF/UAF/Source/UAFEditor/UAFEditor.Build.cs
--- a/Engine/Plugins/Experimental/UAF/UAF/Source/UAFEditor/UAFEditor.Build.cs
+++ b/Engine/Plugins/Experimental/UAF/UAF/Source/UAFEditor/UAFEditor.Build.cs
@@ -36,6 +36,8 @@ namespace UnrealBuildTool.Rules
 					"StructUtilsEditor",
 					"BlueprintGraph",	// For K2 Schema
 					"DesktopWidgets",
+					"ClassViewer",
+					"ContentBrowserAssetDataSource",
 					"ContentBrowserFileDataSource",
 					"SubobjectEditor",
 					"Settings",
@@ -68,7 +70,7 @@ namespace UnrealBuildTool.Rules
 					"HierarchyTableRuntime",
 					"HierarchyTableEditor",
 					"HierarchyTableAnimationRuntime",
-				"PropertyBindingUtils",
+					"PropertyBindingUtils",
 				}
 			);
 			
diff --git a/Engine/Source/Editor/UserAssetTagsEditor/Private/Widgets/STaggedAssetBrowser.cpp b/Engine/Source/Editor/UserAssetTagsEditor/Private/Widgets/STaggedAssetBrowser.cpp
--- a/Engine/Source/Editor/UserAssetTagsEditor/Private/Widgets/STaggedAssetBrowser.cpp
+++ b/Engine/Source/Editor/UserAssetTagsEditor/Private/Widgets/STaggedAssetBrowser.cpp
@@ -46,6 +46,8 @@ void STaggedAssetBrowser::Construct(const FArguments& InArgs, const UTaggedAsset
 	OnGetExtraFrontendFiltersDelegate = InArgs._OnGetExtraFrontendFilters;
 	OnExtendAddFilterMenuDelegate = InArgs._OnExtendAddFilterMenu;
 	OnOverrideAssetPickerConfigDelegate = InArgs._OnOverrideAssetPickerConfig;
+	OnActiveSectionChangedDelegate = InArgs._OnActiveSectionChanged;
+	OnPrimaryFilterSelectionChangedDelegate = InArgs._OnPrimaryFilterSelectionChanged;
 	InterfaceOverrideProfiles = InArgs._InterfaceOverrideProfiles;
 	OnGenerateAssetTooltipWidgetDelegate = InArgs._OnGenerateAssetTooltipWidget;
 	OnIsAssetValidForCustomToolTipDelegate = InArgs._OnIsAssetValidForCustomToolTip;
@@ -696,7 +698,7 @@ void STaggedAssetBrowser::InitializeFilters()
 	FTaggedAssetBrowserContext FilterContext;
 	FilterContext.TaggedAssetBrowser = SharedThis(this);
 	FilterContext.OnGetActiveSectionDelegate = FTaggedAssetBrowserContext::FOnGetActiveSection::CreateSP(this, &STaggedAssetBrowser::GetActiveSection);
-	FilterContext.OnGetSelectedFilters = FTaggedAssetBrowserContext::FOnGetSelectedFilters::CreateSP(this, &STaggedAssetBrowser::GetSelectedFilters);
+	FilterContext.OnGetSelectedFilters = FTaggedAssetBrowserContext::FOnGetSelectedFilters::CreateSP(this, &STaggedAssetBrowser::GetSelectedPrimaryFilters);
 	
 	TArray<UTaggedAssetBrowserFilterBase*> AllFilters;
 	FilterHierarchyRoot->GetChildrenOfType(AllFilters, true);
@@ -720,14 +722,6 @@ void STaggedAssetBrowser::InitializeFilters()
 	}
 }
 
-TArray<UTaggedAssetBrowserFilterBase*> STaggedAssetBrowser::GetAllFilters() const
-{
-	TArray<UTaggedAssetBrowserFilterBase*> AllFilters;
-	FilterHierarchyRoot->GetChildrenOfType(AllFilters, true);
-
-	return AllFilters;
-}
-
 void STaggedAssetBrowser::OnFilterChanged() const
 {
 	RefreshBackendFilter();
@@ -814,6 +808,7 @@ void STaggedAssetBrowser::OnSectionSelected(const UTaggedAssetBrowserSection* In
 	
 	RefreshPrimaryFilterRootItems(true);
 	RefreshBackendFilter();
+	OnActiveSectionChangedDelegate.ExecuteIfBound(InTaggedAssetBrowserSection);
 }
 
 const UTaggedAssetBrowserSection* STaggedAssetBrowser::GetActiveSection() const
@@ -826,6 +821,41 @@ const UTaggedAssetBrowserSection* STaggedAssetBrowser::GetActiveSection() const
 	return nullptr;
 }
 
+TArray<const UTaggedAssetBrowserFilterBase*> STaggedAssetBrowser::GetSelectedPrimaryFilters() const
+{
+	TArray<const UTaggedAssetBrowserFilterBase*> Result;
+	
+	if(PrimaryFilterSelector.IsValid())
+	{
+		TArray<UHierarchyElement*> SelectedItems = PrimaryFilterSelector->GetSelectedItems();
+		Algo::TransformIf(SelectedItems, Result, [](const UHierarchyElement* Element) -> bool
+		{
+			return Element->IsA<UTaggedAssetBrowserFilterBase>();
+		},
+	[](const UHierarchyElement* Element) -> const UTaggedAssetBrowserFilterBase*
+		{
+			return CastChecked<UTaggedAssetBrowserFilterBase>(Element);
+		});		
+	}
+
+	return Result;
+}
+
+TArray<const UTaggedAssetBrowserFilterBase*> STaggedAssetBrowser::GetAllPrimaryFilters() const
+{
+	TArray<UTaggedAssetBrowserFilterBase*> AllFilters;
+	FilterHierarchyRoot->GetChildrenOfType(AllFilters, true);
+
+	// Second array for const
+	TArray<const UTaggedAssetBrowserFilterBase*> Result;
+	Result.Reserve(AllFilters.Num());
+	for (UTaggedAssetBrowserFilterBase* Filter : AllFilters)
+	{
+		Result.Add(Filter);
+	}
+	return Result;
+}
+
 void STaggedAssetBrowser::GenerateSearchItems(TWeakObjectPtr<UHierarchyElement> Root, TArray<TWeakObjectPtr<UHierarchyElement>> ParentChain, TArray<FFilterPath>& OutSearchItems) const
 {
 	TConstArrayView<TObjectPtr<UHierarchyElement>> FilteredChildren = Root->GetChildren();
@@ -1048,28 +1078,9 @@ void STaggedAssetBrowser::OnPrimaryFilterSelected(TObjectPtr<UHierarchyElement>
 	{
 		LastSelectedPrimaryFilterIdentifierFallback = Cast<UTaggedAssetBrowserFilterBase>(Filter)->GetIdentifier();
 	}
-	
-	SaveSettings();
-}
-
-TArray<const UTaggedAssetBrowserFilterBase*> STaggedAssetBrowser::GetSelectedFilters() const
-{
-	TArray<const UTaggedAssetBrowserFilterBase*> Result;
-	
-	if(PrimaryFilterSelector.IsValid())
-	{
-		TArray<UHierarchyElement*> SelectedItems = PrimaryFilterSelector->GetSelectedItems();
-		Algo::TransformIf(SelectedItems, Result, [](const UHierarchyElement* Element) -> bool
-		{
-			return Element->IsA<UTaggedAssetBrowserFilterBase>();
-		},
-	[](const UHierarchyElement* Element) -> const UTaggedAssetBrowserFilterBase*
-		{
-			return CastChecked<UTaggedAssetBrowserFilterBase>(Element);
-		});		
-	}
 
-	return Result;
+	SaveSettings();
+	OnPrimaryFilterSelectionChangedDelegate.ExecuteIfBound();
 }
 
 TArray<TSharedRef<FFrontendFilter>> STaggedAssetBrowser::OnGetExtraFrontendFilters() const
diff --git a/Engine/Source/Editor/UserAssetTagsEditor/Public/Widgets/STaggedAssetBrowser.h b/Engine/Source/Editor/UserAssetTagsEditor/Public/Widgets/STaggedAssetBrowser.h
--- a/Engine/Source/Editor/UserAssetTagsEditor/Public/Widgets/STaggedAssetBrowser.h
+++ b/Engine/Source/Editor/UserAssetTagsEditor/Public/Widgets/STaggedAssetBrowser.h
@@ -54,6 +54,8 @@ class STaggedAssetBrowser : public SCompoundWidget
 public:
 	DECLARE_DELEGATE_RetVal(TArray<TSharedRef<FFrontendFilter>>, FOnGetExtraFrontendFilters)
 	DECLARE_DELEGATE_OneParam(FOnOverrideAssetPickerConfig, FAssetPickerConfig&)
+	DECLARE_DELEGATE_OneParam(FOnActiveSectionChanged, const UTaggedAssetBrowserSection*)
+	DECLARE_DELEGATE(FOnPrimaryFilterSelectionChanged)
 
 	/** Contains override names for specific menu names. By default, the ProfileName is used as suffix for all menus. */
 	struct FInterfaceOverrideProfiles
@@ -116,6 +118,12 @@ public:
 		
 		/** An additional widget to attach below the asset browser. */
 		SLATE_ARGUMENT(TSharedPtr<SWidget>, AdditionalBottomWidget)
+
+		/** Called when the active section changes. Fired after internal state has been updated. */
+		SLATE_EVENT(FOnActiveSectionChanged, OnActiveSectionChanged)
+
+		/** Called when the primary filter tree selection changes. */
+		SLATE_EVENT(FOnPrimaryFilterSelectionChanged, OnPrimaryFilterSelectionChanged)
 	SLATE_END_ARGS()
 
 	
@@ -132,6 +140,15 @@ public:
 
 	UE_API TAttribute<const FMainMRUFavoritesList*> GetFavoritesListAttribute() const;
 
+	/** Returns the currently active section, or nullptr if no section is active. */
+	UE_API const UTaggedAssetBrowserSection* GetActiveSection() const;
+
+	/** Returns the currently selected primary filter items. */
+	UE_API TArray<const UTaggedAssetBrowserFilterBase*> GetSelectedPrimaryFilters() const;
+
+	/** Returns all primary filter items in the current section, regardless of selection. */
+	UE_API TArray<const UTaggedAssetBrowserFilterBase*> GetAllPrimaryFilters() const;
+
 private:
 	void InitializeValues();
 
@@ -160,7 +177,6 @@ private:
 
 	/** We initialize the duplicated filters contained in the filter root of the asset so that the filters have context available. */
 	void InitializeFilters();
-	TArray<UTaggedAssetBrowserFilterBase*> GetAllFilters() const;
 
 private:
 	/** This has to be called whenever any kind of filtering changes and refreshes the actual asset view */
@@ -200,7 +216,6 @@ private:
 
 	/** Sections functionality */
 	void OnSectionSelected(const UTaggedAssetBrowserSection* InTaggedAssetBrowserSection);
-	const UTaggedAssetBrowserSection* GetActiveSection() const;
 
 	void OnGetChildFiltersForFilter(TObjectPtr<UHierarchyElement> PrimaryFilter, TArray<TObjectPtr<UHierarchyElement>>& OutChildren) const;
 	bool OnCompareFiltersForEquality(const UTaggedAssetBrowserFilterBase& PrimaryFilterA, const UTaggedAssetBrowserFilterBase& PrimaryFilterB) const;
@@ -212,8 +227,6 @@ private:
 	TSharedPtr<SWidget> OnGetAssetContextMenu(const TArray<FAssetData>& AssetData) const;
 	TSharedRef<SToolTip> OnGetCustomAssetTooltip(FAssetData& AssetData);
 	void OnPrimaryFilterSelected(TObjectPtr<UHierarchyElement> Filter, ESelectInfo::Type Arg);
-	
-	TArray<const UTaggedAssetBrowserFilterBase*> GetSelectedFilters() const;
 
 	TArray<TSharedRef<FFrontendFilter>> OnGetExtraFrontendFilters() const;
 	
@@ -243,6 +256,8 @@ private:
 	FOnExtendAddFilterMenu OnExtendAddFilterMenuDelegate;
 	FOnShouldFilterAsset OnAdditionalShouldFilterAssetDelegate;
 	FOnOverrideAssetPickerConfig OnOverrideAssetPickerConfigDelegate;
+	FOnActiveSectionChanged OnActiveSectionChangedDelegate;
+	FOnPrimaryFilterSelectionChanged OnPrimaryFilterSelectionChangedDelegate;
 	ESelectionMode::Type AssetSelectionMode = ESelectionMode::Single;
 	TSharedPtr<IPropertyHandle> PropertyHandle;
 	TArray<FAssetData> AdditionalReferencingAssets;
