Remove old AnimNode entry from RunAnimNode RigUnit as it is no longer supported
Move AnimNode template to the correct plugin

#virtualized
#Synced-CL 52402383

diff --git a/Engine/Plugins/Experimental/UAF/UAF/Source/UAFEditor/Private/AnimNextEditorModule.cpp b/Engine/Plugins/Experimental/UAF/UAF/Source/UAFEditor/Private/AnimNextEditorModule.cpp
--- a/Engine/Plugins/Experimental/UAF/UAF/Source/UAFEditor/Private/AnimNextEditorModule.cpp
+++ b/Engine/Plugins/Experimental/UAF/UAF/Source/UAFEditor/Private/AnimNextEditorModule.cpp
@@ -110,12 +110,11 @@ namespace UE::UAF::Editor
 {
 	const auto ToggleUAFContentHidden = FConsoleCommandDelegate::CreateStatic([]()
     {
-    	static const TCHAR* DenyItemKey = TEXT("UAFContent");
-    
     	if (TSharedPtr<IPlugin> UAFPlugin = IPluginManager::Get().FindPlugin("UAF"))
     	{
     		TSharedRef<FPathPermissionList>& FolderPermissions = FAssetToolsModule::GetModule().Get().GetFolderPermissionList();
 
+			static const TCHAR* DenyItemKey = TEXT("UAFContent");
 			const FString DenyItemValue = UAFPlugin->GetMountedAssetPath();
 
     		if (FolderPermissions->ContainsDenyListItem(DenyItemValue))
@@ -127,6 +126,23 @@ namespace UE::UAF::Editor
 		    	FolderPermissions->AddDenyListItem(DenyItemKey, DenyItemValue);
 		    }
     	}
+
+		if (TSharedPtr<IPlugin> UAFPlugin = IPluginManager::Get().FindPlugin("UAFAnimNode"))
+		{
+			TSharedRef<FPathPermissionList>& FolderPermissions = FAssetToolsModule::GetModule().Get().GetFolderPermissionList();
+
+			static const TCHAR* DenyItemKey = TEXT("UAFAnimNodeContent");
+			const FString DenyItemValue = UAFPlugin->GetMountedAssetPath();
+
+			if (FolderPermissions->ContainsDenyListItem(DenyItemValue))
+			{
+				FolderPermissions->RemoveDenyListItem(DenyItemKey, DenyItemValue);
+			}
+			else
+			{
+				FolderPermissions->AddDenyListItem(DenyItemKey, DenyItemValue);
+			}
+		}
     });
 
 static FAutoConsoleCommand ToggleHiddenUAFContent(
diff --git a/Engine/Plugins/Experimental/UAF/UAFAnimNode/Source/UAFAnimNode/Internal/AnimNode/RigUnit_RunAnimNode.h b/Engine/Plugins/Experimental/UAF/UAFAnimNode/Source/UAFAnimNode/Internal/AnimNode/RigUnit_RunAnimNode.h
--- a/Engine/Plugins/Experimental/UAF/UAFAnimNode/Source/UAFAnimNode/Internal/AnimNode/RigUnit_RunAnimNode.h
+++ b/Engine/Plugins/Experimental/UAF/UAFAnimNode/Source/UAFAnimNode/Internal/AnimNode/RigUnit_RunAnimNode.h
@@ -45,13 +45,9 @@ struct FRigUnit_RunAnimNode : public FRigUnit_AnimNextBase
 	RIGVM_METHOD()
 	UE_API void Execute();
 
-	// The AnimNode to run (DO NOT USE, RETAINING IT TEMPORARILY)
-	UPROPERTY(EditAnywhere, Category = "Run", meta = (Input))
-	FUAFAnimNodeReference AnimNode;
-
 	// The anim node to run
 	UPROPERTY(EditAnywhere, Category = "Run", meta = (Input))
-	UE::UAF::FUAFAnimNodeDataEx RootAnimNode;
+	UE::UAF::FUAFAnimNodeDataEx AnimNode;
 
 	// LOD to run at. If this is -1 then the reference pose's source LOD will be used
 	UPROPERTY(EditAnywhere, Category = "Run", meta = (Input))
diff --git a/Engine/Plugins/Experimental/UAF/UAFAnimNode/Source/UAFAnimNode/Private/RigUnit_RunAnimNode.cpp b/Engine/Plugins/Experimental/UAF/UAFAnimNode/Source/UAFAnimNode/Private/RigUnit_RunAnimNode.cpp
--- a/Engine/Plugins/Experimental/UAF/UAFAnimNode/Source/UAFAnimNode/Private/RigUnit_RunAnimNode.cpp
+++ b/Engine/Plugins/Experimental/UAF/UAFAnimNode/Source/UAFAnimNode/Private/RigUnit_RunAnimNode.cpp
@@ -41,7 +41,7 @@ FRigUnit_RunAnimNode_Execute()
 		return;
 	}
 	
-	if (!RootAnimNode.IsValid())
+	if (!AnimNode.IsValid())
 	{
 		UAF_RIGUNIT_LOG(Warning, TEXT("Could not run Anim Node - No Anim Node Selected]")); 
 		return;
@@ -91,7 +91,7 @@ FRigUnit_RunAnimNode_Execute()
 	if (WorkData.AnimNodeInstance == nullptr)
 	{
 		// allocate AnimNode instance
-		WorkData.AnimNodeInstance = RootAnimNode.Get()->CreateInstance(Context);
+		WorkData.AnimNodeInstance = AnimNode.Get()->CreateInstance(Context);
 	}
 
 	TRACE_UAF_MODULE(ModuleInstance);
diff --git a/Engine/Plugins/Experimental/UAF/UAFAnimNode/UAFAnimNode.uplugin b/Engine/Plugins/Experimental/UAF/UAFAnimNode/UAFAnimNode.uplugin
--- a/Engine/Plugins/Experimental/UAF/UAFAnimNode/UAFAnimNode.uplugin
+++ b/Engine/Plugins/Experimental/UAF/UAFAnimNode/UAFAnimNode.uplugin
@@ -11,12 +11,11 @@
 	"MarketplaceURL" : "",
 	"SupportURL" : "",
 	"EnabledByDefault" : false,
-	"CanContainContent" : false,
+	"CanContainContent" : true,
 	"IsBetaVersion": false,
 	"IsExperimentalVersion": true,
 	"Installed" : false,
-	"Plugins":
-	[
+	"Plugins": [
 		{
 			"Name": "UAF",
 			"Enabled": true
@@ -28,6 +27,10 @@
 		{
 			"Name": "ControlRig",
 			"Enabled": true
+		},
+		{
+			"Name": "RigVM",
+			"Enabled": true
 		}
 	],
 	"Modules" :
