Add SkeletalMeshComponent as input source for UAF
- Support SkeletalMeshComponent as input for UAF evaluation via new VirtualValueBundle_ComponentOutput
- Add BlueprintSetInputBinding UFUNCTION for dynamic input binding at runtime
- Add K2Node_UAFComponentSetDynamicInputBinding blueprint node, restricted to FValueBundle variables
- Extract ParseVariableReferenceFromPin and PreloadVariableReferenceAssets helpers into FUtils
- Consolidate shared binding structs into UE::UAF::Private namespace

#Synced-CL 52466902

diff --git a/Engine/Plugins/Experimental/UAF/UAF/Source/UAF/Private/Component/AnimNextComponent.cpp b/Engine/Plugins/Experimental/UAF/UAF/Source/UAF/Private/Component/AnimNextComponent.cpp
--- a/Engine/Plugins/Experimental/UAF/UAF/Source/UAF/Private/Component/AnimNextComponent.cpp
+++ b/Engine/Plugins/Experimental/UAF/UAF/Source/UAF/Private/Component/AnimNextComponent.cpp
@@ -3,6 +3,7 @@
 #include "Component/AnimNextComponent.h"
 
 #include "Blueprint/BlueprintExceptionInfo.h"
+#include "Engine/SkeletalMesh.h"
 #include "Engine/World.h"
 #include "Graph/RigUnit_AnimNextWriteSkeletalMeshComponentPose.h"
 #include "Graph/UAFSystemOutputComponent.h"
@@ -15,6 +16,7 @@
 #include "Module/UAFWeakSystemReference.h"
 #include "UAF/UAFAssetFactory.h"
 #include "UAF/ValueRuntime/ValueBundle.h"
+#include "UAF/ValueRuntime/VirtualValueBundle_ComponentOutput.h"
 #include "UAF/ValueRuntime/VirtualValueBundle_SystemOutput.h"
 
 #include UE_INLINE_GENERATED_CPP_BY_NAME(AnimNextComponent)
@@ -657,6 +659,158 @@ void UUAFComponent::SetAsset(TInstancedStruct<FUAFSystemFactoryAsset>&& InAssetD
 	}
 }
 
+namespace UE::UAF::Private
+{
+struct FSystemInputBinding
+{
+	FAnimNextVariableReference Variable;
+	TWeakObjectPtr<UUAFComponent> Component;
+};
+
+struct FComponentInputBinding
+{
+	FAnimNextVariableReference Variable;
+	TWeakObjectPtr<USkeletalMeshComponent> Component;
+};
+} // namespace UE::UAF::Private
+
+void UUAFComponent::BlueprintSetInputBinding(const FAnimNextVariableReference& InVariable, UActorComponent* InComponent)
+{
+	using namespace UE::UAF::Private;
+
+	check(IsInGameThread());
+	AActor* Owner = GetOwner();
+	check(Owner);
+	
+	if (InVariable.IsValid())
+	{
+		if (SystemReference.IsValid())
+		{
+			// Fetch the current value
+			FUAFValueBundle ExistingBundleValue;
+			SystemReference.ReadVariable(InVariable,  FAnimNextParamType::GetType<FUAFValueBundle>(),  [&](TConstArrayView<uint8> VariableValue)
+			{
+				check(sizeof(ExistingBundleValue) == VariableValue.Num());
+				FUAFValueBundle::StaticStruct()->CopyScriptStruct(&ExistingBundleValue, VariableValue.GetData());
+			});
+			
+			// Handle any previously bound inputs
+			TWeakObjectPtr<UActorComponent>* PreviouslyBoundActorComponentPtr = InputToComponentMappings.Find(InVariable);
+			if (PreviouslyBoundActorComponentPtr)
+			{
+				// Cache old value
+				const TWeakObjectPtr<UActorComponent> PreviouslyBoundActorComponent = *PreviouslyBoundActorComponentPtr;
+				
+				// Update cached binding component (if valid)
+				if (InComponent)
+				{
+					*PreviouslyBoundActorComponentPtr = InComponent;
+				}
+				else
+				{
+					InputToComponentMappings.Remove(InVariable);
+				}				
+				
+				TArray<TWeakObjectPtr<UActorComponent>> ValueArray;
+				InputToComponentMappings.GenerateValueArray(ValueArray);
+				
+                if (!ValueArray.Contains(PreviouslyBoundActorComponent))
+                {
+                    // There are _no_ other bindings against this actor component, so we should remove tick dependencies accordingly
+                    if (UActorComponent* Component = PreviouslyBoundActorComponent.Get())
+                    {
+	                    if (UUAFComponent* UAFComponent = Cast<UUAFComponent>(Component))
+	                    {
+		                    const UE::UAF::FSystemReference& OtherSystemRef = UAFComponent->SystemReference;
+		                    RemoveModuleEventPrerequisite(SystemReference.GetFirstUserEventName(), UAFComponent, OtherSystemRef.GetLastUserEventName());
+	                    }
+	                    else if (USkeletalMeshComponent* SkelMeshComponent = Cast<USkeletalMeshComponent>(Component))
+	                    {
+		                    RemoveComponentPrerequisite(SkelMeshComponent, SystemReference.GetFirstUserEventName());
+	                    }
+                    }
+                }
+			}
+			
+			if (UUAFComponent* UAFComponent = Cast<UUAFComponent>(InComponent))
+			{
+				const FSystemInputBinding SystemInputBinding = { .Variable = InVariable, .Component = UAFComponent };
+				
+				InputToComponentMappings.Add(InVariable, UAFComponent);
+				
+				// Allocate the prerequisite system in case it has not been registered yet
+				SystemInputBinding.Component->AllocateSystem();
+
+				UE::UAF::FSystemReference& InputSystem = SystemInputBinding.Component->SystemReference;
+				SystemReference.AddSystemEventDependency(SystemReference.GetFirstUserEventName(), FUAFWeakSystemReference(InputSystem), InputSystem.GetLastUserEventName(), UE::UAF::ESystemDependency::Prerequisite);
+				
+				// Queue up task to run after last user event
+				SystemReference.QueueTask(SystemReference.GetLastUserEventName(), [WeakThis = TWeakObjectPtr<UUAFComponent>(this), SystemInputBinding](const UE::UAF::FModuleTaskContext& InContext)
+				{
+					if (UUAFComponent* This = WeakThis.Get())
+					{
+						if (UUAFComponent* Component = SystemInputBinding.Component.Get())
+						{
+							if (!InContext.AccessVariable<FUAFValueBundle>(SystemInputBinding.Variable, [Component](FUAFValueBundle& InValueBundle)
+							{
+								InValueBundle.SetAs<UE::UAF::FVirtualValueBundle_SystemOutput>(FUAFWeakSystemReference(Component->SystemReference));
+							}))
+							{
+								UE_LOGFMT(LogAnimation, Warning, "UUAFComponent::SetupInputOutput: Cannot bind input pose variable: '{0}'", SystemInputBinding.Variable.GetName());
+							}
+						}
+					}
+				}, UE::UAF::ETaskRunLocation::After);
+			}
+			else if (USkeletalMeshComponent* SkelMeshComponent = Cast<USkeletalMeshComponent>(InComponent))
+			{
+				const FComponentInputBinding ComponentInputBinding = { .Variable = InVariable, .Component = SkelMeshComponent };
+				AddComponentPrerequisite(SkelMeshComponent, SystemReference.GetFirstUserEventName());
+				
+				InputToComponentMappings.Add(InVariable, SkelMeshComponent);
+
+				SystemReference.QueueTask([WeakThis = TWeakObjectPtr<UUAFComponent>(this), ComponentInputBinding](const UE::UAF::FModuleTaskContext& InContext)
+				{
+					if (UUAFComponent* This = WeakThis.Get())
+					{
+						if (USkeletalMeshComponent* Component = ComponentInputBinding.Component.Get())
+						{
+							if (!InContext.AccessVariable<FUAFValueBundle>(ComponentInputBinding.Variable, [Component, This](FUAFValueBundle& InValueBundle)
+							{
+								InValueBundle.SetAs<UE::UAF::FVirtualValueBundle_ComponentOutput>(TWeakObjectPtr<USkeletalMeshComponent>(Component));
+							}))
+							{
+								UE_LOGFMT(LogAnimation, Warning, "UUAFComponent::SetupInputOutput: Cannot bind input component pose variable: '{0}'", ComponentInputBinding.Variable.GetName());
+							}
+						}
+					}
+				});
+			}
+			else if (InComponent)
+			{
+				ensureAlwaysMsgf(false, TEXT("UUAFComponent::BlueprintSetInputBinding, Unexpected component type %s"), InComponent->GetClass() ? *InComponent->GetClass()->GetName() : TEXT("None"));
+			}
+			else
+			{
+				// Reset to ref-pose if null component
+				SystemReference.QueueTask([WeakThis = TWeakObjectPtr<UUAFComponent>(this), Variable=InVariable](const UE::UAF::FModuleTaskContext& InContext)
+				{
+					if (UUAFComponent* This = WeakThis.Get())
+					{
+						if (!InContext.AccessVariable<FUAFValueBundle>(Variable, [This](FUAFValueBundle& InValueBundle)
+						{
+							InValueBundle.SetAsRefPose();
+						}))
+						{
+							UE_LOGFMT(LogAnimation, Warning, "UUAFComponent::SetupInputOutput: Failed to reset value bundle to ref-pose variable: '{0}'", Variable.GetName());
+						}
+					}
+				});
+			}
+		}
+	}
+}
+
 void UUAFComponent::SetAssetFromObject(const UObject* InObject)
 {
 	TInstancedStruct<FUAFSystemFactoryAsset> AssetDataFromFactory = UE::UAF::FAssetDataFactory::CreateUAFAssetDataFromObject<FUAFSystemFactoryAsset>(InObject);
@@ -699,6 +853,8 @@ bool UUAFComponent::IsModuleValid()
 
 void UUAFComponent::SetupInputOutput()
 {
+	using namespace UE::UAF::Private;
+
 	check(IsInGameThread());
 
 	CachedSkeletalMeshComponent = DetermineSkeletalMeshComponent();
@@ -708,16 +864,10 @@ void UUAFComponent::SetupInputOutput()
 		UE_LOGFMT(LogAnimation, Warning, "UUAFComponent::SetupInputOutput: Could not find a valid mesh component to use for system output");
 	}
 
-	struct FInputBinding
-	{
-		FAnimNextVariableReference Variable;
-		TWeakObjectPtr<UUAFComponent> Component;
-	};
-
 	AActor* Owner = GetOwner();
 	check(Owner);
-	TArray<FInputBinding> InputBindings;
-	InputBindings.Reserve(Inputs.Num());
+	TArray<FSystemInputBinding> SystemInputBindings;
+	SystemInputBindings.Reserve(Inputs.Num());
 	for (const FUAFComponentInputDesc& InputDesc : Inputs)
 	{
 		if (InputDesc.Input.IsValid())
@@ -730,7 +880,8 @@ void UUAFComponent::SetupInputOutput()
 					continue;
 				}
 
-				InputBindings.Add({ .Variable = InputDesc.Input, .Component = InputComponent });
+				SystemInputBindings.Add({ .Variable = InputDesc.Input, .Component = InputComponent });
+				InputToComponentMappings.Add(InputDesc.Input, InputComponent);
 			}
 			else
 			{
@@ -746,7 +897,7 @@ void UUAFComponent::SetupInputOutput()
 	// Set up tick dependencies for valid inputs
 	// By default we link [Input Last Event] -> [Our First Event].
 	// TODO: Allow user-specified dependencies here
-	for (const FInputBinding& InputBinding : InputBindings)
+	for (const FSystemInputBinding& InputBinding : SystemInputBindings)
 	{
 		// Allocate the prerequisite system in case it has not been registered yet
 		InputBinding.Component->AllocateSystem();
@@ -761,7 +912,32 @@ void UUAFComponent::SetupInputOutput()
 		AddComponentSubsequent(CachedSkeletalMeshComponent, SystemReference.GetLastUserEventName());
 	}
 
-	SystemReference.QueueTask([WeakThis = TWeakObjectPtr<UUAFComponent>(this), bValidOutputComponent, InputBindings = MoveTemp(InputBindings)](const UE::UAF::FModuleTaskContext& InContext)
+	TArray<FComponentInputBinding> ComponentInputBindings;
+	ComponentInputBindings.Reserve(Inputs.Num());
+	
+	for (const FUAFComponentInputDesc& InputDesc : Inputs)
+	{
+		if (InputDesc.Input.IsValid())
+		{
+			if (USkeletalMeshComponent* InputComponent = Cast<USkeletalMeshComponent>(InputDesc.Component.GetComponent(Owner)))
+			{
+				ComponentInputBindings.Add({ .Variable = InputDesc.Input, .Component = InputComponent });
+				InputToComponentMappings.Add(InputDesc.Input, InputComponent);
+			}
+		}
+		else
+		{
+			UE_LOGFMT(LogAnimation, Warning, "UUAFComponent::SetupInputOutput: Invalid input pose variable specified");
+		}
+	}
+	
+	for (const FComponentInputBinding& InputBinding : ComponentInputBindings)
+	{
+		USkeletalMeshComponent* Component = InputBinding.Component.Get();
+		AddComponentPrerequisite(Component, SystemReference.GetFirstUserEventName());
+	}
+
+	SystemReference.QueueTask([WeakThis = TWeakObjectPtr<UUAFComponent>(this), bValidOutputComponent, SystemInputBindings = MoveTemp(SystemInputBindings), ComponentInputBindings = MoveTemp(ComponentInputBindings)](const UE::UAF::FModuleTaskContext& InContext)
 	{
 		if (UUAFComponent* This = WeakThis.Get())
 		{
@@ -773,7 +949,7 @@ void UUAFComponent::SetupInputOutput()
 				});
 			}
 
-			for (const FInputBinding& InputBinding : InputBindings)
+			for (const FSystemInputBinding& InputBinding : SystemInputBindings)
 			{
 				if (UUAFComponent* Component = InputBinding.Component.Get())
 				{
@@ -782,7 +958,21 @@ void UUAFComponent::SetupInputOutput()
 						InValueBundle.SetAs<UE::UAF::FVirtualValueBundle_SystemOutput>(FUAFWeakSystemReference(Component->SystemReference));
 					}))
 					{
-						UE_LOGFMT(LogAnimation, Warning, "UUAFComponent::SetupInputOutput: Cannot bind input pose variable: '{0}'", InputBinding.Variable.GetName());
+						UE_LOGFMT(LogAnimation, Warning, "UUAFComponent::SetupInputOutput: Cannot bind system input pose variable: '{0}'", InputBinding.Variable.GetName());
+					}
+				}
+			}
+
+			for (const FComponentInputBinding& InputBinding : ComponentInputBindings)
+			{
+				if (USkeletalMeshComponent* Component = InputBinding.Component.Get())
+				{
+					if (!InContext.AccessVariable<FUAFValueBundle>(InputBinding.Variable, [Component, This](FUAFValueBundle& InValueBundle)
+					{
+						InValueBundle.SetAs<UE::UAF::FVirtualValueBundle_ComponentOutput>(TWeakObjectPtr<USkeletalMeshComponent>(Component));
+					}))
+					{
+						UE_LOGFMT(LogAnimation, Warning, "UUAFComponent::SetupInputOutput: Cannot bind component input pose variable: '{0}'", InputBinding.Variable.GetName());
 					}
 				}
 			}
diff --git a/Engine/Plugins/Experimental/UAF/UAF/Source/UAF/Private/UAF/ValueRuntime/VirtualValueBundle_ComponentOutput.cpp b/Engine/Plugins/Experimental/UAF/UAF/Source/UAF/Private/UAF/ValueRuntime/VirtualValueBundle_ComponentOutput.cpp
new file mode 100644
--- /dev/null
+++ b/Engine/Plugins/Experimental/UAF/UAF/Source/UAF/Private/UAF/ValueRuntime/VirtualValueBundle_ComponentOutput.cpp
@@ -0,0 +1,100 @@
+// Copyright Epic Games, Inc. All Rights Reserved.
+
+#include "VirtualValueBundle_ComponentOutput.h"
+
+#include "Engine/SkeletalMesh.h"
+#include "Graph/UAFSystemOutputComponent.h"
+
+namespace UE::UAF
+{
+
+void ReadPose(USkeletalMeshComponent* InputComponent, FAnimNextGraphLODPose& OutPose)
+{
+	QUICK_SCOPE_CYCLE_COUNTER(STAT_UAF_READ_SMC_Pose);
+
+	using namespace UE::UAF;
+	FTaskTagScope TaskTag(ETaskTag::EParallelGameThread);
+
+	if(InputComponent == nullptr)
+	{
+		UE_LOG(LogAnimation, Warning, TEXT("Could not read from skeletal mesh component - No valid skeletal mesh component found."));
+		OutPose.LODPose.SetRefPose();
+		return;
+	}
+
+	USkeletalMesh* SkeletalMesh = InputComponent->GetSkeletalMeshAsset();
+	if(SkeletalMesh == nullptr)
+	{
+		OutPose.LODPose.SetRefPose();
+		return;
+	}
+
+	FDataHandle RefPoseHandle = FDataRegistry::Get()->GetOrGenerateReferencePose(InputComponent);
+	const FReferencePose& RefPose = RefPoseHandle.GetRef<UE::UAF::FReferencePose>();
+	OutPose.LODPose.PrepareForLOD(RefPose, InputComponent->GetPredictedLODLevel());
+
+	// In case animation is enabled, use the (local) bone space transforms - otherwise (if driven by UAF) use component space transforms instead
+	const bool bUseLocalSpaceTransforms = InputComponent->bEnableAnimation;
+	const TArrayView<const FBoneIndexType> LODBoneToMeshBone = OutPose.LODPose.GetLODBoneIndexToMeshBoneIndexMap();
+	const TArray<FBoneIndexType>& MeshBoneIndices = InputComponent->FillComponentSpaceTransformsRequiredBones;
+	const int32 NumBoneTransforms = MeshBoneIndices.Num();
+	const int32 NumPoseTransforms = OutPose.LODPose.LocalTransformsView.Num();
+	ensure(NumBoneTransforms >= NumPoseTransforms);
+
+	if (bUseLocalSpaceTransforms)
+	{
+		TArrayView<const FTransform> LocalSpaceTransforms = InputComponent->GetBoneSpaceTransformsView();
+
+		//Set root bone which is always evaluated.
+		OutPose.LODPose.LocalTransforms[0] = LocalSpaceTransforms[0];
+		for (int32 PoseBoneIdx = 1; PoseBoneIdx < NumPoseTransforms; ++PoseBoneIdx)
+		{
+			const FBoneIndexType MeshIndex = LODBoneToMeshBone[PoseBoneIdx];
+			const int32 ComponentBoneIndex = MeshBoneIndices.IndexOfByKey(MeshIndex);
+			if (ComponentBoneIndex != INDEX_NONE)
+			{
+				OutPose.LODPose.LocalTransforms[PoseBoneIdx] = LocalSpaceTransforms[MeshIndex];
+			}
+		}
+	}
+	else
+	{
+		const TArray<FTransform>& ComponentSpaceTransforms = InputComponent->GetComponentSpaceTransforms();
+
+		//Set root bone which is always evaluated.
+		OutPose.LODPose.LocalTransforms[0] = ComponentSpaceTransforms[0];
+		const TArrayView<const FBoneIndexType> MeshBoneToParentMapping = RefPose.GetMeshBoneIndexToParentMeshBoneIndexMap();
+		for (int32 PoseBoneIdx = 1; PoseBoneIdx < NumPoseTransforms; ++PoseBoneIdx)
+		{
+			const FBoneIndexType MeshIndex = LODBoneToMeshBone[PoseBoneIdx];
+			const int32 ComponentBoneIndex = MeshBoneIndices.IndexOfByKey(MeshIndex);
+			if (ComponentBoneIndex != INDEX_NONE)
+			{
+				const int32 ParentIndex = MeshBoneToParentMapping[MeshIndex];
+				const FTransform& ParentTransform = ComponentSpaceTransforms[ParentIndex];
+				const FTransform& ChildTransform = ComponentSpaceTransforms[MeshIndex];
+				// Convert component to local space
+				OutPose.LODPose.LocalTransforms[PoseBoneIdx] = ChildTransform.GetRelativeTransform(ParentTransform);
+			}
+		}
+	}
+}
+
+const FAnimNextGraphLODPose* FVirtualValueBundle_ComponentOutput::GetLODPose() const
+{
+	ReadPose(WeakSMC.Get(), CachedPose);
+
+	return &CachedPose;
+}
+
+const FValueBundleHeap* FVirtualValueBundle_ComponentOutput::GetValueBundle() const
+{
+	ReadPose(WeakSMC.Get(), CachedPose);
+
+	// TODO: Implement this for value runtime
+	return &CachedValueBundle;
+}
+
+}
+
+
diff --git a/Engine/Plugins/Experimental/UAF/UAF/Source/UAF/Private/UAF/ValueRuntime/VirtualValueBundle_ComponentOutput.h b/Engine/Plugins/Experimental/UAF/UAF/Source/UAF/Private/UAF/ValueRuntime/VirtualValueBundle_ComponentOutput.h
new file mode 100644
--- /dev/null
+++ b/Engine/Plugins/Experimental/UAF/UAF/Source/UAF/Private/UAF/ValueRuntime/VirtualValueBundle_ComponentOutput.h
@@ -0,0 +1,36 @@
+// Copyright Epic Games, Inc. All Rights Reserved.
+
+#pragma once
+
+#include "RemapPoseData.h"
+#include "Graph/AnimNext_LODPose.h"
+#include "Module/UAFWeakSystemReference.h"
+#include "UAF/ValueRuntime/IVirtualValueBundle.h"
+#include "UAF/ValueRuntime/ValueBundle.h"
+
+class USkeletalMeshComponent;
+
+#define UE_API UAF_API
+
+namespace UE::UAF
+{
+
+struct FVirtualValueBundle_ComponentOutput : public IVirtualValueBundle
+{
+	explicit FVirtualValueBundle_ComponentOutput(TWeakObjectPtr<USkeletalMeshComponent> InWeakSMC)
+		: WeakSMC(InWeakSMC)
+	{}
+
+private:
+	// IVirtualValueBundle interface
+	virtual const FAnimNextGraphLODPose* GetLODPose() const override;
+	virtual const FValueBundleHeap* GetValueBundle() const override;
+
+	mutable FAnimNextGraphLODPose CachedPose;
+	mutable FValueBundleHeap CachedValueBundle;
+	TWeakObjectPtr<USkeletalMeshComponent> WeakSMC;
+};
+
+}
+
+#undef UE_API
diff --git a/Engine/Plugins/Experimental/UAF/UAF/Source/UAF/Public/Component/AnimNextComponent.h b/Engine/Plugins/Experimental/UAF/UAF/Source/UAF/Public/Component/AnimNextComponent.h
--- a/Engine/Plugins/Experimental/UAF/UAF/Source/UAF/Public/Component/AnimNextComponent.h
+++ b/Engine/Plugins/Experimental/UAF/UAF/Source/UAF/Public/Component/AnimNextComponent.h
@@ -52,7 +52,7 @@ struct FUAFComponentInputDesc
 	GENERATED_BODY()
 
 	// The component that this input uses for its input pose
-	UPROPERTY(EditAnywhere, Category="Input", meta = (UseComponentPicker, AllowedClasses = "/Script/UAF.UAFComponent"))
+	UPROPERTY(EditAnywhere, Category="Input", meta = (UseComponentPicker, AllowedClasses = "/Script/UAF.UAFComponent,/Script/Engine.SkeletalMeshComponent"))
 	FComponentReference Component;
 
 	// The input pose variable that this input maps to
@@ -113,6 +113,12 @@ public:
 	UFUNCTION(BlueprintCallable, BlueprintInternalUseOnly, Category = "UAF", DisplayName = "Get Variable", CustomThunk, meta = (CustomStructureParam = Value, AutoCreateRefTerm = "Variable, Value", UnsafeDuringActorConstruction))
 	UAF_API void BlueprintGetVariableReference(const FAnimNextVariableReference& Variable, int32& Value);
 	
+	// Sets a system's (of FValueBundle type) variable.
+	// @param    Variable The variable to set
+	// @param    Component    The component to use as input source
+	UFUNCTION(BlueprintCallable, BlueprintInternalUseOnly, Category = "UAF", DisplayName = "Set Input Binding", meta = (AutoCreateRefTerm = "Variable, Component", UnsafeDuringActorConstruction))
+	UAF_API void BlueprintSetInputBinding(const FAnimNextVariableReference& Variable, UActorComponent* Component);
+	
 	// Set a variable.
 	// @param    InVariable         The variable to set
 	// @param    InValue            The value to set the variable to
@@ -333,6 +339,10 @@ private:
 
 	// Owning reference to the registered system
 	UE::UAF::FSystemReference SystemReference;
+	
+	// Cached mapping between publicly exposed FValueBundle variables and their currently bound actor-components
+	UPROPERTY(transient)
+	TMap<FAnimNextVariableReference, TWeakObjectPtr<UActorComponent>> InputToComponentMappings;
 
 	// Serial number, updated on worker threads when output is written
 	uint32 SerialNumber = 0;
diff --git a/Engine/Plugins/Experimental/UAF/UAF/Source/UAFUncookedOnly/Internal/UncookedOnlyUtils.h b/Engine/Plugins/Experimental/UAF/UAF/Source/UAFUncookedOnly/Internal/UncookedOnlyUtils.h
--- a/Engine/Plugins/Experimental/UAF/UAF/Source/UAFUncookedOnly/Internal/UncookedOnlyUtils.h
+++ b/Engine/Plugins/Experimental/UAF/UAF/Source/UAFUncookedOnly/Internal/UncookedOnlyUtils.h
@@ -30,6 +30,7 @@ struct FRigVMCompileSettings;
 struct FRigVMGraphFunctionHeaderArray;
 struct FGraphContextMenuBuilder;
 struct FRigVMFunction;
+class UEdGraphPin;
 class UUAFSystem;
 class UUAFSystem_EditorData;
 class UAnimNextEdGraph;
@@ -242,6 +243,13 @@ struct FUtils
 	static UE_API FWorkspaceOutlinerItemExport MakeAssetReferenceExport(const FSoftObjectPath& ReferencedAssetPath, const FWorkspaceOutlinerItemExport& ParentExport, bool bShouldExpand = true);
 
 	static UE_API void GetExternalVariables(const UUAFRigVMAssetEditorData* EditorData, TArray<FRigVMExternalVariable>& OutVariables);
+
+	/** Parse a FAnimNextVariableReference from a pin's default value. Returns the parsed reference (may be None/invalid). */
+	static UE_API FAnimNextVariableReference ParseVariableReferenceFromPin(const UEdGraphPin* InPin);
+
+	/** Preload assets referenced by a variable reference pin. Call from PreloadRequiredAssets. */
+	static UE_API void PreloadVariableReferenceAssets(const UEdGraphPin* InVariablePin);
+
 private:
 	static void GetInternalVariables(const UUAFRigVMAssetEditorData* EditorData, TArray<FRigVMExternalVariable>& OutVariables);
 
diff --git a/Engine/Plugins/Experimental/UAF/UAF/Source/UAFUncookedOnly/Internal/Variables/K2Node_UAFComponentSetInputBinding.h b/Engine/Plugins/Experimental/UAF/UAF/Source/UAFUncookedOnly/Internal/Variables/K2Node_UAFComponentSetInputBinding.h
new file mode 100644
--- /dev/null
+++ b/Engine/Plugins/Experimental/UAF/UAF/Source/UAFUncookedOnly/Internal/Variables/K2Node_UAFComponentSetInputBinding.h
@@ -0,0 +1,32 @@
+// Copyright Epic Games, Inc. All Rights Reserved.
+
+#pragma once
+
+#include "K2Node_CallFunction.h"
+#include "K2Node_UAFComponentSetInputBinding.generated.h"
+
+/** A custom node that supports setting dynamic input bindings on UAF components, restricted to FValueBundle variables */
+UCLASS(MinimalAPI)
+class UK2Node_UAFComponentSetInputBinding : public UK2Node_CallFunction
+{
+	GENERATED_BODY()
+
+public:
+	UAFUNCOOKEDONLY_API UK2Node_UAFComponentSetInputBinding();
+
+private:
+	//~ Begin UEdGraphNode Interface
+	UAFUNCOOKEDONLY_API virtual bool IsConnectionDisallowed(const UEdGraphPin* MyPin, const UEdGraphPin* OtherPin, FString& OutReason) const override;
+	//~ End UEdGraphNode Interface
+
+	//~ Begin UK2Node Interface
+	UAFUNCOOKEDONLY_API virtual FString GetPinMetaData(FName InPinName, FName InKey) override;
+	UAFUNCOOKEDONLY_API virtual void GetMenuActions(FBlueprintActionDatabaseRegistrar& InActionRegistrar) const override;
+	UAFUNCOOKEDONLY_API virtual void ValidateNodeDuringCompilation(FCompilerResultsLog& MessageLog) const override;
+	UAFUNCOOKEDONLY_API virtual void PreloadRequiredAssets() override;
+	//~ End K2Node Interface
+	
+	static const FName VariablePinName;
+	static const FName ComponentPinName;
+	static const TArray<const UClass*> AllowedComponentClasses;
+};
diff --git a/Engine/Plugins/Experimental/UAF/UAF/Source/UAFUncookedOnly/Private/UncookedOnlyUtils.cpp b/Engine/Plugins/Experimental/UAF/UAF/Source/UAFUncookedOnly/Private/UncookedOnlyUtils.cpp
--- a/Engine/Plugins/Experimental/UAF/UAF/Source/UAFUncookedOnly/Private/UncookedOnlyUtils.cpp
+++ b/Engine/Plugins/Experimental/UAF/UAF/Source/UAFUncookedOnly/Private/UncookedOnlyUtils.cpp
@@ -28,6 +28,8 @@
 #include "Entries/AnimNextSharedVariablesEntry.h"
 #include "Variables/AnimNextSharedVariables.h"
 #include "Variables/AnimNextSharedVariables_EditorData.h"
+#include "Variables/AnimNextVariableReference.h"
+#include "EdGraph/EdGraphPin.h"
 #include "Logging/MessageLog.h"
 #include "RigVMFunctions/Execution/RigVMFunction_UserDefinedEvent.h"
 #include "UObject/AssetRegistryTagsContext.h"
@@ -2380,6 +2382,31 @@ FGuid FUtils::GenerateScriptStructPropertyGUID(const FProperty* Property)
 	checkf(Property->GetOwnerChecked<UStruct>(), TEXT("Expected valid owning UStruct for GUID generation"));
 	return FGuid::NewDeterministicGuid(Property->GetPathName(), FCrc::StrCrc32<TCHAR>(TEXT("StructBasedSharedVariableEntry")));
 }
+
+FAnimNextVariableReference FUtils::ParseVariableReferenceFromPin(const UEdGraphPin* InPin)
+{
+	FAnimNextVariableReference VariableReference;
+
+	if (InPin)
+	{
+		FAnimNextVariableReference::StaticStruct()->ImportText(*InPin->DefaultValue, &VariableReference, nullptr, PPF_None, nullptr, FAnimNextVariableReference::StaticStruct()->GetName());
+	}
+	
+	return VariableReference;
+}
+
+void FUtils::PreloadVariableReferenceAssets(const UEdGraphPin* InVariablePin)
+{
+	const FAnimNextVariableReference VariableReference = ParseVariableReferenceFromPin(InVariablePin);
+	if (!VariableReference.IsNone() && VariableReference.GetObject())
+	{
+		if (UObject* Object = const_cast<UObject*>(VariableReference.GetObject().Get()))
+		{
+			Object->ConditionalPreload();
+		}
+	}
+}
+
 }
 
 #undef LOCTEXT_NAMESPACE
diff --git a/Engine/Plugins/Experimental/UAF/UAF/Source/UAFUncookedOnly/Private/Variables/K2Node_AnimNextComponentSetVariable.cpp b/Engine/Plugins/Experimental/UAF/UAF/Source/UAFUncookedOnly/Private/Variables/K2Node_AnimNextComponentSetVariable.cpp
--- a/Engine/Plugins/Experimental/UAF/UAF/Source/UAFUncookedOnly/Private/Variables/K2Node_AnimNextComponentSetVariable.cpp
+++ b/Engine/Plugins/Experimental/UAF/UAF/Source/UAFUncookedOnly/Private/Variables/K2Node_AnimNextComponentSetVariable.cpp
@@ -112,13 +112,12 @@ void UK2Node_AnimNextComponentSetVariable::ValidateNodeDuringCompilation(FCompil
 	UEdGraphPin* VariablePin = SourceNode->FindPinChecked(TEXT("Variable"));
 	UEdGraphPin* ValuePin = SourceNode->FindPinChecked(TEXT("Value"));
 
-	FAnimNextVariableReference VariableReference;
-	FAnimNextVariableReference::StaticStruct()->ImportText(*VariablePin->DefaultValue, &VariableReference, nullptr, PPF_None, nullptr, FAnimNextVariableReference::StaticStruct()->GetName());
+	const FAnimNextVariableReference VariableReference = UE::UAF::UncookedOnly::FUtils::ParseVariableReferenceFromPin(VariablePin);
 	if (VariableReference.IsValid())
 	{
-		FAnimNextParamType Type = UE::UAF::UncookedOnly::FUtils::FindVariableType(VariableReference);				// The variable from the UAF Module
-		FAnimNextParamType PinType = UE::UAF::UncookedOnly::FUtils::GetParamTypeFromPinType(ValuePin->PinType);	// The variable from BP pin
-		UE::UAF::FParamCompatibility Compatibility = UE::UAF::FParamUtils::GetCompatibility(Type, PinType);
+		const FAnimNextParamType Type = UE::UAF::UncookedOnly::FUtils::FindVariableType(VariableReference);				// The variable from the UAF Module
+		const FAnimNextParamType PinType = UE::UAF::UncookedOnly::FUtils::GetParamTypeFromPinType(ValuePin->PinType);	// The variable from BP pin
+		const UE::UAF::FParamCompatibility Compatibility = UE::UAF::FParamUtils::GetCompatibility(Type, PinType);
 		if (Compatibility == UE::UAF::EParamCompatibility::Incompatible || Compatibility == UE::UAF::EParamCompatibility::Incompatible_DataLoss)
 		{
 			MessageLog.Error(*FString::Printf(TEXT("@@ mismatched type, setting '%s' type, but variable is '%s' type"), *PinType.ToString(), *Type.ToString()), this);
@@ -149,15 +148,7 @@ void UK2Node_AnimNextComponentSetVariable::GetMenuActions(FBlueprintActionDataba
 void UK2Node_AnimNextComponentSetVariable::PreloadRequiredAssets()
 {
 	Super::PreloadRequiredAssets();
-
-	UEdGraphPin* VariablePin = FindPinChecked(TEXT("Variable"));
-	FAnimNextVariableReference VariableReference;
-	FAnimNextVariableReference::StaticStruct()->ImportText(*VariablePin->DefaultValue, &VariableReference, nullptr, PPF_None, nullptr, FAnimNextVariableReference::StaticStruct()->GetName());
-	if (!VariableReference.IsNone() && VariableReference.GetObject())
-	{
-		UObject* Object = const_cast<UObject*>(VariableReference.GetObject().Get());
-		PreloadObject(Object);
-	}
+	UE::UAF::UncookedOnly::FUtils::PreloadVariableReferenceAssets(FindPinChecked(TEXT("Variable")));
 }
 
 #undef LOCTEXT_NAMESPACE
diff --git a/Engine/Plugins/Experimental/UAF/UAF/Source/UAFUncookedOnly/Private/Variables/K2Node_UAFComponentGetVariable.cpp b/Engine/Plugins/Experimental/UAF/UAF/Source/UAFUncookedOnly/Private/Variables/K2Node_UAFComponentGetVariable.cpp
--- a/Engine/Plugins/Experimental/UAF/UAF/Source/UAFUncookedOnly/Private/Variables/K2Node_UAFComponentGetVariable.cpp
+++ b/Engine/Plugins/Experimental/UAF/UAF/Source/UAFUncookedOnly/Private/Variables/K2Node_UAFComponentGetVariable.cpp
@@ -101,13 +101,12 @@ void UK2Node_UAFComponentGetVariable::ValidateNodeDuringCompilation(FCompilerRes
 	UEdGraphPin* VariablePin = SourceNode->FindPinChecked(TEXT("Variable"));
 	UEdGraphPin* ValuePin = SourceNode->FindPinChecked(TEXT("Value"));
 
-	FAnimNextVariableReference VariableReference;
-	FAnimNextVariableReference::StaticStruct()->ImportText(*VariablePin->DefaultValue, &VariableReference, nullptr, PPF_None, nullptr, FAnimNextVariableReference::StaticStruct()->GetName());
+	const FAnimNextVariableReference VariableReference = UE::UAF::UncookedOnly::FUtils::ParseVariableReferenceFromPin(VariablePin);
 	if (VariableReference.IsValid())
 	{
-		FAnimNextParamType Type = UE::UAF::UncookedOnly::FUtils::FindVariableType(VariableReference);				// The variable from the UAF Module
-		FAnimNextParamType PinType = UE::UAF::UncookedOnly::FUtils::GetParamTypeFromPinType(ValuePin->PinType);	// The variable from BP pin
-		UE::UAF::FParamCompatibility Compatibility = UE::UAF::FParamUtils::GetCompatibility(Type, PinType);
+		const FAnimNextParamType Type = UE::UAF::UncookedOnly::FUtils::FindVariableType(VariableReference);				// The variable from the UAF Module
+		const FAnimNextParamType PinType = UE::UAF::UncookedOnly::FUtils::GetParamTypeFromPinType(ValuePin->PinType);	// The variable from BP pin
+		const UE::UAF::FParamCompatibility Compatibility = UE::UAF::FParamUtils::GetCompatibility(Type, PinType);
 		if (Compatibility == UE::UAF::EParamCompatibility::Incompatible || Compatibility == UE::UAF::EParamCompatibility::Incompatible_DataLoss)
 		{
 			MessageLog.Error(*FString::Printf(TEXT("@@ mismatched type, setting '%s' type, but variable is '%s' type"), *PinType.ToString(), *Type.ToString()), this);
@@ -119,7 +118,7 @@ void UK2Node_UAFComponentGetVariable::ValidateNodeDuringCompilation(FCompilerRes
 	}
 	else
 	{
-		MessageLog.Error(*FString::Printf(TEXT("@@ variable reference '%s' was not found in '%s'"), *VariableReference.GetName().ToString(), VariableReference.GetObject() ? TEXT("None") : *VariableReference.GetObject()->GetPathName()), this);
+		MessageLog.Error(*FString::Printf(TEXT("@@ variable reference '%s' was not found in '%s'"), *VariableReference.GetName().ToString(), VariableReference.GetObject() ? *VariableReference.GetObject()->GetPathName() : TEXT("None")), this);
 	}
 }
 
@@ -138,15 +137,7 @@ void UK2Node_UAFComponentGetVariable::GetMenuActions(FBlueprintActionDatabaseReg
 void UK2Node_UAFComponentGetVariable::PreloadRequiredAssets()
 {
 	Super::PreloadRequiredAssets();
-
-	UEdGraphPin* VariablePin = FindPinChecked(TEXT("Variable"));
-	FAnimNextVariableReference VariableReference;
-	FAnimNextVariableReference::StaticStruct()->ImportText(*VariablePin->DefaultValue, &VariableReference, nullptr, PPF_None, nullptr, FAnimNextVariableReference::StaticStruct()->GetName());
-	if (!VariableReference.IsNone() && VariableReference.GetObject())
-	{
-		UObject* Object = const_cast<UObject*>(VariableReference.GetObject().Get());
-		PreloadObject(Object);
-	}
+	UE::UAF::UncookedOnly::FUtils::PreloadVariableReferenceAssets(FindPinChecked(TEXT("Variable")));
 }
 
 #undef LOCTEXT_NAMESPACE
diff --git a/Engine/Plugins/Experimental/UAF/UAF/Source/UAFUncookedOnly/Private/Variables/K2Node_UAFComponentSetInputBinding.cpp b/Engine/Plugins/Experimental/UAF/UAF/Source/UAFUncookedOnly/Private/Variables/K2Node_UAFComponentSetInputBinding.cpp
new file mode 100644
--- /dev/null
+++ b/Engine/Plugins/Experimental/UAF/UAF/Source/UAFUncookedOnly/Private/Variables/K2Node_UAFComponentSetInputBinding.cpp
@@ -0,0 +1,131 @@
+// Copyright Epic Games, Inc. All Rights Reserved.
+
+#include "Variables/K2Node_UAFComponentSetInputBinding.h"
+
+#include "BlueprintActionDatabaseRegistrar.h"
+#include "BlueprintNodeSpawner.h"
+#include "UncookedOnlyUtils.h"
+#include "Kismet2/BlueprintEditorUtils.h"
+#include "Component/AnimNextComponent.h"
+#include "Components/SkeletalMeshComponent.h"
+#include "Param/ParamCompatibility.h"
+#include "Param/ParamUtils.h"
+#include "Param/ParamType.h"
+#include "Variables/AnimNextVariableReference.h"
+#include "UAF/ValueRuntime/ValueBundle.h"
+
+#include UE_INLINE_GENERATED_CPP_BY_NAME(K2Node_UAFComponentSetInputBinding)
+
+#define LOCTEXT_NAMESPACE "K2Node_UAFComponentSetInputBinding"
+
+const FName UK2Node_UAFComponentSetInputBinding::VariablePinName(TEXT("Variable"));
+const FName UK2Node_UAFComponentSetInputBinding::ComponentPinName(TEXT("Component"));
+const TArray<const UClass*> UK2Node_UAFComponentSetInputBinding::AllowedComponentClasses = { UUAFComponent::StaticClass(), USkeletalMeshComponent::StaticClass() };
+
+UK2Node_UAFComponentSetInputBinding::UK2Node_UAFComponentSetInputBinding()
+{
+	FunctionReference.SetExternalMember(GET_FUNCTION_NAME_CHECKED(UUAFComponent, BlueprintSetInputBinding), UUAFComponent::StaticClass());
+}
+
+FString UK2Node_UAFComponentSetInputBinding::GetPinMetaData(FName InPinName, FName InKey)
+{
+	if (InPinName == VariablePinName && InKey == TEXT("AllowedType"))
+	{
+		return FAnimNextParamType::GetType<FUAFValueBundle>().ToString();
+	}
+
+	return Super::GetPinMetaData(InPinName, InKey);
+}
+
+bool UK2Node_UAFComponentSetInputBinding::IsConnectionDisallowed(const UEdGraphPin* MyPin, const UEdGraphPin* OtherPin, FString& OutReason) const
+{
+	if (MyPin->GetFName() == ComponentPinName && MyPin->Direction == EGPD_Input)
+	{
+		if (const UClass* ConnectedClass = Cast<UClass>(OtherPin->PinType.PinSubCategoryObject.Get()))
+		{
+			if (!AllowedComponentClasses.ContainsByPredicate([ConnectedClass](const UClass* AllowedClass)
+			{
+				return ConnectedClass->IsChildOf(AllowedClass);
+			}))
+			{
+				OutReason = LOCTEXT("DisallowedComponentType", "Only UAFComponent and SkeletalMeshComponent types are allowed for input bindings.").ToString();
+				return true;
+			}
+		}
+	}
+
+	return Super::IsConnectionDisallowed(MyPin, OtherPin, OutReason);
+}
+
+void UK2Node_UAFComponentSetInputBinding::ValidateNodeDuringCompilation(FCompilerResultsLog& MessageLog) const
+{
+	Super::ValidateNodeDuringCompilation(MessageLog);
+
+	const UK2Node* SourceNode = Cast<UK2Node>(MessageLog.FindSourceObject(this));
+	if (SourceNode == nullptr)
+	{
+		return;
+	}
+
+	const UEdGraphPin* VariablePin = SourceNode->FindPinChecked(VariablePinName);
+
+	FAnimNextVariableReference VariableReference = UE::UAF::UncookedOnly::FUtils::ParseVariableReferenceFromPin(VariablePin);
+	if (VariableReference.IsValid())
+	{
+		const FAnimNextParamType Type = UE::UAF::UncookedOnly::FUtils::FindVariableType(VariableReference);
+		const FAnimNextParamType ExpectedType = FAnimNextParamType::GetType<FUAFValueBundle>();
+		const UE::UAF::FParamCompatibility Compatibility = UE::UAF::FParamUtils::GetCompatibility(ExpectedType, Type);
+		if (Compatibility == UE::UAF::EParamCompatibility::Incompatible || Compatibility == UE::UAF::EParamCompatibility::Incompatible_DataLoss)
+		{
+			MessageLog.Error(*FString::Printf(TEXT("@@ variable '%s' is of type '%s', but only ValueBundle type variables are allowed"), *VariableReference.GetName().ToString(), *Type.ToString()), this);
+		}
+	}
+	else if (VariableReference.IsNone())
+	{
+		MessageLog.Error(TEXT("@@ variable reference not set"), this);
+	}
+	else
+	{
+		MessageLog.Error(*FString::Printf(TEXT("@@ variable reference '%s' was not found in '%s'"), *VariableReference.GetName().ToString(), VariableReference.GetObject() ? *VariableReference.GetObject()->GetPathName() : TEXT("None")), this);
+	}
+
+	// Validate the component pin type if it is connected
+	UEdGraphPin* ComponentPin = SourceNode->FindPinChecked(ComponentPinName);
+	if (ComponentPin->LinkedTo.Num() > 0)
+	{
+		if (const UEdGraphPin* LinkedPin = ComponentPin->LinkedTo[0])
+		{
+			if (const UClass* ConnectedClass = Cast<UClass>(LinkedPin->PinType.PinSubCategoryObject.Get()))
+			{
+				const bool bIsAllowed = AllowedComponentClasses.ContainsByPredicate([ConnectedClass](const UClass* AllowedClass)
+				{
+					return ConnectedClass->IsChildOf(AllowedClass);
+				});
+				if (!bIsAllowed)
+				{
+					MessageLog.Error(*FString::Printf(TEXT("@@ component type '%s' is not allowed, only UAFComponent and SkeletalMeshComponent are supported"), *ConnectedClass->GetName()), this);
+				}
+			}
+		}
+	}
+}
+
+void UK2Node_UAFComponentSetInputBinding::GetMenuActions(FBlueprintActionDatabaseRegistrar& InActionRegistrar) const
+{
+	const UClass* ActionKey = GetClass();
+	if (InActionRegistrar.IsOpenForRegistration(ActionKey))
+	{
+		UBlueprintNodeSpawner* NodeSpawner = UBlueprintNodeSpawner::Create(GetClass());
+		check(NodeSpawner != nullptr);
+
+		InActionRegistrar.AddBlueprintAction(ActionKey, NodeSpawner);
+	}
+}
+
+void UK2Node_UAFComponentSetInputBinding::PreloadRequiredAssets()
+{
+	Super::PreloadRequiredAssets();
+	UE::UAF::UncookedOnly::FUtils::PreloadVariableReferenceAssets(FindPinChecked(VariablePinName));
+}
+
+#undef LOCTEXT_NAMESPACE
