Add FPoseValueBundle as a typed pose overlay for FValueBundle

FValueBundle is a generic attribute container, but animation evaluation overwhelmingly works with bone transforms and float curves. FPoseValueBundle is a zero-cost typed overlay that gives direct, named access to these common pose types instead of going through the generic attribute lookup API.

#Synced-CL 52535900

diff --git a/Engine/Plugins/Experimental/UAF/UAF/Source/UAF/Public/UAF/ValueRuntime/PoseValueBundle.h b/Engine/Plugins/Experimental/UAF/UAF/Source/UAF/Public/UAF/ValueRuntime/PoseValueBundle.h
new file mode 100644
--- /dev/null
+++ b/Engine/Plugins/Experimental/UAF/UAF/Source/UAF/Public/UAF/ValueRuntime/PoseValueBundle.h
@@ -0,0 +1,201 @@
+// Copyright Epic Games, Inc. All Rights Reserved.
+
+#pragma once
+
+#include "UAF/ValueRuntime/ValueBundle.h"
+
+namespace UE::UAF
+{
+	/**
+	 * FPoseValueBundle
+	 *
+	 * A non-owning typed overlay on FValueBundle that adds convenient accessors for common
+	 * animation pose data (bone transforms, float curves, etc.).
+	 *
+	 * This class is not meant to be constructed directly. It adds no data members and
+	 * FValueBundle has no vtable, so it can be safely reinterpret-cast from any FValueBundle
+	 * regardless of allocator (stack or heap). Use the static From(), AsMutable(), or AsConst()
+	 * helpers to obtain a FPoseValueBundle reference from an existing FValueBundle.
+	 */
+	class FPoseValueBundle : public FValueBundle
+	{
+	public:
+		using FValueBundle::operator=;
+
+		// -------------------------------------------------------------------
+		//  Typed accessors for common pose data
+		// -------------------------------------------------------------------
+
+		/**
+		 * Returns the bone transform map, or nullptr if this bundle doesn't contain bone transforms.
+		 *
+		 * Example: Read/write a single bone by index
+		 * @code
+		 *   TBoundValueMap<FBoneTransformAnimationAttribute>* BoneTransforms = Pose->FindBoneTransforms();
+		 *
+		 *   // Read
+		 *   FTransform CurrentTransform = (*BoneTransforms)[Mapping.BoneIndex].Value;
+		 *
+		 *   // Write
+		 *   (*BoneTransforms)[Mapping.BoneIndex].Value = NewTransform;
+		 * @endcode
+		 *
+		 * Example: Look up a bone by name
+		 * @code
+		 *   TBoundValueMap<FBoneTransformAnimationAttribute>* BoneTransforms = Pose->FindBoneTransforms();
+		 *   FAttributeTypedSetPtr BoneTypedSet = BoneTransforms->GetTypedSet();
+		 *   FAttributeSetIndex HeadIndex = BoneTypedSet->FindIndex(FName("head"));
+		 *   if (HeadIndex.IsValid())
+		 *   {
+		 *       FTransform HeadTransform = (*BoneTransforms)[HeadIndex].Value;
+		 *   }
+		 * @endcode
+		 *
+		 * Example: Iterate through all bones:
+		 * @code
+		 *   TBoundValueMap<FBoneTransformAnimationAttribute>* BoneTransforms = Pose->FindBoneTransforms();
+		 *   const int32 NumBones = BoneTransforms->Num();
+		 *   FBoneTransformAnimationAttribute* BoneData = BoneTransforms->GetData();
+		 *   for (int32 BoneIndex = 0; BoneIndex < NumBones; ++BoneIndex)
+		 *   {
+		 *       FTransform& BoneTransform = BoneData[BoneIndex].Value;
+		 *       // Read or modify BoneTransform
+		 *   }
+		 * @endcode
+		 */
+		[[nodiscard]] TBoundValueMap<FBoneTransformAnimationAttribute>* FindBoneTransforms();
+		[[nodiscard]] const TBoundValueMap<FBoneTransformAnimationAttribute>* FindBoneTransforms() const;
+
+		/**
+		 * Returns the float curve map (morph targets, control curves, material params), or nullptr if not present.
+		 * To look up curves by name, use GetTypedSet() on the returned map for the name-to-index mapping.
+		 *
+		 * Example: Read/write a curve by name
+		 * @code
+		 *   TBoundValueMap<FFloatAnimationAttribute>* FloatCurves = Pose->FindFloatCurves();
+		 *   FAttributeTypedSetPtr FloatTypedSet = FloatCurves->GetTypedSet();
+		 *   FAttributeSetIndex Index = FloatTypedSet->FindIndex(FName("smile_left"));
+		 *   if (Index.IsValid())
+		 *   {
+		 *       float Value = (*FloatCurves)[Index].Value;        // Read
+		 *       (*FloatCurves)[Index].Value = 0.75f;              // Write
+		 *   }
+		 * @endcode
+		 *
+		 * Example: Iterate all curves and scale them:
+		 * @code
+		 *   TBoundValueMap<FFloatAnimationAttribute>* FloatCurves = Pose->FindFloatCurves();
+		 *   const int32 NumCurves = FloatCurves->Num();
+		 *   FFloatAnimationAttribute* CurveData = FloatCurves->GetData();
+		 *   for (int32 CurveIndex = 0; CurveIndex < NumCurves; ++CurveIndex)
+		 *   {
+		 *       CurveData[CurveIndex].Value *= 0.5f;
+		 *   }
+		 * @endcode
+		 */
+		[[nodiscard]] TBoundValueMap<FFloatAnimationAttribute>* FindFloatCurves();
+		[[nodiscard]] const TBoundValueMap<FFloatAnimationAttribute>* FindFloatCurves() const;
+
+		/**
+		 * Returns whether this bundle contains bone transforms.
+		 * @note Performs a map lookup internally (O(log N) where N is the number of bound value maps, typically very small).
+		 *		 Cache the result rather than calling per-iteration in tight loops.
+		 */
+		[[nodiscard]] bool HasBoneTransforms() const;
+
+		/**
+		 * Returns whether this bundle contains float curves.
+		 * @note Performs a map lookup internally (O(log N) where N is the number of bound value maps, typically very small).
+		 *		 Cache the result rather than calling per-iteration in tight loops.
+		 */
+		[[nodiscard]] bool HasFloatCurves() const;
+
+		// -------------------------------------------------------------------
+		//  Cast helpers
+		// -------------------------------------------------------------------
+
+		/**
+		 * Reinterpret a FValueBundle reference as a FPoseValueBundle.
+		 * Safe because FPoseValueBundle adds no data members and FValueBundle has no vtable.
+		 */
+		[[nodiscard]] static FPoseValueBundle& From(FValueBundle& Bundle)
+		{
+			return static_cast<FPoseValueBundle&>(Bundle);
+		}
+
+		/** @copydoc From(FValueBundle&) */
+		[[nodiscard]] static const FPoseValueBundle& From(const FValueBundle& Bundle)
+		{
+			return static_cast<const FPoseValueBundle&>(Bundle);
+		}
+
+		/** Cast a mutable FValueBundle to FPoseValueBundle. */
+		[[nodiscard]] static FPoseValueBundle* AsMutable(FValueBundle& Bundle)
+		{
+			return &From(Bundle);
+		}
+
+		/**
+		 * Cast a CoW reference to a mutable FPoseValueBundle, returning nullptr if the bundle is empty.
+		 * Checks mutability internally so the caller doesn't have to.
+		 *
+		 * Typical usage:
+		 * @code
+		 *   FValueBundleCoWRef InputRef = Evaluator.GetEvaluationStack().Pop();
+		 *   FPoseValueBundle* Pose = FPoseValueBundle::AsMutable(InputRef);
+		 * @endcode
+		 */
+		template<typename CoWRefType>
+		[[nodiscard]] static FPoseValueBundle* AsMutable(CoWRefType& Ref)
+		{
+			check(Ref.IsMutable());
+			return &From(Ref.GetMutable());
+		}
+
+		/**
+		 * Cast a CoW reference to a const FPoseValueBundle, returning nullptr if the bundle is empty.
+		 */
+		template<typename CoWRefType>
+		[[nodiscard]] static const FPoseValueBundle* AsConst(const CoWRefType& Ref)
+		{
+			return &From(Ref.Get());
+		}
+	};
+
+	//////////////////////////////////////////////////////////////////////////
+	// Inline implementation
+
+	inline TBoundValueMap<FBoneTransformAnimationAttribute>* FPoseValueBundle::FindBoneTransforms()
+	{
+		return GetBoundValueMaps().Find<FBoneTransformAnimationAttribute>(
+			FAttributeMappingKey::MakeFromTo<FBoneTransformAnimationAttribute>());
+	}
+
+	inline const TBoundValueMap<FBoneTransformAnimationAttribute>* FPoseValueBundle::FindBoneTransforms() const
+	{
+		return GetBoundValueMaps().Find<FBoneTransformAnimationAttribute>(
+			FAttributeMappingKey::MakeFromTo<FBoneTransformAnimationAttribute>());
+	}
+
+	inline TBoundValueMap<FFloatAnimationAttribute>* FPoseValueBundle::FindFloatCurves()
+	{
+		return GetBoundValueMaps().Find<FFloatAnimationAttribute>(
+			FAttributeMappingKey::MakeFromTo<FFloatAnimationAttribute>());
+	}
+
+	inline const TBoundValueMap<FFloatAnimationAttribute>* FPoseValueBundle::FindFloatCurves() const
+	{
+		return GetBoundValueMaps().Find<FFloatAnimationAttribute>(
+			FAttributeMappingKey::MakeFromTo<FFloatAnimationAttribute>());
+	}
+
+	inline bool FPoseValueBundle::HasBoneTransforms() const
+	{
+		return FindBoneTransforms() != nullptr;
+	}
+
+	inline bool FPoseValueBundle::HasFloatCurves() const
+	{
+		return FindFloatCurves() != nullptr;
+	}
+} // namespace UE::UAF
