Fix uninitialized ContentOnly in SRigVMFunctionPicker causing state tree function binding regression

#Synced-CL 52518742

diff --git a/Engine/Plugins/Experimental/UAF/UAF/Source/UAFEditor/Internal/Common/SRigVMFunctionPicker.h b/Engine/Plugins/Experimental/UAF/UAF/Source/UAFEditor/Internal/Common/SRigVMFunctionPicker.h
--- a/Engine/Plugins/Experimental/UAF/UAF/Source/UAFEditor/Internal/Common/SRigVMFunctionPicker.h
+++ b/Engine/Plugins/Experimental/UAF/UAF/Source/UAFEditor/Internal/Common/SRigVMFunctionPicker.h
@@ -32,6 +32,8 @@ class SRigVMFunctionPicker : public SCompoundWidget
 	SLATE_BEGIN_ARGS(SRigVMFunctionPicker)
 		: _AllowNew(true)
 		, _AllowClear(true)
+		, _ContentOnly(false)
+		, _IsContextSensitive(true)
 	{}
 
 	// The 'current' asset, used for accessing private functions
