Fix FBindableValueCustomization editor bugs

- FBindableTransform: show Location/Rotation/Scale as expandable child rows
(CreatePropertyValueWidgetWithCustomization returns empty widget for FTransform)
- Function binding validation: search all public exports when no asset context
is available (e.g. chooser table), instead of immediately reporting Missing

#Synced-CL 52518742

diff --git a/Engine/Plugins/Experimental/UAF/UAF/Source/UAFEditor/Private/BindableValue/UAFBindableValueCustomization.cpp b/Engine/Plugins/Experimental/UAF/UAF/Source/UAFEditor/Private/BindableValue/UAFBindableValueCustomization.cpp
--- a/Engine/Plugins/Experimental/UAF/UAF/Source/UAFEditor/Private/BindableValue/UAFBindableValueCustomization.cpp
+++ b/Engine/Plugins/Experimental/UAF/UAF/Source/UAFEditor/Private/BindableValue/UAFBindableValueCustomization.cpp
@@ -233,25 +233,26 @@ EFunctionBindingStatus ValidateFunctionBinding(
 	const UUAFRigVMAsset* InAsset,
 	FName* OutResolvedName = nullptr)
 {
-	if (!InAsset)
-	{
-		return EFunctionBindingStatus::FunctionNotFound;
-	}
-
-	const UUAFRigVMAssetEditorData* EditorData = UE::UAF::UncookedOnly::FUtils::GetEditorData(InAsset);
-	if (!EditorData)
-	{
-		return EFunctionBindingStatus::FunctionNotFound;
-	}
-
-	// Only validate against public functions — private functions are not compiled
-	// into FunctionData and cannot be called at runtime.
+	// Collect public function headers to validate against.
 	FRigVMGraphFunctionHeaderArray AllHeaders;
 	TMap<FAssetData, FRigVMGraphFunctionHeaderArray> PublicExports;
 	UE::UAF::UncookedOnly::FUtils::GetExportedFunctionsFromAssetRegistry(UE::UAF::AnimNextPublicGraphFunctionsExportsRegistryTag, PublicExports);
-	if (const FRigVMGraphFunctionHeaderArray* AssetPublicHeaders = PublicExports.Find(FAssetData(InAsset)))
+
+	if (InAsset)
+	{
+		// Prefer the specific asset's exports when available.
+		if (const FRigVMGraphFunctionHeaderArray* AssetPublicHeaders = PublicExports.Find(FAssetData(InAsset)))
+		{
+			AllHeaders.Headers.Append(AssetPublicHeaders->Headers);
+		}
+	}
+	else
 	{
-		AllHeaders.Headers.Append(AssetPublicHeaders->Headers);
+		// No asset context (e.g. chooser table) -- search all public exports.
+		for (const auto& Pair : PublicExports)
+		{
+			AllHeaders.Headers.Append(Pair.Value.Headers);
+		}
 	}
 
 	// Match by GUID first (rename-robust), fall back to name if GUID is invalid (legacy data)
@@ -672,6 +673,10 @@ void FBindableValueCustomization::CustomizeHeader(
 				.PropertyHandle(ConstantValueHandle)
 				.AllowedClass(ObjClass ? ObjClass : UObject::StaticClass());
 		}
+		else if (BindableStruct == FBindableTransform::StaticStruct())
+		{
+			bIsBindableTransform = true;
+		}
 		else
 		{
 			ConstantWidget = ConstantValueHandle->CreatePropertyValueWidgetWithCustomization(nullptr);
@@ -1006,7 +1011,34 @@ void FBindableValueCustomization::CustomizeChildren(
 			}
 		}
 	}
-	// All other types: intentionally empty — suppresses ConstantValue and Binding as separate child rows.
+	else if (bIsBindableTransform && ConstantValueHandle.IsValid())
+	{
+		auto HasBinding = [Handle = PropertyHandle]() -> bool
+		{
+			bool bResult = false;
+			Handle->EnumerateConstRawData([&bResult](const void* RawData, const int32, const int32) -> bool
+			{
+				bResult = static_cast<const FBindableValueBase*>(RawData)->HasBinding();
+				return false;
+			});
+			return bResult;
+		};
+
+		// Iterate children generically -- the property system applies the correct
+		// customizations for each sub-property (FVector inline editors, etc.).
+		uint32 NumChildren = 0;
+		ConstantValueHandle->GetNumChildren(NumChildren);
+		for (uint32 ChildIndex = 0; ChildIndex < NumChildren; ++ChildIndex)
+		{
+			TSharedPtr<IPropertyHandle> ChildHandle = ConstantValueHandle->GetChildHandle(ChildIndex);
+			if (ChildHandle.IsValid())
+			{
+				ChildBuilder.AddProperty(ChildHandle.ToSharedRef())
+					.IsEnabled(TAttribute<bool>::CreateLambda([HasBinding]() { return !HasBinding(); }));
+			}
+		}
+	}
+	// All other types: intentionally empty -- suppresses ConstantValue and Binding as separate child rows.
 }
 
 } // namespace UE::UAF::Editor
diff --git a/Engine/Plugins/Experimental/UAF/UAF/Source/UAFEditor/Private/BindableValue/UAFBindableValueCustomization.h b/Engine/Plugins/Experimental/UAF/UAF/Source/UAFEditor/Private/BindableValue/UAFBindableValueCustomization.h
--- a/Engine/Plugins/Experimental/UAF/UAF/Source/UAFEditor/Private/BindableValue/UAFBindableValueCustomization.h
+++ b/Engine/Plugins/Experimental/UAF/UAF/Source/UAFEditor/Private/BindableValue/UAFBindableValueCustomization.h
@@ -37,6 +37,7 @@ private:
 	TSharedPtr<IPropertyHandle> ConstantValueHandle;
 	const UScriptStruct* ResolvedStructClass = nullptr;
 	bool bIsBindableStruct = false;
+	bool bIsBindableTransform = false;
 
 	bool bIsContextSensitive = true;
 	TArray<TWeakObjectPtr<const UUAFRigVMAsset>> FilterAssets;
