Fix FBindableXxx delta serialization: add WithIdentical to preserve bindings

#Synced-CL 52518742

diff --git a/Engine/Plugins/Experimental/UAF/UAF/Source/UAF/Private/BindableValue/UAFBindableTypes.cpp b/Engine/Plugins/Experimental/UAF/UAF/Source/UAF/Private/BindableValue/UAFBindableTypes.cpp
--- a/Engine/Plugins/Experimental/UAF/UAF/Source/UAF/Private/BindableValue/UAFBindableTypes.cpp
+++ b/Engine/Plugins/Experimental/UAF/UAF/Source/UAF/Private/BindableValue/UAFBindableTypes.cpp
@@ -847,3 +847,69 @@ IMPLEMENT_BINDABLE_TEXT_SERIALIZATION(FBindableStruct)
 IMPLEMENT_BINDABLE_TEXT_SERIALIZATION(FBindableObject)
 
 #undef IMPLEMENT_BINDABLE_TEXT_SERIALIZATION
+
+//------------------------------------------------------------------------------
+// Per-type Identical implementations (WithIdentical)
+//
+// Delta serialization compares each struct property against defaults to decide
+// whether to write a property tag. Without WithIdentical, the comparison only
+// checks UPROPERTYs and misses the non-UPROPERTY Binding. When ConstantValue
+// is at its default, the struct appears "identical" and the property tag is
+// skipped -- silently dropping the binding on save.
+//------------------------------------------------------------------------------
+
+bool FBindableValueBase::BindingsIdentical(const FBindableValueBase* Other, const UScriptStruct* ConcreteStruct, uint32 PortFlags) const
+{
+	if (!Other)
+	{
+		return false;
+	}
+
+	// Compare binding presence
+	if (HasBinding() != Other->HasBinding())
+	{
+		return false;
+	}
+
+	// Compare binding contents
+	if (HasBinding())
+	{
+		if (!FUAFPropertyBinding::StaticStruct()->CompareScriptStruct(Binding.Get(), Other->Binding.Get(), PortFlags))
+		{
+			return false;
+		}
+	}
+
+	// Compare all UPROPERTYs (ConstantValue, EnumClass, StructClass, etc.)
+	for (TFieldIterator<FProperty> It(ConcreteStruct); It; ++It)
+	{
+		if (!It->Identical_InContainer(this, Other, 0, PortFlags))
+		{
+			return false;
+		}
+	}
+
+	return true;
+}
+
+#define IMPLEMENT_BINDABLE_IDENTICAL(StructType) \
+	bool StructType::Identical(const StructType* Other, uint32 PortFlags) const \
+	{ \
+		return BindingsIdentical(Other, StructType::StaticStruct(), PortFlags); \
+	}
+
+IMPLEMENT_BINDABLE_IDENTICAL(FBindableBool)
+IMPLEMENT_BINDABLE_IDENTICAL(FBindableFloat)
+IMPLEMENT_BINDABLE_IDENTICAL(FBindableDouble)
+IMPLEMENT_BINDABLE_IDENTICAL(FBindableInt32)
+IMPLEMENT_BINDABLE_IDENTICAL(FBindableInt64)
+IMPLEMENT_BINDABLE_IDENTICAL(FBindableByte)
+IMPLEMENT_BINDABLE_IDENTICAL(FBindableName)
+IMPLEMENT_BINDABLE_IDENTICAL(FBindableVector)
+IMPLEMENT_BINDABLE_IDENTICAL(FBindableQuat)
+IMPLEMENT_BINDABLE_IDENTICAL(FBindableTransform)
+IMPLEMENT_BINDABLE_IDENTICAL(FBindableEnum)
+IMPLEMENT_BINDABLE_IDENTICAL(FBindableStruct)
+IMPLEMENT_BINDABLE_IDENTICAL(FBindableObject)
+
+#undef IMPLEMENT_BINDABLE_IDENTICAL
diff --git a/Engine/Plugins/Experimental/UAF/UAF/Source/UAF/Public/BindableValue/UAFBindableTypes.h b/Engine/Plugins/Experimental/UAF/UAF/Source/UAF/Public/BindableValue/UAFBindableTypes.h
--- a/Engine/Plugins/Experimental/UAF/UAF/Source/UAF/Public/BindableValue/UAFBindableTypes.h
+++ b/Engine/Plugins/Experimental/UAF/UAF/Source/UAF/Public/BindableValue/UAFBindableTypes.h
@@ -99,6 +99,13 @@ struct FBindableValueBase
 	UAF_API bool ImportTextItem(const TCHAR*& Buffer, int32 PortFlags,
 		UObject* Parent, FOutputDevice* ErrorText);
 
+	/**
+	 * Compares this bindable value against Other, including the non-UPROPERTY Binding.
+	 * Without this, delta serialization skips the struct when all UPROPERTYs match defaults,
+	 * silently dropping the binding. Called by per-type Identical() via DECLARE_BINDABLE_IDENTICAL.
+	 */
+	UAF_API bool BindingsIdentical(const FBindableValueBase* Other, const UScriptStruct* ConcreteStruct, uint32 PortFlags) const;
+
 protected:
 	/** Shared helpers for text serialization — called by concrete subclasses via IMPLEMENT_BINDABLE_TEXT_SERIALIZATION. */
 	UAF_API static void ExportBindingText(FString& ValueStr, const FUAFPropertyBinding* InBinding, int32 PortFlags);
@@ -128,6 +135,32 @@ namespace UE::UAF::Private
 	bool ImportTextItem(const TCHAR*& Buffer, int32 PortFlags, \
 		UObject* Parent, FOutputDevice* ErrorText);
 
+/**
+ * Declares Identical() on a concrete FBindableXxx type, delegating to the base BindingsIdentical helper.
+ * Required so that delta serialization includes the non-UPROPERTY binding in its comparison.
+ * Place inside the struct body alongside DECLARE_BINDABLE_TEXT_SERIALIZATION.
+ */
+#define DECLARE_BINDABLE_IDENTICAL(StructType) \
+	bool Identical(const StructType* Other, uint32 PortFlags) const;
+
+/**
+ * Declares TStructOpsTypeTraits for a concrete FBindableXxx type.
+ * UE does not inherit traits from base structs, so each subclass must declare its own.
+ */
+#define DECLARE_BINDABLE_STRUCT_OPS(StructType) \
+	template<> \
+	struct TStructOpsTypeTraits<StructType> : public TStructOpsTypeTraitsBase2<StructType> \
+	{ \
+		enum \
+		{ \
+			WithSerializer = true, \
+			WithCopy = true, \
+			WithIdentical = true, \
+			WithExportTextItem = true, \
+			WithImportTextItem = true, \
+		}; \
+	};
+
 template<>
 struct TStructOpsTypeTraits<FBindableValueBase> : public TStructOpsTypeTraitsBase2<FBindableValueBase>
 {
@@ -149,6 +182,7 @@ struct FBindableBool : public FBindableValueBase
 {
 	GENERATED_BODY()
 	DECLARE_BINDABLE_TEXT_SERIALIZATION(FBindableBool)
+	DECLARE_BINDABLE_IDENTICAL(FBindableBool)
 
 	FBindableBool() = default;
 	FBindableBool(bool InValue) : ConstantValue(InValue) {}
@@ -164,17 +198,7 @@ private:
 	bool ConstantValue = false;
 };
 
-template <>
-struct TStructOpsTypeTraits<FBindableBool> : public TStructOpsTypeTraitsBase2<FBindableBool>
-{
-	enum
-	{
-		WithSerializer = true,
-		WithCopy = true,
-		WithExportTextItem = true,
-		WithImportTextItem = true
-	};
-};
+DECLARE_BINDABLE_STRUCT_OPS(FBindableBool)
 
 //------------------------------------------------------------------------------
 // FBindableFloat
@@ -185,6 +209,7 @@ struct FBindableFloat : public FBindableValueBase
 {
 	GENERATED_BODY()
 	DECLARE_BINDABLE_TEXT_SERIALIZATION(FBindableFloat)
+	DECLARE_BINDABLE_IDENTICAL(FBindableFloat)
 
 	FBindableFloat() = default;
 	FBindableFloat(float InValue) : ConstantValue(InValue) {}
@@ -199,17 +224,7 @@ private:
 	float ConstantValue = 0.f;
 };
 
-template<>
-struct TStructOpsTypeTraits<FBindableFloat> : public TStructOpsTypeTraitsBase2<FBindableFloat>
-{
-	enum
-	{
-		WithSerializer = true,
-		WithCopy = true,
-		WithExportTextItem = true,
-		WithImportTextItem = true,
-	};
-};
+DECLARE_BINDABLE_STRUCT_OPS(FBindableFloat)
 
 //------------------------------------------------------------------------------
 // FBindableDouble
@@ -220,6 +235,7 @@ struct FBindableDouble : public FBindableValueBase
 {
 	GENERATED_BODY()
 	DECLARE_BINDABLE_TEXT_SERIALIZATION(FBindableDouble)
+	DECLARE_BINDABLE_IDENTICAL(FBindableDouble)
 
 	FBindableDouble() = default;
 	FBindableDouble(double InValue) : ConstantValue(InValue) {}
@@ -234,17 +250,7 @@ private:
 	double ConstantValue = 0.0;
 };
 
-template<>
-struct TStructOpsTypeTraits<FBindableDouble> : public TStructOpsTypeTraitsBase2<FBindableDouble>
-{
-	enum
-	{
-		WithSerializer = true,
-		WithCopy = true,
-		WithExportTextItem = true,
-		WithImportTextItem = true,
-	};
-};
+DECLARE_BINDABLE_STRUCT_OPS(FBindableDouble)
 
 //------------------------------------------------------------------------------
 // FBindableInt32
@@ -255,6 +261,7 @@ struct FBindableInt32 : public FBindableValueBase
 {
 	GENERATED_BODY()
 	DECLARE_BINDABLE_TEXT_SERIALIZATION(FBindableInt32)
+	DECLARE_BINDABLE_IDENTICAL(FBindableInt32)
 
 	FBindableInt32() = default;
 	FBindableInt32(int32 InValue) : ConstantValue(InValue) {}
@@ -269,17 +276,7 @@ private:
 	int32 ConstantValue = 0;
 };
 
-template<>
-struct TStructOpsTypeTraits<FBindableInt32> : public TStructOpsTypeTraitsBase2<FBindableInt32>
-{
-	enum
-	{
-		WithSerializer = true,
-		WithCopy = true,
-		WithExportTextItem = true,
-		WithImportTextItem = true,
-	};
-};
+DECLARE_BINDABLE_STRUCT_OPS(FBindableInt32)
 
 //------------------------------------------------------------------------------
 // FBindableInt64
@@ -290,6 +287,7 @@ struct FBindableInt64 : public FBindableValueBase
 {
 	GENERATED_BODY()
 	DECLARE_BINDABLE_TEXT_SERIALIZATION(FBindableInt64)
+	DECLARE_BINDABLE_IDENTICAL(FBindableInt64)
 
 	FBindableInt64() = default;
 	FBindableInt64(int64 InValue) : ConstantValue(InValue) {}
@@ -304,17 +302,7 @@ private:
 	int64 ConstantValue = 0;
 };
 
-template<>
-struct TStructOpsTypeTraits<FBindableInt64> : public TStructOpsTypeTraitsBase2<FBindableInt64>
-{
-	enum
-	{
-		WithSerializer = true,
-		WithCopy = true,
-		WithExportTextItem = true,
-		WithImportTextItem = true,
-	};
-};
+DECLARE_BINDABLE_STRUCT_OPS(FBindableInt64)
 
 //------------------------------------------------------------------------------
 // FBindableByte
@@ -325,6 +313,7 @@ struct FBindableByte : public FBindableValueBase
 {
 	GENERATED_BODY()
 	DECLARE_BINDABLE_TEXT_SERIALIZATION(FBindableByte)
+	DECLARE_BINDABLE_IDENTICAL(FBindableByte)
 
 	FBindableByte() = default;
 	FBindableByte(uint8 InValue) : ConstantValue(InValue) {}
@@ -339,17 +328,7 @@ private:
 	uint8 ConstantValue = 0;
 };
 
-template<>
-struct TStructOpsTypeTraits<FBindableByte> : public TStructOpsTypeTraitsBase2<FBindableByte>
-{
-	enum
-	{
-		WithSerializer = true,
-		WithCopy = true,
-		WithExportTextItem = true,
-		WithImportTextItem = true,
-	};
-};
+DECLARE_BINDABLE_STRUCT_OPS(FBindableByte)
 
 //------------------------------------------------------------------------------
 // FBindableName
@@ -360,6 +339,7 @@ struct FBindableName : public FBindableValueBase
 {
 	GENERATED_BODY()
 	DECLARE_BINDABLE_TEXT_SERIALIZATION(FBindableName)
+	DECLARE_BINDABLE_IDENTICAL(FBindableName)
 
 	FBindableName() = default;
 	FBindableName(FName InValue) : ConstantValue(InValue) {}
@@ -374,17 +354,7 @@ private:
 	FName ConstantValue;
 };
 
-template<>
-struct TStructOpsTypeTraits<FBindableName> : public TStructOpsTypeTraitsBase2<FBindableName>
-{
-	enum
-	{
-		WithSerializer = true,
-		WithCopy = true,
-		WithExportTextItem = true,
-		WithImportTextItem = true,
-	};
-};
+DECLARE_BINDABLE_STRUCT_OPS(FBindableName)
 
 //------------------------------------------------------------------------------
 // FBindableVector
@@ -395,6 +365,7 @@ struct FBindableVector : public FBindableValueBase
 {
 	GENERATED_BODY()
 	DECLARE_BINDABLE_TEXT_SERIALIZATION(FBindableVector)
+	DECLARE_BINDABLE_IDENTICAL(FBindableVector)
 
 	FBindableVector() = default;
 	FBindableVector(const FVector& InValue) : ConstantValue(InValue) {}
@@ -409,17 +380,7 @@ private:
 	FVector ConstantValue = FVector::ZeroVector;
 };
 
-template<>
-struct TStructOpsTypeTraits<FBindableVector> : public TStructOpsTypeTraitsBase2<FBindableVector>
-{
-	enum
-	{
-		WithSerializer = true,
-		WithCopy = true,
-		WithExportTextItem = true,
-		WithImportTextItem = true,
-	};
-};
+DECLARE_BINDABLE_STRUCT_OPS(FBindableVector)
 
 //------------------------------------------------------------------------------
 // FBindableQuat
@@ -430,6 +391,7 @@ struct FBindableQuat : public FBindableValueBase
 {
 	GENERATED_BODY()
 	DECLARE_BINDABLE_TEXT_SERIALIZATION(FBindableQuat)
+	DECLARE_BINDABLE_IDENTICAL(FBindableQuat)
 
 	FBindableQuat() = default;
 	FBindableQuat(const FQuat& InValue) : ConstantValue(InValue) {}
@@ -444,17 +406,7 @@ private:
 	FQuat ConstantValue = FQuat::Identity;
 };
 
-template<>
-struct TStructOpsTypeTraits<FBindableQuat> : public TStructOpsTypeTraitsBase2<FBindableQuat>
-{
-	enum
-	{
-		WithSerializer = true,
-		WithCopy = true,
-		WithExportTextItem = true,
-		WithImportTextItem = true,
-	};
-};
+DECLARE_BINDABLE_STRUCT_OPS(FBindableQuat)
 
 //------------------------------------------------------------------------------
 // FBindableTransform
@@ -465,6 +417,7 @@ struct FBindableTransform : public FBindableValueBase
 {
 	GENERATED_BODY()
 	DECLARE_BINDABLE_TEXT_SERIALIZATION(FBindableTransform)
+	DECLARE_BINDABLE_IDENTICAL(FBindableTransform)
 
 	FBindableTransform() = default;
 	FBindableTransform(const FTransform& InValue) : ConstantValue(InValue) {}
@@ -479,17 +432,7 @@ private:
 	FTransform ConstantValue = FTransform::Identity;
 };
 
-template<>
-struct TStructOpsTypeTraits<FBindableTransform> : public TStructOpsTypeTraitsBase2<FBindableTransform>
-{
-	enum
-	{
-		WithSerializer = true,
-		WithCopy = true,
-		WithExportTextItem = true,
-		WithImportTextItem = true,
-	};
-};
+DECLARE_BINDABLE_STRUCT_OPS(FBindableTransform)
 
 //------------------------------------------------------------------------------
 // FBindableEnum
@@ -512,6 +455,7 @@ struct FBindableEnum : public FBindableValueBase
 {
 	GENERATED_BODY()
 	DECLARE_BINDABLE_TEXT_SERIALIZATION(FBindableEnum)
+	DECLARE_BINDABLE_IDENTICAL(FBindableEnum)
 
 	FBindableEnum() = default;
 	explicit FBindableEnum(UEnum* InEnumClass) : EnumClass(InEnumClass) {}
@@ -571,17 +515,7 @@ private:
 	int32 ConstantValue = 0;
 };
 
-template<>
-struct TStructOpsTypeTraits<FBindableEnum> : public TStructOpsTypeTraitsBase2<FBindableEnum>
-{
-	enum
-	{
-		WithSerializer = true,
-		WithCopy = true,
-		WithExportTextItem = true,
-		WithImportTextItem = true,
-	};
-};
+DECLARE_BINDABLE_STRUCT_OPS(FBindableEnum)
 
 //------------------------------------------------------------------------------
 // FBindableStruct
@@ -603,6 +537,7 @@ struct FBindableStruct : public FBindableValueBase
 {
 	GENERATED_BODY()
 	DECLARE_BINDABLE_TEXT_SERIALIZATION(FBindableStruct)
+	DECLARE_BINDABLE_IDENTICAL(FBindableStruct)
 
 	FBindableStruct() = default;
 	explicit FBindableStruct(UScriptStruct* InStructClass) : StructClass(InStructClass) {}
@@ -644,17 +579,7 @@ private:
 	FInstancedStruct ConstantValue;
 };
 
-template<>
-struct TStructOpsTypeTraits<FBindableStruct> : public TStructOpsTypeTraitsBase2<FBindableStruct>
-{
-	enum
-	{
-		WithSerializer = true,
-		WithCopy = true,
-		WithExportTextItem = true,
-		WithImportTextItem = true,
-	};
-};
+DECLARE_BINDABLE_STRUCT_OPS(FBindableStruct)
 
 //------------------------------------------------------------------------------
 // FBindableObject
@@ -672,6 +597,7 @@ struct FBindableObject : public FBindableValueBase
 {
 	GENERATED_BODY()
 	DECLARE_BINDABLE_TEXT_SERIALIZATION(FBindableObject)
+	DECLARE_BINDABLE_IDENTICAL(FBindableObject)
 
 	FBindableObject() = default;
 	explicit FBindableObject(UObject* InValue) : ConstantValue(InValue) {}
@@ -699,14 +625,4 @@ private:
 	TObjectPtr<UObject> ConstantValue = nullptr;
 };
 
-template<>
-struct TStructOpsTypeTraits<FBindableObject> : public TStructOpsTypeTraitsBase2<FBindableObject>
-{
-	enum
-	{
-		WithSerializer = true,
-		WithCopy = true,
-		WithExportTextItem = true,
-		WithImportTextItem = true,
-	};
-};
+DECLARE_BINDABLE_STRUCT_OPS(FBindableObject)
