Add GetValue helper to cast to a specific object type

#Synced-CL 52539819

diff --git a/Engine/Plugins/Experimental/UAF/UAF/Source/UAF/Public/BindableValue/UAFBindableTypes.h b/Engine/Plugins/Experimental/UAF/UAF/Source/UAF/Public/BindableValue/UAFBindableTypes.h
--- a/Engine/Plugins/Experimental/UAF/UAF/Source/UAF/Public/BindableValue/UAFBindableTypes.h
+++ b/Engine/Plugins/Experimental/UAF/UAF/Source/UAF/Public/BindableValue/UAFBindableTypes.h
@@ -610,6 +610,13 @@ struct FBindableObject : public FBindableValueBase
 
 	UAF_API UObject* GetValue(FUAFAssetInstance* Instance) const;
 
+	template<class T> requires(std::is_base_of_v<UObject, std::decay_t<T>>)
+	T* GetValue(FUAFAssetInstance* Instance) const
+	{
+		checkf(ObjectClass == nullptr || ObjectClass == T::StaticClass(), TEXT("Attempting to cast to incompatible object type."));
+		return static_cast<T*>(GetValue(Instance));
+	}
+
 	UObject* GetConstantValue() const { return ConstantValue; }
 	void     SetConstantValue(UObject* InValue) { ConstantValue = InValue; }
 
