Fix FBindableXxx binding serialization issues when used with UAF traits

- Set/clear bindings via direct raw data modification + NotifyPostChange
(SetValueFromFormattedString fails for FStructOnScope-backed property handles)
- Always export __Binding=() marker even when no binding is set, so the RigVM
controller updates the parent pin DefaultValue (fixes clear binding on traits)
- Propagate ConstantValue child edits to parent handle via SetOnPropertyValueChanged
- Update serialization test for new __Binding=() export behavior
- Register FBindableValues with rig vm as compatible struct types

<description: restricted, no permission to view>

#Synced-CL 52555156

diff --git a/Engine/Plugins/Experimental/UAF/UAF/Source/UAF/Private/AnimNextModuleImpl.cpp b/Engine/Plugins/Experimental/UAF/UAF/Source/UAF/Private/AnimNextModuleImpl.cpp
--- a/Engine/Plugins/Experimental/UAF/UAF/Source/UAF/Private/AnimNextModuleImpl.cpp
+++ b/Engine/Plugins/Experimental/UAF/UAF/Source/UAF/Private/AnimNextModuleImpl.cpp
@@ -28,6 +28,7 @@
 #include "Variables/AnimNextSoftFunctionPtr.h"
 #include "RigVMCore/RigVMGraphFunctionDefinition.h"
 #include "HierarchyTable.h"
+#include "BindableValue/UAFBindableTypes.h"
 #include "CrashReporter/CrashReporterHandler.h"
 #include "Factory/SystemFactory.h"
 #include "Features/IModularFeatures.h"
@@ -92,6 +93,19 @@ namespace UE::UAF
 			FAnimNextGraphLODPose::StaticStruct(),
 			FAnimNextGraphReferencePose::StaticStruct(),
 			FUAFValueBundle::StaticStruct(),
+			FBindableBool::StaticStruct(),
+			FBindableFloat::StaticStruct(),
+			FBindableDouble::StaticStruct(),
+			FBindableInt32::StaticStruct(),
+			FBindableInt64::StaticStruct(),
+			FBindableByte::StaticStruct(),
+			FBindableName::StaticStruct(),
+			FBindableVector::StaticStruct(),
+			FBindableQuat::StaticStruct(),
+			FBindableTransform::StaticStruct(),
+			FBindableEnum::StaticStruct(),
+			FBindableStruct::StaticStruct(),
+			FBindableObject::StaticStruct(),
 		};
 
 		RigVMRegistry.RegisterStructTypes(AllowedStructTypes);
diff --git a/Engine/Plugins/Experimental/UAF/UAF/Source/UAF/Private/BindableValue/UAFBindableTypes.cpp b/Engine/Plugins/Experimental/UAF/UAF/Source/UAF/Private/BindableValue/UAFBindableTypes.cpp
--- a/Engine/Plugins/Experimental/UAF/UAF/Source/UAF/Private/BindableValue/UAFBindableTypes.cpp
+++ b/Engine/Plugins/Experimental/UAF/UAF/Source/UAF/Private/BindableValue/UAFBindableTypes.cpp
@@ -589,13 +589,22 @@ static const TCHAR* BindingPropertyName = TEXT("__Binding");
 
 void FBindableValueBase::ExportBindingText(FString& ValueStr, const FUAFPropertyBinding* InBinding, int32 PortFlags)
 {
-	if (!InBinding)
+	// Always export __Binding, even when empty. When the RigVM controller processes
+	// SetPinDefaultValue for a struct pin, it distributes values to sub-pins. If ALL
+	// text properties map to sub-pins, the parent pin's DefaultValue is NOT updated.
+	// By always including __Binding (which has no sub-pin), the controller detects a
+	// "skipped" property and updates the parent pin's DefaultValue -- ensuring that
+	// cleared bindings are persisted correctly.
+	FString BindingText;
+	if (InBinding)
 	{
-		return;
+		FUAFPropertyBinding::StaticStruct()->ExportText(BindingText, InBinding, nullptr, nullptr, PortFlags, nullptr, true);
+	}
+	else
+	{
+		BindingText = TEXT("()");
 	}
 
-	FString BindingText;
-	FUAFPropertyBinding::StaticStruct()->ExportText(BindingText, InBinding, nullptr, nullptr, PortFlags, nullptr, true);
 	// Insert __Binding=(...) before the closing )
 	if (ValueStr.Len() > 0 && ValueStr[ValueStr.Len() - 1] == TCHAR(')'))
 	{
@@ -750,14 +759,23 @@ bool FBindableValueBase::ImportBindingText(const TCHAR*& Buffer, TUniquePtr<FUAF
 		if (*ValueStart == TCHAR('='))
 		{
 			ValueStart++;
-			OutBinding = MakeUnique<FUAFPropertyBinding>();
-			FUAFPropertyBinding::StaticStruct()->ImportText(ValueStart,
-				OutBinding.Get(),
-				nullptr,
-				PortFlags,
-				GLog,
-				TEXT("FUAFPropertyBinding"),
-				/*bAllowNativeOverride=*/false);
+			// __Binding=() is the "no binding" marker (always exported to force RigVM
+			// parent pin updates). Treat it the same as no __Binding at all.
+			if (FCString::Strncmp(ValueStart, TEXT("()"), 2) == 0)
+			{
+				OutBinding.Reset();
+			}
+			else
+			{
+				OutBinding = MakeUnique<FUAFPropertyBinding>();
+				FUAFPropertyBinding::StaticStruct()->ImportText(ValueStart,
+					OutBinding.Get(),
+					nullptr,
+					PortFlags,
+					GLog,
+					TEXT("FUAFPropertyBinding"),
+					/*bAllowNativeOverride=*/false);
+			}
 		}
 	}
 	else
diff --git a/Engine/Plugins/Experimental/UAF/UAF/Source/UAFEditor/Private/BindableValue/UAFBindableValueCustomization.cpp b/Engine/Plugins/Experimental/UAF/UAF/Source/UAFEditor/Private/BindableValue/UAFBindableValueCustomization.cpp
--- a/Engine/Plugins/Experimental/UAF/UAF/Source/UAFEditor/Private/BindableValue/UAFBindableValueCustomization.cpp
+++ b/Engine/Plugins/Experimental/UAF/UAF/Source/UAFEditor/Private/BindableValue/UAFBindableValueCustomization.cpp
@@ -376,43 +376,19 @@ void FBindableValueCustomization::CustomizeHeader(
 		return bResult;
 	};
 
-	// Clears the binding via copy + export + SetValueFromFormattedString.
-	// We must NOT modify raw data before SetValueFromFormattedString, because the property
-	// system compares the new text against the current in-memory state. If they match (because
-	// we already mutated raw data), no change is detected and the pin default is never updated.
-	auto ClearBinding = [Handle = PropertyHandle]()
-	{
-		const FStructProperty* StructProp = CastField<FStructProperty>(Handle->GetProperty());
-		if (!StructProp)
-		{
-			return;
-		}
-
-		FString TextValue;
-		Handle->EnumerateConstRawData(
-			[&TextValue, StructProp](const void* RawData, const int32, const int32) -> bool
+	// Clears the binding by modifying raw data directly and firing a property change notification.
+	// Direct modification is used instead of SetValueFromFormattedString because the latter
+	// fails silently for property handles backed by FStructOnScope (e.g. trait properties).
+	auto ClearBindingDirect = [](TSharedPtr<IPropertyHandle> Handle)
+	{
+		Handle->NotifyPreChange();
+		Handle->EnumerateRawData(
+			[](void* RawData, const int32, const int32) -> bool
 			{
-				UScriptStruct* Struct = StructProp->Struct;
-				const int32 StructSize = Struct->GetStructureSize();
-
-				TArray<uint8> CopyBuf;
-				CopyBuf.SetNumUninitialized(StructSize);
-				Struct->InitializeStruct(CopyBuf.GetData());
-				Struct->CopyScriptStruct(CopyBuf.GetData(), RawData);
-				static_cast<FBindableValueBase*>(static_cast<void*>(CopyBuf.GetData()))->ClearBinding();
-
-				TArray<uint8> DefaultBuf;
-				DefaultBuf.SetNumUninitialized(StructSize);
-				Struct->InitializeStruct(DefaultBuf.GetData());
-
-				Struct->ExportText(TextValue, CopyBuf.GetData(), DefaultBuf.GetData(), nullptr, 0, nullptr, true);
-
-				Struct->DestroyStruct(DefaultBuf.GetData());
-				Struct->DestroyStruct(CopyBuf.GetData());
-				return false;
+				static_cast<FBindableValueBase*>(RawData)->ClearBinding();
+				return true;
 			});
-
-		Handle->SetValueFromFormattedString(TextValue);
+		Handle->NotifyPostChange(EPropertyChangeType::ValueSet);
 	};
 
 	// -----------------------------------------------------------------------
@@ -472,62 +448,36 @@ void FBindableValueCustomization::CustomizeHeader(
 		PickerArgs.CompatibleTypes = { TargetParamType };
 	}
 
-	// Helper: sets a binding by exporting from a copy and calling SetValueFromFormattedString.
-	// We must NOT modify the property handle's raw data before SetValueFromFormattedString,
-	// because the property system compares the new text against the current in-memory state.
-	// If they match (because we already mutated raw data), no change is detected and the
-	// AnimNext editor's UpdatePinDefaultValue callback never fires.
-	auto SetBindingViaText = [](TSharedPtr<IPropertyHandle> Handle, const FUAFPropertyBinding& InBinding)
+	// Helper: sets a binding by modifying raw data directly and firing a property change notification.
+	// Direct modification is used instead of SetValueFromFormattedString because the latter
+	// fails silently for property handles backed by FStructOnScope (e.g. trait properties).
+	// The notification triggers SetOnPropertyValueChangedWithData callbacks which export the
+	// full struct value (including binding via custom ExportTextItem) to the RigVM pin.
+	auto SetBindingDirect = [](TSharedPtr<IPropertyHandle> Handle, const FUAFPropertyBinding& InBinding)
 	{
-		const FStructProperty* StructProp = CastField<FStructProperty>(Handle->GetProperty());
-		if (!StructProp)
-		{
-			return;
-		}
-
-		FString TextValue;
-		Handle->EnumerateConstRawData(
-			[&TextValue, &InBinding, StructProp](const void* RawData, const int32, const int32) -> bool
+		Handle->NotifyPreChange();
+		Handle->EnumerateRawData(
+			[&InBinding](void* RawData, const int32, const int32) -> bool
 			{
-				UScriptStruct* Struct = StructProp->Struct;
-				const int32 StructSize = Struct->GetStructureSize();
-
-				// Make a copy of the current value and set the binding on the copy
-				TArray<uint8> CopyBuf;
-				CopyBuf.SetNumUninitialized(StructSize);
-				Struct->InitializeStruct(CopyBuf.GetData());
-				Struct->CopyScriptStruct(CopyBuf.GetData(), RawData);
-				static_cast<FBindableValueBase*>(static_cast<void*>(CopyBuf.GetData()))->SetBinding(InBinding);
-
-				// Export text from the copy (property handle data is untouched)
-				TArray<uint8> DefaultBuf;
-				DefaultBuf.SetNumUninitialized(StructSize);
-				Struct->InitializeStruct(DefaultBuf.GetData());
-
-				Struct->ExportText(TextValue, CopyBuf.GetData(), DefaultBuf.GetData(), nullptr, 0, nullptr, /*bAllowNativeOverride=*/true);
-
-				Struct->DestroyStruct(DefaultBuf.GetData());
-				Struct->DestroyStruct(CopyBuf.GetData());
-				return false; // first instance only
+				static_cast<FBindableValueBase*>(RawData)->SetBinding(InBinding);
+				return true;
 			});
-
-		// Text differs from current state → property system detects the change → fires notification
-		Handle->SetValueFromFormattedString(TextValue);
+		Handle->NotifyPostChange(EPropertyChangeType::ValueSet);
 	};
 
 	// Write a Variable binding.
 	PickerArgs.OnVariablePicked = FOnVariablePicked::CreateLambda(
-		[this, SetBindingViaText](const FAnimNextSoftVariableReference& SoftRef, const FAnimNextParamType& /*Type*/)
+		[this, SetBindingDirect](const FAnimNextSoftVariableReference& SoftRef, const FAnimNextParamType& /*Type*/)
 		{
 			FUAFPropertyBinding NewBinding;
 			NewBinding.SourceType     = EUAFBindingSourceType::Variable;
 			NewBinding.SourceVariable = FAnimNextVariableReference(SoftRef);
-			SetBindingViaText(PropertyHandle, NewBinding);
+			SetBindingDirect(PropertyHandle, NewBinding);
 		});
 
 	// Write a SubProperty binding.
 	PickerArgs.OnSubPropertyVariablePicked = FOnSubPropertyVariablePicked::CreateLambda(
-		[this, SetBindingViaText](
+		[this, SetBindingDirect](
 			const FAnimNextSoftVariableReference& TopVarRef,
 			const TArray<FName>&                  SubPropertyPath,
 			const FAnimNextParamType& /*LeafType*/)
@@ -539,7 +489,7 @@ void FBindableValueCustomization::CustomizeHeader(
 			{
 				NewBinding.SubPropertyPath.AddPathSegment(Seg);
 			}
-			SetBindingViaText(PropertyHandle, NewBinding);
+			SetBindingDirect(PropertyHandle, NewBinding);
 
 			FSlateApplication::Get().DismissAllMenus();
 		});
@@ -605,6 +555,17 @@ void FBindableValueCustomization::CustomizeHeader(
 
 	ConstantValueHandle = PropertyHandle->GetChildHandle(FName("ConstantValue"));
 
+	// When the constant value changes through the widget, also notify the parent FBindableXxx
+	// handle. This ensures that external callbacks (e.g. trait pin default update) that are
+	// registered on the parent handle see the change and export the full struct value.
+	if (ConstantValueHandle.IsValid())
+	{
+		ConstantValueHandle->SetOnPropertyValueChanged(FSimpleDelegate::CreateLambda([this]()
+		{
+			PropertyHandle->NotifyPostChange(EPropertyChangeType::ValueSet);
+		}));
+	}
+
 	TSharedRef<SWidget> ConstantWidget = SNullWidget::NullWidget;
 	if (ConstantValueHandle.IsValid())
 	{
@@ -872,7 +833,7 @@ void FBindableValueCustomization::CustomizeHeader(
 				]
 			]
 			.OnGetMenuContent_Lambda(
-				[this, PickerArgs, SetBindingViaText, TargetParamType]() -> TSharedRef<SWidget>
+				[this, PickerArgs, SetBindingDirect, TargetParamType]() -> TSharedRef<SWidget>
 			{
 				TSharedRef<SWidgetSwitcher> Switcher = SNew(SWidgetSwitcher);
 
@@ -897,7 +858,7 @@ void FBindableValueCustomization::CustomizeHeader(
 						return IsFunctionCompatibleWithBindableType(InHeader, TargetParamType);
 					})
 					.OnRigVMFunctionPicked_Lambda(
-						[Handle = PropertyHandle, SetBindingViaText, WeakAsset = CurrentRigVMAsset]
+						[Handle = PropertyHandle, SetBindingDirect, WeakAsset = CurrentRigVMAsset]
 						(const FRigVMGraphFunctionHeader& InHeader)
 					{
 						if (!InHeader.IsValid())
@@ -909,7 +870,7 @@ void FBindableValueCustomization::CustomizeHeader(
 						FUAFPropertyBinding NewBinding;
 						NewBinding.SourceType = EUAFBindingSourceType::Function;
 						NewBinding.SourceFunction = FAnimNextFunctionReference::FromHeader(InHeader, Asset);
-						SetBindingViaText(Handle, NewBinding);
+						SetBindingDirect(Handle, NewBinding);
 					})
 				];
 
@@ -955,9 +916,9 @@ void FBindableValueCustomization::CustomizeHeader(
 			{
 				return HasBinding() ? EVisibility::Visible : EVisibility::Collapsed;
 			})
-			.OnClicked_Lambda([ClearBinding]() -> FReply
+			.OnClicked_Lambda([this, ClearBindingDirect]() -> FReply
 			{
-				ClearBinding();
+				ClearBindingDirect(PropertyHandle);
 				return FReply::Handled();
 			})
 			[
@@ -971,7 +932,7 @@ void FBindableValueCustomization::CustomizeHeader(
 	// Right-click "Clear Binding" context menu action on the detail row.
 	InHeaderRow.AddCustomContextMenuAction(
 		FUIAction(
-			FExecuteAction::CreateLambda([ClearBinding]() { ClearBinding(); }),
+			FExecuteAction::CreateLambda([this, ClearBindingDirect]() { ClearBindingDirect(PropertyHandle); }),
 			FCanExecuteAction::CreateLambda([HasBinding]() { return HasBinding(); })),
 		LOCTEXT("ClearBindingAction",        "Clear Binding"),
 		LOCTEXT("ClearBindingActionTooltip", "Remove the variable binding from this property"),
diff --git a/Engine/Plugins/Experimental/UAF/UAF/Tests/UAF/Private/BindableValueSerializationTests.cpp b/Engine/Plugins/Experimental/UAF/UAF/Tests/UAF/Private/BindableValueSerializationTests.cpp
--- a/Engine/Plugins/Experimental/UAF/UAF/Tests/UAF/Private/BindableValueSerializationTests.cpp
+++ b/Engine/Plugins/Experimental/UAF/UAF/Tests/UAF/Private/BindableValueSerializationTests.cpp
@@ -306,7 +306,8 @@ TEST_CASE("FBindableBool.TextSerialize.NoBinding", "[UAF][unit][MustPass]")
 	FString Text;
 	FBindableBool Default;
 	Src.ExportTextItem(Text, Default, nullptr, 0, nullptr);
-	CHECK_FALSE(Text.Contains(TEXT("__Binding")));
+	// __Binding=() is always exported (empty marker forces RigVM parent pin update).
+	CHECK(Text.Contains(TEXT("__Binding=()")));
 
 	FBindableBool Dst = TextRoundTrip(Src);
 	CHECK(Dst.GetConstantValue() == true);
