UAF: Add property meta data forwarding for a bindables constant value (e.g. Clamp, UIMin, UIMax etc) and add null check for Property Utilities

#Synced-CL 52614475

diff --git a/Engine/Plugins/Experimental/UAF/UAF/Source/UAFEditor/Private/BindableValue/UAFBindableValueCustomization.cpp b/Engine/Plugins/Experimental/UAF/UAF/Source/UAFEditor/Private/BindableValue/UAFBindableValueCustomization.cpp
--- a/Engine/Plugins/Experimental/UAF/UAF/Source/UAFEditor/Private/BindableValue/UAFBindableValueCustomization.cpp
+++ b/Engine/Plugins/Experimental/UAF/UAF/Source/UAFEditor/Private/BindableValue/UAFBindableValueCustomization.cpp
@@ -420,7 +420,12 @@ void FBindableValueCustomization::CustomizeHeader(
 	// -----------------------------------------------------------------------
 
 	{
-		TArray<TWeakObjectPtr<UObject>> Objects = InCustomizationUtils.GetPropertyUtilities()->GetSelectedObjects();
+		TArray<TWeakObjectPtr<UObject>> Objects;
+		if (const TSharedPtr<IPropertyUtilities> PropertyUtilities = InCustomizationUtils.GetPropertyUtilities())
+		{
+			Objects = PropertyUtilities->GetSelectedObjects();
+		}
+		
 		for (const TWeakObjectPtr<UObject>& WeakObject : Objects)
 		{
 			UObject* Object = WeakObject.Get();
@@ -608,6 +613,27 @@ void FBindableValueCustomization::CustomizeHeader(
 	TSharedRef<SWidget> ConstantWidget = SNullWidget::NullWidget;
 	if (ConstantValueHandle.IsValid())
 	{
+		// Forwards any property level meta data to the constant value (meta = .. on UPROPERTY) 
+		if (const FProperty* OuterProperty = PropertyHandle->GetMetaDataProperty())
+		{
+			if (const TMap<FName, FString>* MetaDataMap = OuterProperty->GetMetaDataMap())
+			{
+				for (const TPair<FName, FString>& Entry : *MetaDataMap)
+				{
+					ConstantValueHandle->SetInstanceMetaData(Entry.Key, Entry.Value);
+				}
+			}
+		}
+
+		// Forwards any runtime level meta data overrides on the property to the constant value 
+		if (const TMap<FName, FString>* InstanceMap = PropertyHandle->GetInstanceMetaDataMap())
+		{
+			for (const TPair<FName, FString>& Entry : *InstanceMap)
+			{
+				ConstantValueHandle->SetInstanceMetaData(Entry.Key, Entry.Value);
+			}
+		}
+		
 		UEnum* EnumClass = nullptr;
 		if (BindableStruct == FBindableEnum::StaticStruct())
 		{
