Minor optimization, only change column visibility in UAF browser when columns actually change value

#tests editor
#Synced-CL 52651501

diff --git a/Engine/Plugins/Experimental/UAF/UAF/Source/UAFEditor/Private/Common/SUAFBrowser.cpp b/Engine/Plugins/Experimental/UAF/UAF/Source/UAFEditor/Private/Common/SUAFBrowser.cpp
--- a/Engine/Plugins/Experimental/UAF/UAF/Source/UAFEditor/Private/Common/SUAFBrowser.cpp
+++ b/Engine/Plugins/Experimental/UAF/UAF/Source/UAFEditor/Private/Common/SUAFBrowser.cpp
@@ -668,13 +668,16 @@ void SUAFBrowser::OnBrowserColumnVisibilityChanged(const TArray<FString>& Hidden
 {
 	if (UUAFEditorUserSettings* UAFSettings = GetMutableDefault<UUAFEditorUserSettings>())
 	{
-		UAFSettings->UAFBrowserHiddenColumns = HiddenColumnIds;
-		UAFSettings->SaveConfig();
+		if (UAFSettings->UAFBrowserHiddenColumns != HiddenColumnIds)
+		{
+			UAFSettings->UAFBrowserHiddenColumns = HiddenColumnIds;
+			UAFSettings->SaveConfig();
 
-		// Write to .ini
-		SaveAssetViewSettingsDelegate.ExecuteIfBound();
+			// Write to .ini
+			SaveAssetViewSettingsDelegate.ExecuteIfBound();
 
-		UAFSettings->OnBrowserHiddenColumnsChanged.Broadcast(HiddenColumnIds);
+			UAFSettings->OnBrowserHiddenColumnsChanged.Broadcast(HiddenColumnIds);
+		}
 	}
 }
 
