Use FPoseValueBundle in AnimOp value bundle evaluator

#Synced-CL 52678682

diff --git a/Engine/Plugins/Experimental/UAF/UAF/Source/UAF/Private/UAF/Attributes/AttributeBindingData.cpp b/Engine/Plugins/Experimental/UAF/UAF/Source/UAF/Private/UAF/Attributes/AttributeBindingData.cpp
--- a/Engine/Plugins/Experimental/UAF/UAF/Source/UAF/Private/UAF/Attributes/AttributeBindingData.cpp
+++ b/Engine/Plugins/Experimental/UAF/UAF/Source/UAF/Private/UAF/Attributes/AttributeBindingData.cpp
@@ -342,7 +342,7 @@ namespace UE::UAF
 				FAttributeNamedSetPtr NamedSet = *NamedSetIt;
 				const FName SetName = NamedSet->GetName();
 
-				FValueBundlePtr DefaultAttributeValues = MakeShared<FValueBundleHeap>(NamedSet);
+				FPoseValueBundlePtr DefaultAttributeValues = MakeShared<FPoseValueBundleHeap>(NamedSet);
 				DefaultAttributeValues->SetValueSpace(FValueSpace(EValueSpaceType::Local));
 
 				for (auto TypedSetIt = NamedSet->CreateTypedSetIterator(); TypedSetIt; ++TypedSetIt)
@@ -515,15 +515,15 @@ namespace UE::UAF
 		return FAttributeBindingIndex();
 	}
 
-	FValueBundlePtr FAttributeBindingData::GetDefaultAttributeValues(FName SetName) const
+	FPoseValueBundlePtr FAttributeBindingData::GetDefaultAttributeValues(FName SetName) const
 	{
-		if (const FValueBundlePtr* AttributeCollection = DefaultAttributeValuesByName.Find(SetName))
+		if (const FPoseValueBundlePtr* AttributeCollection = DefaultAttributeValuesByName.Find(SetName))
 		{
 			return *AttributeCollection;
 		}
 
 		// Set not found
-		return FValueBundlePtr();
+		return FPoseValueBundlePtr();
 	}
 
 	uint32 FAttributeBindingData::Release() const
diff --git a/Engine/Plugins/Experimental/UAF/UAF/Source/UAF/Private/UAF/Attributes/AttributeNamedSet.cpp b/Engine/Plugins/Experimental/UAF/UAF/Source/UAF/Private/UAF/Attributes/AttributeNamedSet.cpp
--- a/Engine/Plugins/Experimental/UAF/UAF/Source/UAF/Private/UAF/Attributes/AttributeNamedSet.cpp
+++ b/Engine/Plugins/Experimental/UAF/UAF/Source/UAF/Private/UAF/Attributes/AttributeNamedSet.cpp
@@ -61,7 +61,7 @@ namespace UE::UAF
 		return FAttributeNamedSetPtr(DesiredLOD, true);
 	}
 
-	FValueBundlePtr FAttributeNamedSet::GetDefaultAttributeValues() const
+	FPoseValueBundlePtr FAttributeNamedSet::GetDefaultAttributeValues() const
 	{
 		return Owner->GetDefaultAttributeValues(Name);
 	}
diff --git a/Engine/Plugins/Experimental/UAF/UAF/Source/UAF/Private/UAF/ValueRuntime/Transformers/BoneSpace.cpp b/Engine/Plugins/Experimental/UAF/UAF/Source/UAF/Private/UAF/ValueRuntime/Transformers/BoneSpace.cpp
--- a/Engine/Plugins/Experimental/UAF/UAF/Source/UAF/Private/UAF/ValueRuntime/Transformers/BoneSpace.cpp
+++ b/Engine/Plugins/Experimental/UAF/UAF/Source/UAF/Private/UAF/ValueRuntime/Transformers/BoneSpace.cpp
@@ -2,12 +2,12 @@
 
 #include "UAF/ValueRuntime/Transformers/BoneSpace.h"
 
-#include "UAF/ValueRuntime/ValueBundle.h"
+#include "UAF/ValueRuntime/PoseValueBundle.h"
 #include "UAF/Attributes/EngineAttributes.h"
 
 namespace UE::UAF::Transformers
 {
-	void FBoneSpace::LocalToComponent(const FValueBundle& InputValues, FValueBundle& OutputValues)
+	void FBoneSpace::LocalToComponent(const FPoseValueBundle& InputValues, FPoseValueBundle& OutputValues)
 	{
 		checkf(InputValues.GetValueSpace().GetType() == EValueSpaceType::Local, TEXT("Input must be in local space"));
 		checkf(!InputValues.GetValueSpace().IsAdditive(), TEXT("Input space cannot be additive"));
@@ -21,7 +21,7 @@ namespace UE::UAF::Transformers
 		OutputValues.SetValueSpace(FValueSpace(EValueSpaceType::Component));
 
 		// Modify in place
-		TBoundValueMap<FBoneTransformAnimationAttribute>* BoneTransforms = OutputValues.GetBoundValueMaps().Find<FBoneTransformAnimationAttribute>(FAttributeMappingKey::MakeFromTo<FBoneTransformAnimationAttribute>());
+		TBoundValueMap<FBoneTransformAnimationAttribute>* BoneTransforms = OutputValues.FindBoneTransforms();
 		if (BoneTransforms == nullptr)
 		{
 			// No bone transforms, nothing to convert
@@ -52,7 +52,7 @@ namespace UE::UAF::Transformers
 		}
 	}
 
-	void FBoneSpace::ComponentToLocal(const FValueBundle& InputValues, FValueBundle& OutputValues)
+	void FBoneSpace::ComponentToLocal(const FPoseValueBundle& InputValues, FPoseValueBundle& OutputValues)
 	{
 		checkf(InputValues.GetValueSpace().GetType() == EValueSpaceType::Component, TEXT("Input must be in component space"));
 		checkf(!InputValues.GetValueSpace().IsAdditive(), TEXT("Input space cannot be additive"));
@@ -66,7 +66,7 @@ namespace UE::UAF::Transformers
 		OutputValues.SetValueSpace(FValueSpace(EValueSpaceType::Local));
 
 		// Modify in place
-		TBoundValueMap<FBoneTransformAnimationAttribute>* BoneTransforms = OutputValues.GetBoundValueMaps().Find<FBoneTransformAnimationAttribute>(FAttributeMappingKey::MakeFromTo<FBoneTransformAnimationAttribute>());
+		TBoundValueMap<FBoneTransformAnimationAttribute>* BoneTransforms = OutputValues.FindBoneTransforms();
 		if (BoneTransforms == nullptr)
 		{
 			// No bone transforms, nothing to convert
@@ -97,7 +97,7 @@ namespace UE::UAF::Transformers
 		}
 	}
 
-	void FBoneSpace::LocalToMeshRotation(const FValueBundle& InputValues, FValueBundle& OutputValues)
+	void FBoneSpace::LocalToMeshRotation(const FPoseValueBundle& InputValues, FPoseValueBundle& OutputValues)
 	{
 		const FValueSpace InputSpace = InputValues.GetValueSpace();
 		checkf(InputSpace.GetType() == EValueSpaceType::Local || InputSpace.GetType() == EValueSpaceType::Mixed, TEXT("Input must be in local or mixed space"));
@@ -116,7 +116,7 @@ namespace UE::UAF::Transformers
 		OutputValues.SetValueSpace(FValueSpace(MixedSpaceFlags));
 
 		// Modify in place
-		TBoundValueMap<FBoneTransformAnimationAttribute>* BoneTransforms = OutputValues.GetBoundValueMaps().Find<FBoneTransformAnimationAttribute>(FAttributeMappingKey::MakeFromTo<FBoneTransformAnimationAttribute>());
+		TBoundValueMap<FBoneTransformAnimationAttribute>* BoneTransforms = OutputValues.FindBoneTransforms();
 		if (BoneTransforms == nullptr)
 		{
 			// No bone transforms, nothing to convert
@@ -146,7 +146,7 @@ namespace UE::UAF::Transformers
 		}
 	}
 
-	void FBoneSpace::MeshRotationToLocal(const FValueBundle& InputValues, FValueBundle& OutputValues)
+	void FBoneSpace::MeshRotationToLocal(const FPoseValueBundle& InputValues, FPoseValueBundle& OutputValues)
 	{
 		const FValueSpace InputSpace = InputValues.GetValueSpace();
 		EMixedSpaceFlags MixedSpaceFlags = InputSpace.GetMixedSpaceFlags();
@@ -174,7 +174,7 @@ namespace UE::UAF::Transformers
 		OutputValues.SetValueSpace(OutputSpace);
 
 		// Modify in place
-		TBoundValueMap<FBoneTransformAnimationAttribute>* BoneTransforms = OutputValues.GetBoundValueMaps().Find<FBoneTransformAnimationAttribute>(FAttributeMappingKey::MakeFromTo<FBoneTransformAnimationAttribute>());
+		TBoundValueMap<FBoneTransformAnimationAttribute>* BoneTransforms = OutputValues.FindBoneTransforms();
 		if (BoneTransforms == nullptr)
 		{
 			// No bone transforms, nothing to convert
@@ -205,7 +205,7 @@ namespace UE::UAF::Transformers
 		}
 	}
 
-	void FBoneSpace::LocalToRootRotation(const FValueBundle& InputValues, FValueBundle& OutputValues)
+	void FBoneSpace::LocalToRootRotation(const FPoseValueBundle& InputValues, FPoseValueBundle& OutputValues)
 	{
 		const FValueSpace InputSpace = InputValues.GetValueSpace();
 		checkf(InputSpace.GetType() == EValueSpaceType::Local || InputSpace.GetType() == EValueSpaceType::Mixed, TEXT("Input must be in local or mixed space"));
@@ -224,7 +224,7 @@ namespace UE::UAF::Transformers
 		OutputValues.SetValueSpace(FValueSpace(MixedSpaceFlags));
 
 		// Modify in place
-		TBoundValueMap<FBoneTransformAnimationAttribute>* BoneTransforms = OutputValues.GetBoundValueMaps().Find<FBoneTransformAnimationAttribute>(FAttributeMappingKey::MakeFromTo<FBoneTransformAnimationAttribute>());
+		TBoundValueMap<FBoneTransformAnimationAttribute>* BoneTransforms = OutputValues.FindBoneTransforms();
 		if (BoneTransforms == nullptr)
 		{
 			// No bone transforms, nothing to convert
@@ -256,7 +256,7 @@ namespace UE::UAF::Transformers
 		}
 	}
 
-	void FBoneSpace::RootRotationToLocal(const FValueBundle& InputValues, FValueBundle& OutputValues)
+	void FBoneSpace::RootRotationToLocal(const FPoseValueBundle& InputValues, FPoseValueBundle& OutputValues)
 	{
 		const FValueSpace InputSpace = InputValues.GetValueSpace();
 		EMixedSpaceFlags MixedSpaceFlags = InputSpace.GetMixedSpaceFlags();
@@ -284,7 +284,7 @@ namespace UE::UAF::Transformers
 		OutputValues.SetValueSpace(OutputSpace);
 
 		// Modify in place
-		TBoundValueMap<FBoneTransformAnimationAttribute>* BoneTransforms = OutputValues.GetBoundValueMaps().Find<FBoneTransformAnimationAttribute>(FAttributeMappingKey::MakeFromTo<FBoneTransformAnimationAttribute>());
+		TBoundValueMap<FBoneTransformAnimationAttribute>* BoneTransforms = OutputValues.FindBoneTransforms();
 		if (BoneTransforms == nullptr)
 		{
 			// No bone transforms, nothing to convert
@@ -317,7 +317,7 @@ namespace UE::UAF::Transformers
 		}
 	}
 
-	void FBoneSpace::LocalToMeshScale(const FValueBundle& InputValues, FValueBundle& OutputValues)
+	void FBoneSpace::LocalToMeshScale(const FPoseValueBundle& InputValues, FPoseValueBundle& OutputValues)
 	{
 		const FValueSpace InputSpace = InputValues.GetValueSpace();
 		checkf(InputSpace.GetType() == EValueSpaceType::Local || InputSpace.GetType() == EValueSpaceType::Mixed, TEXT("Input must be in local or mixed space"));
@@ -335,7 +335,7 @@ namespace UE::UAF::Transformers
 		OutputValues.SetValueSpace(FValueSpace(MixedSpaceFlags));
 
 		// Modify in place
-		TBoundValueMap<FBoneTransformAnimationAttribute>* BoneTransforms = OutputValues.GetBoundValueMaps().Find<FBoneTransformAnimationAttribute>(FAttributeMappingKey::MakeFromTo<FBoneTransformAnimationAttribute>());
+		TBoundValueMap<FBoneTransformAnimationAttribute>* BoneTransforms = OutputValues.FindBoneTransforms();
 		if (BoneTransforms == nullptr)
 		{
 			// No bone transforms, nothing to convert
@@ -366,7 +366,7 @@ namespace UE::UAF::Transformers
 		}
 	}
 
-	void FBoneSpace::MeshScaleToLocal(const FValueBundle& InputValues, FValueBundle& OutputValues)
+	void FBoneSpace::MeshScaleToLocal(const FPoseValueBundle& InputValues, FPoseValueBundle& OutputValues)
 	{
 		const FValueSpace InputSpace = InputValues.GetValueSpace();
 		EMixedSpaceFlags MixedSpaceFlags = InputSpace.GetMixedSpaceFlags();
@@ -394,7 +394,7 @@ namespace UE::UAF::Transformers
 		OutputValues.SetValueSpace(OutputSpace);
 
 		// Modify in place
-		TBoundValueMap<FBoneTransformAnimationAttribute>* BoneTransforms = OutputValues.GetBoundValueMaps().Find<FBoneTransformAnimationAttribute>(FAttributeMappingKey::MakeFromTo<FBoneTransformAnimationAttribute>());
+		TBoundValueMap<FBoneTransformAnimationAttribute>* BoneTransforms = OutputValues.FindBoneTransforms();
 		if (BoneTransforms == nullptr)
 		{
 			// No bone transforms, nothing to convert
diff --git a/Engine/Plugins/Experimental/UAF/UAF/Source/UAF/Private/UAF/ValueRuntime/ValueBundle.cpp b/Engine/Plugins/Experimental/UAF/UAF/Source/UAF/Private/UAF/ValueRuntime/ValueBundle.cpp
--- a/Engine/Plugins/Experimental/UAF/UAF/Source/UAF/Private/UAF/ValueRuntime/ValueBundle.cpp
+++ b/Engine/Plugins/Experimental/UAF/UAF/Source/UAF/Private/UAF/ValueRuntime/ValueBundle.cpp
@@ -1,6 +1,7 @@
 // Copyright Epic Games, Inc. All Rights Reserved.
 
 #include "UAF/ValueRuntime/ValueBundle.h"
+#include "UAF/ValueRuntime/PoseValueBundle.h"
 
 namespace UE::UAF
 {
diff --git a/Engine/Plugins/Experimental/UAF/UAF/Source/UAF/Public/UAF/Attributes/AttributeBindingData.h b/Engine/Plugins/Experimental/UAF/UAF/Source/UAF/Public/UAF/Attributes/AttributeBindingData.h
--- a/Engine/Plugins/Experimental/UAF/UAF/Source/UAF/Public/UAF/Attributes/AttributeBindingData.h
+++ b/Engine/Plugins/Experimental/UAF/UAF/Source/UAF/Public/UAF/Attributes/AttributeBindingData.h
@@ -6,6 +6,7 @@
 #include "Templates/RefCounting.h"
 #include "UAF/Attributes/AttributeNamedSet.h"
 #include "UAF/Attributes/AttributeTypedSet.h"
+#include "UAF/ValueRuntime/PoseValueBundle.h"
 #include "UObject/WeakObjectPtr.h"
 
 #define UE_API UAF_API
@@ -15,7 +16,6 @@ class UAbstractSkeletonSetBinding;
 
 namespace UE::UAF
 {
-	using FValueBundlePtr = TSharedPtr<class FValueBundle, ESPMode::ThreadSafe>;
 	using FAttributeBindingDataPtr = TRefCountPtr<const class FAttributeBindingData>;
 
 	// Our baked abstract skeleton binding data
@@ -93,7 +93,7 @@ namespace UE::UAF
 		[[nodiscard]] UE_API FAttributeBindingIndex FindBindingIndex(FName AttributeName, UScriptStruct* AttributeType) const;
 
 		// Returns a collection of the default attribute values for the specified named set
-		[[nodiscard]] UE_API FValueBundlePtr GetDefaultAttributeValues(FName SetName) const;
+		[[nodiscard]] UE_API FPoseValueBundlePtr GetDefaultAttributeValues(FName SetName) const;
 
 		// UE::Private::FQueryableRefCountedObject API
 		UE_API uint32 Release() const;
@@ -124,7 +124,7 @@ namespace UE::UAF
 		TMap<UScriptStruct*, TArray<FName>> AttributeNamesByType;
 
 		// Default attribute values for each named set
-		TMap<FName, FValueBundlePtr> DefaultAttributeValuesByName;
+		TMap<FName, FPoseValueBundlePtr> DefaultAttributeValuesByName;
 
 		// Only filled when executing with a skeletal mesh
 		TMap<FSkeletonPoseBoneIndex, FMeshPoseBoneIndex> SkeletonToMeshIndicesMap;
diff --git a/Engine/Plugins/Experimental/UAF/UAF/Source/UAF/Public/UAF/Attributes/AttributeNamedSet.h b/Engine/Plugins/Experimental/UAF/UAF/Source/UAF/Public/UAF/Attributes/AttributeNamedSet.h
--- a/Engine/Plugins/Experimental/UAF/UAF/Source/UAF/Public/UAF/Attributes/AttributeNamedSet.h
+++ b/Engine/Plugins/Experimental/UAF/UAF/Source/UAF/Public/UAF/Attributes/AttributeNamedSet.h
@@ -13,7 +13,7 @@ namespace UE::UAF
 	class FAttributeBindingData;
 
 	using FAttributeBindingDataPtr = TRefCountPtr<const class FAttributeBindingData>;
-	using FValueBundlePtr = TSharedPtr<class FValueBundle, ESPMode::ThreadSafe>;
+	using FPoseValueBundlePtr = TSharedPtr<class FPoseValueBundle, ESPMode::ThreadSafe>;
 	using FAttributeNamedSetPtr = TRefCountPtr<const class FAttributeNamedSet>;
 
 	class FAttributeNamedSet
@@ -97,7 +97,7 @@ namespace UE::UAF
 		[[nodiscard]] UE_FORCEINLINE_HINT FTypedSetIterator CreateTypedSetIterator() const { return FTypedSetIterator(TypedSetMap); }
 
 		// Returns a collection of the default attribute values for this named set
-		[[nodiscard]] UE_API FValueBundlePtr GetDefaultAttributeValues() const;
+		[[nodiscard]] UE_API FPoseValueBundlePtr GetDefaultAttributeValues() const;
 
 		// Reference counting API
 		UE_API void AddRef() const;
diff --git a/Engine/Plugins/Experimental/UAF/UAF/Source/UAF/Public/UAF/ValueRuntime/PoseValueBundle.h b/Engine/Plugins/Experimental/UAF/UAF/Source/UAF/Public/UAF/ValueRuntime/PoseValueBundle.h
--- a/Engine/Plugins/Experimental/UAF/UAF/Source/UAF/Public/UAF/ValueRuntime/PoseValueBundle.h
+++ b/Engine/Plugins/Experimental/UAF/UAF/Source/UAF/Public/UAF/ValueRuntime/PoseValueBundle.h
@@ -13,13 +13,16 @@ namespace UE::UAF
 	 * animation pose data (bone transforms, float curves, etc.).
 	 *
 	 * This class is not meant to be constructed directly. It adds no data members and
-	 * FValueBundle has no vtable, so it can be safely reinterpret-cast from any FValueBundle
+	 * FValueBundle has no v-table, so it can be safely reinterpret-cast from any FValueBundle
 	 * regardless of allocator (stack or heap). Use the static From(), AsMutable(), or AsConst()
 	 * helpers to obtain a FPoseValueBundle reference from an existing FValueBundle.
 	 */
 	class FPoseValueBundle : public FValueBundle
 	{
 	public:
+		// A helper which forwards the arguments to FValueBundle
+		FPoseValueBundle(const FAttributeNamedSetPtr& NamedSet, FReallocFun ReallocFun);
+
 		using FValueBundle::operator=;
 
 		// -------------------------------------------------------------------
@@ -118,22 +121,10 @@ namespace UE::UAF
 		 * Reinterpret a FValueBundle reference as a FPoseValueBundle.
 		 * Safe because FPoseValueBundle adds no data members and FValueBundle has no vtable.
 		 */
-		[[nodiscard]] static FPoseValueBundle& From(FValueBundle& Bundle)
-		{
-			return static_cast<FPoseValueBundle&>(Bundle);
-		}
+		[[nodiscard]] static FPoseValueBundle& From(FValueBundle& Bundle);
 
 		/** @copydoc From(FValueBundle&) */
-		[[nodiscard]] static const FPoseValueBundle& From(const FValueBundle& Bundle)
-		{
-			return static_cast<const FPoseValueBundle&>(Bundle);
-		}
-
-		/** Cast a mutable FValueBundle to FPoseValueBundle. */
-		[[nodiscard]] static FPoseValueBundle* AsMutable(FValueBundle& Bundle)
-		{
-			return &From(Bundle);
-		}
+		[[nodiscard]] static const FPoseValueBundle& From(const FValueBundle& Bundle);
 
 		/**
 		 * Cast a CoW reference to a mutable FPoseValueBundle, returning nullptr if the bundle is empty.
@@ -146,25 +137,53 @@ namespace UE::UAF
 		 * @endcode
 		 */
 		template<typename CoWRefType>
-		[[nodiscard]] static FPoseValueBundle* AsMutable(CoWRefType& Ref)
-		{
-			check(Ref.IsMutable());
-			return &From(Ref.GetMutable());
-		}
+		[[nodiscard]] static FPoseValueBundle* AsMutable(CoWRefType& Ref);
 
 		/**
 		 * Cast a CoW reference to a const FPoseValueBundle, returning nullptr if the bundle is empty.
 		 */
 		template<typename CoWRefType>
-		[[nodiscard]] static const FPoseValueBundle* AsConst(const CoWRefType& Ref)
-		{
-			return &From(Ref.Get());
-		}
+		[[nodiscard]] static const FPoseValueBundle* AsConst(const CoWRefType& Ref);
+	};
+
+	static_assert(sizeof(FPoseValueBundle) == sizeof(FValueBundle), "Types must have the same size to allow casting usage in From(..)");
+
+	// A heap allocated pose value bundle
+	class FPoseValueBundleHeap : public FPoseValueBundle
+	{
+	public:
+		// Constructs an empty pose bundle
+		FPoseValueBundleHeap();
+
+		// Constructs an empty pose bundle based on the specified named set
+		explicit FPoseValueBundleHeap(const FAttributeNamedSetPtr& NamedSet);
+
+		using FPoseValueBundle::operator=;
+	};
+
+	// A Mem-Stack allocated pose value bundle
+	class FPoseValueBundleStack : public FPoseValueBundle
+	{
+	public:
+		// Constructs an empty pose bundle
+		FPoseValueBundleStack();
+
+		// Constructs an empty pose bundle based on the specified named set
+		explicit FPoseValueBundleStack(const FAttributeNamedSetPtr& NamedSet);
+
+		using FPoseValueBundle::operator=;
 	};
 
+	using FPoseValueBundlePtr = TSharedPtr<FPoseValueBundle, ESPMode::ThreadSafe>;
+
 	//////////////////////////////////////////////////////////////////////////
 	// Inline implementation
 
+	inline FPoseValueBundle::FPoseValueBundle(const FAttributeNamedSetPtr& InNamedSet, FReallocFun InReallocFun)
+		: FValueBundle(InNamedSet, InReallocFun)
+	{
+	}
+
 	inline TBoundValueMap<FBoneTransformAnimationAttribute>* FPoseValueBundle::FindBoneTransforms()
 	{
 		return GetBoundValueMaps().Find<FBoneTransformAnimationAttribute>(
@@ -198,4 +217,47 @@ namespace UE::UAF
 	{
 		return FindFloatCurves() != nullptr;
 	}
+
+	inline FPoseValueBundle& FPoseValueBundle::From(FValueBundle& Bundle)
+	{
+		return static_cast<FPoseValueBundle&>(Bundle);
+	}
+
+	inline const FPoseValueBundle& FPoseValueBundle::From(const FValueBundle& Bundle)
+	{
+		return static_cast<const FPoseValueBundle&>(Bundle);
+	}
+
+	template<typename CoWRefType>
+	inline FPoseValueBundle* FPoseValueBundle::AsMutable(CoWRefType& Ref)
+	{
+		check(Ref.IsMutable());
+		return &From(Ref.GetMutable());
+	}
+
+	template<typename CoWRefType>
+	inline const FPoseValueBundle* FPoseValueBundle::AsConst(const CoWRefType& Ref)
+	{
+		return &From(Ref.Get());
+	}
+
+	inline FPoseValueBundleHeap::FPoseValueBundleHeap()
+		: FPoseValueBundle(nullptr, &FAllocatorTypeTrait<FDefaultAllocator>::Realloc)
+	{
+	}
+
+	inline FPoseValueBundleHeap::FPoseValueBundleHeap(const FAttributeNamedSetPtr& NamedSet)
+		: FPoseValueBundle(NamedSet, &FAllocatorTypeTrait<FDefaultAllocator>::Realloc)
+	{
+	}
+
+	inline FPoseValueBundleStack::FPoseValueBundleStack()
+		: FPoseValueBundle(nullptr, &FAllocatorTypeTrait<TMemStackAllocator<>>::Realloc)
+	{
+	}
+
+	inline FPoseValueBundleStack::FPoseValueBundleStack(const FAttributeNamedSetPtr& NamedSet)
+		: FPoseValueBundle(NamedSet, &FAllocatorTypeTrait<TMemStackAllocator<>>::Realloc)
+	{
+	}
 } // namespace UE::UAF
diff --git a/Engine/Plugins/Experimental/UAF/UAF/Source/UAF/Public/UAF/ValueRuntime/Transformers/BoneSpace.h b/Engine/Plugins/Experimental/UAF/UAF/Source/UAF/Public/UAF/ValueRuntime/Transformers/BoneSpace.h
--- a/Engine/Plugins/Experimental/UAF/UAF/Source/UAF/Public/UAF/ValueRuntime/Transformers/BoneSpace.h
+++ b/Engine/Plugins/Experimental/UAF/UAF/Source/UAF/Public/UAF/ValueRuntime/Transformers/BoneSpace.h
@@ -8,7 +8,7 @@
 
 namespace UE::UAF
 {
-	class FValueBundle;
+	class FPoseValueBundle;
 }
 
 namespace UE::UAF::Transformers
@@ -19,31 +19,31 @@ namespace UE::UAF::Transformers
 	{
 		// Converts bone transforms between Local Space and Component Space
 		// These functions support outputting to their input in place
-		static UE_API void LocalToComponent(const FValueBundle& InputValues, FValueBundle& OutputValues);
-		static UE_API void ComponentToLocal(const FValueBundle& InputValues, FValueBundle& OutputValues);
+		static UE_API void LocalToComponent(const FPoseValueBundle& InputValues, FPoseValueBundle& OutputValues);
+		static UE_API void ComponentToLocal(const FPoseValueBundle& InputValues, FPoseValueBundle& OutputValues);
 
 		// Converts bone transforms between Local Space and Mesh Rotation Space
 		// Mesh Rotation Space assumes that bone transform scales are uniform and will yield incorrect results
 		// if non-uniform (aka 3D) scale is used. It transforms the rotations into component space by ignoring
 		// translation and scale. The output values will thus be in Mixed Space.
 		// These functions support outputting to their input in place
-		static UE_API void LocalToMeshRotation(const FValueBundle& InputValues, FValueBundle& OutputValues);
-		static UE_API void MeshRotationToLocal(const FValueBundle& InputValues, FValueBundle& OutputValues);
+		static UE_API void LocalToMeshRotation(const FPoseValueBundle& InputValues, FPoseValueBundle& OutputValues);
+		static UE_API void MeshRotationToLocal(const FPoseValueBundle& InputValues, FPoseValueBundle& OutputValues);
 
 		// Converts bone transforms between Local Space and Root Rotation Space
 		// Root Rotation Space is similar to Mesh Rotation space except that the root bone transform is ignored
 		// and assumed to be the identity transform (root corrections are ignored).
 		// These functions support outputting to their input in place
-		static UE_API void LocalToRootRotation(const FValueBundle& InputValues, FValueBundle& OutputValues);
-		static UE_API void RootRotationToLocal(const FValueBundle& InputValues, FValueBundle& OutputValues);
+		static UE_API void LocalToRootRotation(const FPoseValueBundle& InputValues, FPoseValueBundle& OutputValues);
+		static UE_API void RootRotationToLocal(const FPoseValueBundle& InputValues, FPoseValueBundle& OutputValues);
 
 		// Converts bone transforms between Local Space and Mesh Scale Space
 		// Mesh Scale Space assumes that bone transform scales are uniform and will yield incorrect results
 		// if non-uniform (aka 3D) scale is used. It transforms the scales into component space by ignoring
 		// rotation and translation. The output values will thus be in a Mixed Space.
 		// These functions support outputting to their input in place
-		static UE_API void LocalToMeshScale(const FValueBundle& InputValues, FValueBundle& OutputValues);
-		static UE_API void MeshScaleToLocal(const FValueBundle& InputValues, FValueBundle& OutputValues);
+		static UE_API void LocalToMeshScale(const FPoseValueBundle& InputValues, FPoseValueBundle& OutputValues);
+		static UE_API void MeshScaleToLocal(const FPoseValueBundle& InputValues, FPoseValueBundle& OutputValues);
 	};
 }
 
diff --git a/Engine/Plugins/Experimental/UAF/UAFAnimGraph/Source/UAFAnimGraph/Private/EvaluationVM/Tasks/ApplyAdditiveKeyframe.cpp b/Engine/Plugins/Experimental/UAF/UAFAnimGraph/Source/UAFAnimGraph/Private/EvaluationVM/Tasks/ApplyAdditiveKeyframe.cpp
--- a/Engine/Plugins/Experimental/UAF/UAFAnimGraph/Source/UAFAnimGraph/Private/EvaluationVM/Tasks/ApplyAdditiveKeyframe.cpp
+++ b/Engine/Plugins/Experimental/UAF/UAFAnimGraph/Source/UAFAnimGraph/Private/EvaluationVM/Tasks/ApplyAdditiveKeyframe.cpp
@@ -84,14 +84,14 @@ void FAnimNextApplyAdditiveKeyframeTask::Execute(UE::UAF::FEvaluationVM& VM) con
 		// Modify our base input in-place
 		if (bConvertBaseToMeshSpace)
 		{
-			Transformers::FBoneSpace::LocalToMeshRotation(*BaseCollection, *BaseCollection);
+			Transformers::FBoneSpace::LocalToMeshRotation(FPoseValueBundle::From(*BaseCollection), FPoseValueBundle::From(*BaseCollection));
 		}
 
 		Transformers::FApplyAdditiveSpace::Apply(VM.GetTransformerMap(), *BaseCollection, *AdditiveCollection, AdditiveWeight, *BaseCollection);
 
 		if (bConvertBaseToMeshSpace)
 		{
-			Transformers::FBoneSpace::MeshRotationToLocal(*BaseCollection, *BaseCollection);
+			Transformers::FBoneSpace::MeshRotationToLocal(FPoseValueBundle::From(*BaseCollection), FPoseValueBundle::From(*BaseCollection));
 		}
 
 		VM.PushValue(ATTRIBUTE_STACK_NAME, MoveTemp(BaseCollection));
diff --git a/Engine/Plugins/Experimental/UAF/UAFAnimGraph/Source/UAFAnimGraph/Private/EvaluationVM/Tasks/ApplyAdditiveKeyframePerBone.cpp b/Engine/Plugins/Experimental/UAF/UAFAnimGraph/Source/UAFAnimGraph/Private/EvaluationVM/Tasks/ApplyAdditiveKeyframePerBone.cpp
--- a/Engine/Plugins/Experimental/UAF/UAFAnimGraph/Source/UAFAnimGraph/Private/EvaluationVM/Tasks/ApplyAdditiveKeyframePerBone.cpp
+++ b/Engine/Plugins/Experimental/UAF/UAFAnimGraph/Source/UAFAnimGraph/Private/EvaluationVM/Tasks/ApplyAdditiveKeyframePerBone.cpp
@@ -138,14 +138,14 @@ void FUAFApplyAdditiveKeyframePerBoneTask::Execute(UE::UAF::FEvaluationVM& VM) c
 
 		if (bConvertBaseToMeshSpace)
 		{
-			Transformers::FBoneSpace::LocalToMeshRotation(*BaseInput, *BaseInput);
+			Transformers::FBoneSpace::LocalToMeshRotation(FPoseValueBundle::From(*BaseInput), FPoseValueBundle::From(*BaseInput));
 		}
 
 		Transformers::FApplyAdditiveSpace::Apply(VM.GetTransformerMap(), *BaseInput, *AdditiveInput, PerValueWeights, AdditiveWeight, *BaseInput);
 
 		if (bConvertBaseToMeshSpace)
 		{
-			Transformers::FBoneSpace::MeshRotationToLocal(*BaseInput, *BaseInput);
+			Transformers::FBoneSpace::MeshRotationToLocal(FPoseValueBundle::From(*BaseInput), FPoseValueBundle::From(*BaseInput));
 		}
 
 		VM.PushValue(ATTRIBUTE_STACK_NAME, MoveTemp(BaseInput));
@@ -392,14 +392,14 @@ void FUAFApplyAdditiveKeyframeWithBlendMaskTask::Execute(UE::UAF::FEvaluationVM&
 
 		if (bConvertBaseToMeshSpace)
 		{
-			Transformers::FBoneSpace::LocalToMeshRotation(*BaseInput, *BaseInput);
+			Transformers::FBoneSpace::LocalToMeshRotation(FPoseValueBundle::From(*BaseInput), FPoseValueBundle::From(*BaseInput));
 		}
 
 		Transformers::FApplyAdditiveSpace::Apply(VM.GetTransformerMap(), *BaseInput, *AdditiveInput, PerValueWeights, AdditiveWeight, *BaseInput);
 
 		if (bConvertBaseToMeshSpace)
 		{
-			Transformers::FBoneSpace::MeshRotationToLocal(*BaseInput, *BaseInput);
+			Transformers::FBoneSpace::MeshRotationToLocal(FPoseValueBundle::From(*BaseInput), FPoseValueBundle::From(*BaseInput));
 		}
 
 		VM.PushValue(ATTRIBUTE_STACK_NAME, MoveTemp(BaseInput));
diff --git a/Engine/Plugins/Experimental/UAF/UAFAnimGraph/Source/UAFAnimGraph/Private/EvaluationVM/Tasks/ConvertRotationsToLocalSpace.cpp b/Engine/Plugins/Experimental/UAF/UAFAnimGraph/Source/UAFAnimGraph/Private/EvaluationVM/Tasks/ConvertRotationsToLocalSpace.cpp
--- a/Engine/Plugins/Experimental/UAF/UAFAnimGraph/Source/UAFAnimGraph/Private/EvaluationVM/Tasks/ConvertRotationsToLocalSpace.cpp
+++ b/Engine/Plugins/Experimental/UAF/UAFAnimGraph/Source/UAFAnimGraph/Private/EvaluationVM/Tasks/ConvertRotationsToLocalSpace.cpp
@@ -36,7 +36,7 @@ void FAnimNextConvertRotationsToLocalSpaceTask::Execute(UE::UAF::FEvaluationVM&
 				break;
 			}
 
-			Transformers::FBoneSpace::MeshRotationToLocal(**Collection, **Collection);
+			Transformers::FBoneSpace::MeshRotationToLocal(FPoseValueBundle::From(**Collection), FPoseValueBundle::From(**Collection));
 		}
 	}
 	else
diff --git a/Engine/Plugins/Experimental/UAF/UAFAnimGraph/Source/UAFAnimGraph/Private/EvaluationVM/Tasks/ConvertRotationsToMeshSpace.cpp b/Engine/Plugins/Experimental/UAF/UAFAnimGraph/Source/UAFAnimGraph/Private/EvaluationVM/Tasks/ConvertRotationsToMeshSpace.cpp
--- a/Engine/Plugins/Experimental/UAF/UAFAnimGraph/Source/UAFAnimGraph/Private/EvaluationVM/Tasks/ConvertRotationsToMeshSpace.cpp
+++ b/Engine/Plugins/Experimental/UAF/UAFAnimGraph/Source/UAFAnimGraph/Private/EvaluationVM/Tasks/ConvertRotationsToMeshSpace.cpp
@@ -36,7 +36,7 @@ void FAnimNextConvertRotationsToMeshSpaceTask::Execute(UE::UAF::FEvaluationVM& V
 				break;
 			}
 
-			Transformers::FBoneSpace::LocalToMeshRotation(**Collection, **Collection);
+			Transformers::FBoneSpace::LocalToMeshRotation(FPoseValueBundle::From(**Collection), FPoseValueBundle::From(**Collection));
 		}
 	}
 	else
diff --git a/Engine/Plugins/Experimental/UAF/UAFAnimNode/Source/UAFAnimNode/Private/UAF/AnimNodeCore/UAFAnimNodeUpdate.cpp b/Engine/Plugins/Experimental/UAF/UAFAnimNode/Source/UAFAnimNode/Private/UAF/AnimNodeCore/UAFAnimNodeUpdate.cpp
--- a/Engine/Plugins/Experimental/UAF/UAFAnimNode/Source/UAFAnimNode/Private/UAF/AnimNodeCore/UAFAnimNodeUpdate.cpp
+++ b/Engine/Plugins/Experimental/UAF/UAFAnimNode/Source/UAFAnimNode/Private/UAF/AnimNodeCore/UAFAnimNodeUpdate.cpp
@@ -309,7 +309,7 @@ namespace UE::UAF
 
 		FUAFAnimOpValueEvaluator Evaluator(Binding, SkeletalMesh, SetName, LOD);
 
-		TUAFStack<FValueBundleCoWRef>& EvaluationStack = Evaluator.GetEvaluationStack();
+		TUAFStack<FPoseValueBundleCoWRef>& EvaluationStack = Evaluator.GetEvaluationStack();
 
 		for (FUAFAnimOp* AnimOp : AnimOps)
 		{
@@ -323,7 +323,7 @@ namespace UE::UAF
 		}
 
 		const FValueSpace ExpectedValueSpace(EValueSpaceType::Local);
-		const FValueBundle* BundleToConvert = nullptr;
+		const FPoseValueBundle* BundleToConvert = nullptr;
 
 		if (!EvaluationStack.IsEmpty()) [[likely]]
 		{
@@ -332,8 +332,8 @@ namespace UE::UAF
 				UE_LOGF(LogAnimation, Warning, "Value evaluation produced too many results, first one will be used.");
 			}
 
-			const FValueBundleCoWRef* OutputBundleRef = EvaluationStack.Peek();
-			const FValueBundle& OutputBundle = OutputBundleRef->Get();
+			const FPoseValueBundleCoWRef* OutputBundleRef = EvaluationStack.Peek();
+			const FPoseValueBundle& OutputBundle = OutputBundleRef->Get();
 			
 			if (!OutputBundle.IsEmpty()) [[likely]]
 			{
@@ -352,7 +352,7 @@ namespace UE::UAF
 			}
 		}
 
-		FValueBundleStack DefaultValueBundle;
+		FPoseValueBundleStack DefaultValueBundle;
 		if (BundleToConvert == nullptr)
 		{
 			DefaultValueBundle.InitWithValueSpace(ExpectedValueSpace);
diff --git a/Engine/Plugins/Experimental/UAF/UAFAnimNode/Source/UAFAnimNode/Private/UAF/AnimOpCore/UAFAnimOpArrayView.cpp b/Engine/Plugins/Experimental/UAF/UAFAnimNode/Source/UAFAnimNode/Private/UAF/AnimOpCore/UAFAnimOpArrayView.cpp
--- a/Engine/Plugins/Experimental/UAF/UAFAnimNode/Source/UAFAnimNode/Private/UAF/AnimOpCore/UAFAnimOpArrayView.cpp
+++ b/Engine/Plugins/Experimental/UAF/UAFAnimNode/Source/UAFAnimNode/Private/UAF/AnimOpCore/UAFAnimOpArrayView.cpp
@@ -16,7 +16,7 @@ namespace UE::UAF
 
 	void FUAFAnimOpArrayView::EvaluateValues(FUAFAnimOpValueEvaluator& Evaluator)
 	{
-		TUAFStack<FValueBundleCoWRef>& EvaluationStack = Evaluator.GetEvaluationStack();
+		TUAFStack<FPoseValueBundleCoWRef>& EvaluationStack = Evaluator.GetEvaluationStack();
 
 		for (FUAFAnimOp* AnimOp : AnimOps)
 		{
diff --git a/Engine/Plugins/Experimental/UAF/UAFAnimNode/Source/UAFAnimNode/Private/UAF/AnimOps/UAFApplyAdditiveAnimOp.cpp b/Engine/Plugins/Experimental/UAF/UAFAnimNode/Source/UAFAnimNode/Private/UAF/AnimOps/UAFApplyAdditiveAnimOp.cpp
--- a/Engine/Plugins/Experimental/UAF/UAFAnimNode/Source/UAFAnimNode/Private/UAF/AnimOps/UAFApplyAdditiveAnimOp.cpp
+++ b/Engine/Plugins/Experimental/UAF/UAFAnimNode/Source/UAFAnimNode/Private/UAF/AnimOps/UAFApplyAdditiveAnimOp.cpp
@@ -22,11 +22,11 @@ namespace UE::UAF
 		const FUAFAnimOpValueEvaluationContext& AnimOpEvaluationContext = Evaluator.GetActiveEvaluationContext();
 
 		// Pop our inputs - the additive being on top and base the next underneath 
-		FValueBundleCoWRef AdditiveInput = Evaluator.GetEvaluationStack().Pop();
-		FValueBundleCoWRef BaseInput = Evaluator.GetEvaluationStack().Pop();
+		FPoseValueBundleCoWRef AdditiveInput = Evaluator.GetEvaluationStack().Pop();
+		FPoseValueBundleCoWRef BaseInput = Evaluator.GetEvaluationStack().Pop();
 
-		const FValueBundle& Additive = AdditiveInput.Get();
-		FValueBundle& Base = BaseInput.GetMutable();
+		const FPoseValueBundle& Additive = AdditiveInput.Get();
+		FPoseValueBundle& Base = BaseInput.GetMutable();
 		
 		// If the additive is not relevant or empty, we simply fall back to the base pose
 		// Both scenarios are valid and do not necessarily indicate a setup error so we will not raise a waring 
@@ -100,7 +100,7 @@ namespace UE::UAF
 		OptionalBlendMask = InBlendMask;
 	}
 
-	void FUAFApplyAdditiveAnimOp::BuildPerValueWeights(const FUAFAnimOpValueEvaluator& Evaluator, const FValueBundle& BaseInput, FValueBundleStack& OutPerValueWeights)
+	void FUAFApplyAdditiveAnimOp::BuildPerValueWeights(const FUAFAnimOpValueEvaluator& Evaluator, const FPoseValueBundle& BaseInput, FValueBundleStack& OutPerValueWeights)
 	{
 		if (OptionalBlendMask == nullptr)
 		{
diff --git a/Engine/Plugins/Experimental/UAF/UAFAnimNode/Source/UAFAnimNode/Private/UAF/AnimOps/UAFBlendTwoAnimOp.cpp b/Engine/Plugins/Experimental/UAF/UAFAnimNode/Source/UAFAnimNode/Private/UAF/AnimOps/UAFBlendTwoAnimOp.cpp
--- a/Engine/Plugins/Experimental/UAF/UAFAnimNode/Source/UAFAnimNode/Private/UAF/AnimOps/UAFBlendTwoAnimOp.cpp
+++ b/Engine/Plugins/Experimental/UAF/UAFAnimNode/Source/UAFAnimNode/Private/UAF/AnimOps/UAFBlendTwoAnimOp.cpp
@@ -38,11 +38,11 @@ namespace UE::UAF
 		const FUAFAnimOpValueEvaluationContext& AnimOpEvaluationContext = Evaluator.GetActiveEvaluationContext();
 
 		// Pop our inputs
-		FValueBundleCoWRef InputRefB = Evaluator.GetEvaluationStack().Pop();
-		FValueBundleCoWRef InputRefA = Evaluator.GetEvaluationStack().Pop();
+		FPoseValueBundleCoWRef InputRefB = Evaluator.GetEvaluationStack().Pop();
+		FPoseValueBundleCoWRef InputRefA = Evaluator.GetEvaluationStack().Pop();
 
-		const FValueBundle& InputA = InputRefA.Get();
-		const FValueBundle& InputB = InputRefB.Get();
+		const FPoseValueBundle& InputA = InputRefA.Get();
+		const FPoseValueBundle& InputB = InputRefB.Get();
 
 		const float WeightOfPoseB = GetInterpolationAlpha(InputA, InputB);
 		ActualInterpolationAlpha = WeightOfPoseB;
@@ -89,7 +89,7 @@ namespace UE::UAF
 			// Named sets must match as it ensures that our inputs have the same sizes/shapes
 			check(InputA.GetNamedSet() == InputB.GetNamedSet());
 
-			FValueBundleCoWRef& OutputRef = FindMutableCoWRef(InputRefB, InputRefA);
+			FPoseValueBundleCoWRef& OutputRef = FindMutableCoWRef(InputRefB, InputRefA);
 
 			Transformers::FInterpolate::Apply(Evaluator.GetTransformerMap(), InputA, InputB, WeightOfPoseB, OutputRef.GetMutable());
 			Transformers::FSanitize::Apply(Evaluator.GetTransformerMap(), OutputRef.GetMutable());
@@ -98,15 +98,14 @@ namespace UE::UAF
 		}
 	}
 
-	float FUAFBlendTwoAnimOp::GetInterpolationAlpha(const FValueBundle& InputA, const FValueBundle& InputB) const
+	float FUAFBlendTwoAnimOp::GetInterpolationAlpha(const FPoseValueBundle& InputA, const FPoseValueBundle& InputB) const
 	{
 		float Alpha = DesiredInterpolationAlpha;
 
 		if (AlphaSourceCurveName != NAME_None && AlphaCurveInputIndex != INDEX_NONE)
 		{
-			const FAttributeMappingKey MappingKey = FAttributeMappingKey::MakeFromTo<FFloatAnimationAttribute>();
-			const FValueBundle& AlphaCurveCollection = AlphaCurveInputIndex == 0 ? InputA : InputB;
-			if (const TBoundValueMap<FFloatAnimationAttribute>* FloatMap = AlphaCurveCollection.GetBoundValueMaps().Find<FFloatAnimationAttribute>(MappingKey))
+			const FPoseValueBundle& AlphaCurveCollection = AlphaCurveInputIndex == 0 ? InputA : InputB;
+			if (const TBoundValueMap<FFloatAnimationAttribute>* FloatMap = AlphaCurveCollection.FindFloatCurves())
 			{
 				const FAttributeTypedSetPtr& TypedSet = FloatMap->GetTypedSet();
 				if (FAttributeSetIndex CurveIndex = TypedSet->FindIndex(AlphaSourceCurveName))
diff --git a/Engine/Plugins/Experimental/UAF/UAFAnimNode/Source/UAFAnimNode/Private/UAF/AnimOps/UAFBlendTwoPerValueAnimOp.cpp b/Engine/Plugins/Experimental/UAF/UAFAnimNode/Source/UAFAnimNode/Private/UAF/AnimOps/UAFBlendTwoPerValueAnimOp.cpp
--- a/Engine/Plugins/Experimental/UAF/UAFAnimNode/Source/UAFAnimNode/Private/UAF/AnimOps/UAFBlendTwoPerValueAnimOp.cpp
+++ b/Engine/Plugins/Experimental/UAF/UAFAnimNode/Source/UAFAnimNode/Private/UAF/AnimOps/UAFBlendTwoPerValueAnimOp.cpp
@@ -36,11 +36,11 @@ namespace UE::UAF
 		const FUAFAnimOpValueEvaluationContext& AnimOpEvaluationContext = Evaluator.GetActiveEvaluationContext();
 
 		// Pop our inputs
-		FValueBundleCoWRef InputRefB = Evaluator.GetEvaluationStack().Pop();
-		FValueBundleCoWRef InputRefA = Evaluator.GetEvaluationStack().Pop();
+		FPoseValueBundleCoWRef InputRefB = Evaluator.GetEvaluationStack().Pop();
+		FPoseValueBundleCoWRef InputRefA = Evaluator.GetEvaluationStack().Pop();
 
-		const FValueBundle& InputA = InputRefA.Get();
-		const FValueBundle& InputB = InputRefB.Get();
+		const FPoseValueBundle& InputA = InputRefA.Get();
+		const FPoseValueBundle& InputB = InputRefB.Get();
 
 		const float GlobalWeightB = FAlphaBlend::AlphaToBlendOption(LinearGlobalWeightB, BlendOption);
 		const float GlobalWeightA = 1.0 - GlobalWeightB;
@@ -87,7 +87,7 @@ namespace UE::UAF
 			// Named sets must match as it ensures that our inputs have the same sizes/shapes
 			check(InputA.GetNamedSet() == InputB.GetNamedSet());
 
-			FValueBundleCoWRef& OutputRef = FindMutableCoWRef(InputRefB, InputRefA);
+			FPoseValueBundleCoWRef& OutputRef = FindMutableCoWRef(InputRefB, InputRefA);
 
 			FValueBundleStack PerValueWeights(Evaluator.GetActiveNamedSet());
 			PerValueWeights.InitWithValueSpace(FValueSpace(EValueSpaceType::Local));
diff --git a/Engine/Plugins/Experimental/UAF/UAFAnimNode/Source/UAFAnimNode/Private/UAF/AnimOps/UAFDecompressAnimSequenceAnimOp.cpp b/Engine/Plugins/Experimental/UAF/UAFAnimNode/Source/UAFAnimNode/Private/UAF/AnimOps/UAFDecompressAnimSequenceAnimOp.cpp
--- a/Engine/Plugins/Experimental/UAF/UAFAnimNode/Source/UAFAnimNode/Private/UAF/AnimOps/UAFDecompressAnimSequenceAnimOp.cpp
+++ b/Engine/Plugins/Experimental/UAF/UAFAnimNode/Source/UAFAnimNode/Private/UAF/AnimOps/UAFDecompressAnimSequenceAnimOp.cpp
@@ -45,7 +45,7 @@ namespace UE::UAF
 	{
 		const FUAFAnimOpValueEvaluationContext& AnimOpEvaluationContext = Evaluator.GetActiveEvaluationContext();
 
-		FValueBundleStack ValueBundle(AnimOpEvaluationContext.GetNamedSet());
+		FPoseValueBundleStack ValueBundle(AnimOpEvaluationContext.GetNamedSet());
 
 		const UAnimSequence* AnimSequencePtr = AnimSequence.Get();
 		if (AnimSequencePtr && AnimSequencePtr->GetSkeleton())
@@ -98,7 +98,7 @@ namespace UE::UAF
 			}
 		}
 
-		Evaluator.GetEvaluationStack().Push(FValueBundleCoWRef::MakeFrom(MoveTemp(ValueBundle)));
+		Evaluator.GetEvaluationStack().Push(FPoseValueBundleCoWRef::MakeFrom(MoveTemp(ValueBundle)));
 	}
 
 	void FUAFDecompressAnimSequenceAnimOp::EvaluateNotifies(FUAFAnimOpNotifyEvaluator& Evaluator)
diff --git a/Engine/Plugins/Experimental/UAF/UAFAnimNode/Source/UAFAnimNode/Private/UAF/AnimOps/UAFMakeDynamicAdditiveAnimOp.cpp b/Engine/Plugins/Experimental/UAF/UAFAnimNode/Source/UAFAnimNode/Private/UAF/AnimOps/UAFMakeDynamicAdditiveAnimOp.cpp
--- a/Engine/Plugins/Experimental/UAF/UAFAnimNode/Source/UAFAnimNode/Private/UAF/AnimOps/UAFMakeDynamicAdditiveAnimOp.cpp
+++ b/Engine/Plugins/Experimental/UAF/UAFAnimNode/Source/UAFAnimNode/Private/UAF/AnimOps/UAFMakeDynamicAdditiveAnimOp.cpp
@@ -21,8 +21,8 @@ namespace UE::UAF
 	{
 		// Stack order: Base (bottom), Source (top)
 		// Pop Source first (top of stack), then Base (underneath)
-		FValueBundleCoWRef SourceInput = Evaluator.GetEvaluationStack().Pop();
-		FValueBundleCoWRef BaseInput = Evaluator.GetEvaluationStack().Pop();
+		FPoseValueBundleCoWRef SourceInput = Evaluator.GetEvaluationStack().Pop();
+		FPoseValueBundleCoWRef BaseInput = Evaluator.GetEvaluationStack().Pop();
 
 		const FValueSpace AdditiveSpace = bMeshSpaceAdditive
 			? FValueSpace(EMixedSpaceFlags::MeshRotation, true)
@@ -53,9 +53,9 @@ namespace UE::UAF
 
 		if (bInvalidSetup)
 		{
-			FValueBundleStack Output(Evaluator.GetActiveNamedSet());
+			FPoseValueBundleStack Output(Evaluator.GetActiveNamedSet());
 			Output.InitWithValueSpace(AdditiveSpace);
-			Evaluator.GetEvaluationStack().Push(FValueBundleCoWRef::MakeFrom(MoveTemp(Output)));
+			Evaluator.GetEvaluationStack().Push(FPoseValueBundleCoWRef::MakeFrom(MoveTemp(Output)));
 			return;
 		}
 
@@ -68,12 +68,12 @@ namespace UE::UAF
 			Transformers::FBoneSpace::LocalToMeshRotation(SourceInput.GetMutable(), SourceInput.GetMutable());
 		}
 
-		FValueBundleStack Output(Evaluator.GetActiveNamedSet());
+		FPoseValueBundleStack Output(Evaluator.GetActiveNamedSet());
 		Output.InitWithValueSpace(AdditiveSpace);
 
 		Transformers::FMakeAdditiveSpace::Apply(Evaluator.GetTransformerMap(), BaseInput.Get(), SourceInput.Get(), Output);
 
-		Evaluator.GetEvaluationStack().Push(FValueBundleCoWRef::MakeFrom(MoveTemp(Output)));
+		Evaluator.GetEvaluationStack().Push(FPoseValueBundleCoWRef::MakeFrom(MoveTemp(Output)));
 	}
 
 	void FUAFMakeDynamicAdditiveAnimOp::SetMeshSpaceAdditive(const bool bInMeshSpaceAdditive)
diff --git a/Engine/Plugins/Experimental/UAF/UAFAnimNode/Source/UAFAnimNode/Private/UAF/AnimOps/UAFNullAnimOp.cpp b/Engine/Plugins/Experimental/UAF/UAFAnimNode/Source/UAFAnimNode/Private/UAF/AnimOps/UAFNullAnimOp.cpp
--- a/Engine/Plugins/Experimental/UAF/UAFAnimNode/Source/UAFAnimNode/Private/UAF/AnimOps/UAFNullAnimOp.cpp
+++ b/Engine/Plugins/Experimental/UAF/UAFAnimNode/Source/UAFAnimNode/Private/UAF/AnimOps/UAFNullAnimOp.cpp
@@ -27,7 +27,7 @@ namespace UE::UAF
 	{
 		const FUAFAnimOpValueEvaluationContext& AnimOpEvaluationContext = Evaluator.GetActiveEvaluationContext();
 
-		FValueBundleStack EmptyBundle(AnimOpEvaluationContext.GetNamedSet());
-		Evaluator.GetEvaluationStack().Push(FValueBundleCoWRef::MakeFrom(MoveTemp(EmptyBundle)));
+		FPoseValueBundleStack EmptyBundle(AnimOpEvaluationContext.GetNamedSet());
+		Evaluator.GetEvaluationStack().Push(FPoseValueBundleCoWRef::MakeFrom(MoveTemp(EmptyBundle)));
 	}
 }
diff --git a/Engine/Plugins/Experimental/UAF/UAFAnimNode/Source/UAFAnimNode/Private/UAF/AnimOps/UAFOffsetRootBoneAnimOp.cpp b/Engine/Plugins/Experimental/UAF/UAFAnimNode/Source/UAFAnimNode/Private/UAF/AnimOps/UAFOffsetRootBoneAnimOp.cpp
--- a/Engine/Plugins/Experimental/UAF/UAFAnimNode/Source/UAFAnimNode/Private/UAF/AnimOps/UAFOffsetRootBoneAnimOp.cpp
+++ b/Engine/Plugins/Experimental/UAF/UAFAnimNode/Source/UAFAnimNode/Private/UAF/AnimOps/UAFOffsetRootBoneAnimOp.cpp
@@ -80,7 +80,7 @@ void FUAFOffsetRootBoneAnimOp::EvaluateValues(FUAFAnimOpValueEvaluator& Evaluato
 
 	// Access the value bundle from the evaluation context
 	const FUAFAnimOpValueEvaluationContext& AnimOpEvaluationContext = Evaluator.GetActiveEvaluationContext();
-	FValueBundleCoWRef* ValueBundle = Evaluator.GetEvaluationStack().PeekMutable();
+	FPoseValueBundleCoWRef* ValueBundle = Evaluator.GetEvaluationStack().PeekMutable();
 	check (ValueBundle);
 	
 	FAttributeNamedSetPtr NamedSet = ValueBundle->Get().GetNamedSet();
@@ -103,8 +103,7 @@ void FUAFOffsetRootBoneAnimOp::EvaluateValues(FUAFAnimOpValueEvaluator& Evaluato
 		return;
 	}
 
-	FAttributeMappingKey BoneMappingKey = FAttributeMappingKey::MakeFromTo<FBoneTransformAnimationAttribute>();
-	TBoundValueMap<FBoneTransformAnimationAttribute>* BoneMap = ValueBundle->GetMutable().GetBoundValueMaps().Find<FBoneTransformAnimationAttribute>(BoneMappingKey);
+	TBoundValueMap<FBoneTransformAnimationAttribute>* BoneMap = ValueBundle->GetMutable().FindBoneTransforms();
 	if (!BoneMap)
 	{
 		return;
diff --git a/Engine/Plugins/Experimental/UAF/UAFAnimNode/Source/UAFAnimNode/Private/UAF/AnimOps/UAFSteeringAnimOp.cpp b/Engine/Plugins/Experimental/UAF/UAFAnimNode/Source/UAFAnimNode/Private/UAF/AnimOps/UAFSteeringAnimOp.cpp
--- a/Engine/Plugins/Experimental/UAF/UAFAnimNode/Source/UAFAnimNode/Private/UAF/AnimOps/UAFSteeringAnimOp.cpp
+++ b/Engine/Plugins/Experimental/UAF/UAFAnimNode/Source/UAFAnimNode/Private/UAF/AnimOps/UAFSteeringAnimOp.cpp
@@ -49,7 +49,7 @@ void FUAFSteeringAnimOp::EvaluateValues(FUAFAnimOpValueEvaluator& Evaluator)
 	// Feedback: Would not have been able to discover this API myself, had to read UAFDecompressAnimSequenceAnimOp
 	const FUAFAnimOpValueEvaluationContext& AnimOpEvaluationContext = Evaluator.GetActiveEvaluationContext();
 
-	FValueBundleCoWRef* ValueBundle = Evaluator.GetEvaluationStack().PeekMutable();
+	FPoseValueBundleCoWRef* ValueBundle = Evaluator.GetEvaluationStack().PeekMutable();
 	check (ValueBundle != nullptr)
 	
 	FAttributeNamedSetPtr NamedSet = ValueBundle->Get().GetNamedSet();
diff --git a/Engine/Plugins/Experimental/UAF/UAFAnimNode/Source/UAFAnimNode/Public/UAF/AnimOpCore/UAFAnimOpTuple.h b/Engine/Plugins/Experimental/UAF/UAFAnimNode/Source/UAFAnimNode/Public/UAF/AnimOpCore/UAFAnimOpTuple.h
--- a/Engine/Plugins/Experimental/UAF/UAFAnimNode/Source/UAFAnimNode/Public/UAF/AnimOpCore/UAFAnimOpTuple.h
+++ b/Engine/Plugins/Experimental/UAF/UAFAnimNode/Source/UAFAnimNode/Public/UAF/AnimOpCore/UAFAnimOpTuple.h
@@ -69,7 +69,7 @@ namespace UE::UAF
 	template<typename... AnimOpTypes>
 	inline void TUAFAnimOpTuple<AnimOpTypes...>::EvaluateValues(FUAFAnimOpValueEvaluator& Evaluator)
 	{
-		TUAFStack<FValueBundleCoWRef>& EvaluationStack = Evaluator.GetEvaluationStack();
+		TUAFStack<FPoseValueBundleCoWRef>& EvaluationStack = Evaluator.GetEvaluationStack();
 
 		VisitTupleElements(
 			[&Evaluator, &EvaluationStack](auto& AnimOp)
diff --git a/Engine/Plugins/Experimental/UAF/UAFAnimNode/Source/UAFAnimNode/Public/UAF/AnimOpCore/UAFAnimOpValueEvaluator.h b/Engine/Plugins/Experimental/UAF/UAFAnimNode/Source/UAFAnimNode/Public/UAF/AnimOpCore/UAFAnimOpValueEvaluator.h
--- a/Engine/Plugins/Experimental/UAF/UAFAnimNode/Source/UAFAnimNode/Public/UAF/AnimOpCore/UAFAnimOpValueEvaluator.h
+++ b/Engine/Plugins/Experimental/UAF/UAFAnimNode/Source/UAFAnimNode/Public/UAF/AnimOpCore/UAFAnimOpValueEvaluator.h
@@ -7,7 +7,7 @@
 #include "UAF/Attributes/AttributeBindingData.h"
 #include "UAF/UAFCoWRef.h"
 #include "UAF/ValueRuntime/Transformers/Transformer.h"
-#include "UAF/ValueRuntime/ValueBundle.h"
+#include "UAF/ValueRuntime/PoseValueBundle.h"
 
 #define UE_API UAFANIMNODE_API
 
@@ -16,11 +16,11 @@ class USkeleton;
 
 namespace UE::UAF
 {
-	// A stack of FValueBundle instances used when sampling sequences and blending their results
+	// A stack of FPoseValueBundle instances used when sampling sequences and blending their results
 	// Their container instances will use the MemStack but internally they can use any allocator
-	using FValueBundleCoWRef = TUAFCoWRef<FValueBundle, TMemStackAllocator<>>;
+	using FPoseValueBundleCoWRef = TUAFCoWRef<FPoseValueBundle, TMemStackAllocator<>>;
 	extern UE_API const FUAFStackName VALUE_STACK_NAME;
-	UE_UAF_REGISTER_STACK_TYPE(FValueBundleCoWRef);
+	UE_UAF_REGISTER_STACK_TYPE(FPoseValueBundleCoWRef);
 
 	/*
 	 * FUAFAnimOpValueEvaluationContext
@@ -53,7 +53,7 @@ namespace UE::UAF
 		[[nodiscard]] UE_API bool SetNamedSet(FName SetName);
 
 		// Returns the default attribute values for the current named set (or nullptr if we aren't valid)
-		[[nodiscard]] const FValueBundlePtr& GetDefaultAttributeValues() const;
+		[[nodiscard]] const FPoseValueBundlePtr& GetDefaultAttributeValues() const;
 
 	private:
 		// The set binding we are evaluating with
@@ -72,7 +72,7 @@ namespace UE::UAF
 		FAttributeNamedSetPtr NamedSet;
 
 		// Default values as specified by the current set binding and named set
-		FValueBundlePtr DefaultAttributeValues;
+		FPoseValueBundlePtr DefaultAttributeValues;
 	};
 
 	/*
@@ -86,8 +86,8 @@ namespace UE::UAF
 		UE_API FUAFAnimOpValueEvaluator(TNonNullPtr<const UAbstractSkeletonSetBinding> SetBinding, const USkeletalMesh* SkeletalMesh, FName SetName, int32 CurrentLOD);
 
 		// Returns the evaluation stack for value bundles
-		[[nodiscard]] TUAFStack<FValueBundleCoWRef>& GetEvaluationStack();
-		[[nodiscard]] const TUAFStack<FValueBundleCoWRef>& GetEvaluationStack() const;
+		[[nodiscard]] TUAFStack<FPoseValueBundleCoWRef>& GetEvaluationStack();
+		[[nodiscard]] const TUAFStack<FPoseValueBundleCoWRef>& GetEvaluationStack() const;
 
 		// Returns the evaluation context that controls what values we produce
 		[[nodiscard]] UE_API const FUAFAnimOpValueEvaluationContext& GetActiveEvaluationContext() const;
@@ -115,7 +115,7 @@ namespace UE::UAF
 
 	private:
 		// The main evaluation stack for values
-		TUAFStack<FValueBundleCoWRef> EvaluationStack;
+		TUAFStack<FPoseValueBundleCoWRef> EvaluationStack;
 
 		// The internal context that controls what values we are producing
 		TArray<FUAFAnimOpValueEvaluationContext> EvaluationContextStack;
@@ -150,7 +150,7 @@ namespace UE::UAF
 		return NamedSet;
 	}
 
-	inline const FValueBundlePtr& FUAFAnimOpValueEvaluationContext::GetDefaultAttributeValues() const
+	inline const FPoseValueBundlePtr& FUAFAnimOpValueEvaluationContext::GetDefaultAttributeValues() const
 	{
 		return DefaultAttributeValues;
 	}
@@ -160,12 +160,12 @@ namespace UE::UAF
 		return SetBinding != nullptr;
 	}
 
-	inline TUAFStack<FValueBundleCoWRef>& FUAFAnimOpValueEvaluator::GetEvaluationStack()
+	inline TUAFStack<FPoseValueBundleCoWRef>& FUAFAnimOpValueEvaluator::GetEvaluationStack()
 	{
 		return EvaluationStack;
 	}
 
-	inline const TUAFStack<FValueBundleCoWRef>& FUAFAnimOpValueEvaluator::GetEvaluationStack() const
+	inline const TUAFStack<FPoseValueBundleCoWRef>& FUAFAnimOpValueEvaluator::GetEvaluationStack() const
 	{
 		return EvaluationStack;
 	}
diff --git a/Engine/Plugins/Experimental/UAF/UAFAnimNode/Source/UAFAnimNode/Public/UAF/AnimOps/UAFApplyAdditiveAnimOp.h b/Engine/Plugins/Experimental/UAF/UAFAnimNode/Source/UAFAnimNode/Public/UAF/AnimOps/UAFApplyAdditiveAnimOp.h
--- a/Engine/Plugins/Experimental/UAF/UAFAnimNode/Source/UAFAnimNode/Public/UAF/AnimOps/UAFApplyAdditiveAnimOp.h
+++ b/Engine/Plugins/Experimental/UAF/UAFAnimNode/Source/UAFAnimNode/Public/UAF/AnimOps/UAFApplyAdditiveAnimOp.h
@@ -13,7 +13,7 @@
 namespace UE::UAF
 {
 	class FValueBundleStack;
-	class FValueBundle;
+	class FPoseValueBundle;
 
 	USTRUCT()
 	struct FUAFApplyAdditiveAnimOp : public FUAFAnimOp
@@ -31,7 +31,7 @@ namespace UE::UAF
 		UE_API void SetBlendMask(const TObjectPtr<UUAFBlendMask> InBlendMask);
 		
 	private:
-		void BuildPerValueWeights(const FUAFAnimOpValueEvaluator& Evaluator, const FValueBundle& BaseInput, FValueBundleStack& OutPerValueWeights);
+		void BuildPerValueWeights(const FUAFAnimOpValueEvaluator& Evaluator, const FPoseValueBundle& BaseInput, FValueBundleStack& OutPerValueWeights);
 		
 	private:
 		// Optional blend mask for per-bone application of the additive
diff --git a/Engine/Plugins/Experimental/UAF/UAFAnimNode/Source/UAFAnimNode/Public/UAF/AnimOps/UAFBlendTwoAnimOp.h b/Engine/Plugins/Experimental/UAF/UAFAnimNode/Source/UAFAnimNode/Public/UAF/AnimOps/UAFBlendTwoAnimOp.h
--- a/Engine/Plugins/Experimental/UAF/UAFAnimNode/Source/UAFAnimNode/Public/UAF/AnimOps/UAFBlendTwoAnimOp.h
+++ b/Engine/Plugins/Experimental/UAF/UAFAnimNode/Source/UAFAnimNode/Public/UAF/AnimOps/UAFBlendTwoAnimOp.h
@@ -12,7 +12,7 @@
 
 namespace UE::UAF
 {
-	class FValueBundle;
+	class FPoseValueBundle;
 }
 
 namespace UE::UAF
@@ -47,7 +47,7 @@ namespace UE::UAF
 
 	protected:
 		// Compute the current alpha value for curves or return the interpolation alpha otherwise
-		[[nodiscard]] UE_API float GetInterpolationAlpha(const FValueBundle& InputA, const FValueBundle& InputB) const;
+		[[nodiscard]] UE_API float GetInterpolationAlpha(const FPoseValueBundle& InputA, const FPoseValueBundle& InputB) const;
 
 		// The desired interpolation alpha between the two input keyframes
 		UPROPERTY(VisibleAnywhere, Category = Properties)
diff --git a/Engine/Plugins/Experimental/UAF/UAFAnimNode/Source/UAFAnimNode/Public/UAF/AnimOps/UAFBlendTwoPerValueAnimOp.h b/Engine/Plugins/Experimental/UAF/UAFAnimNode/Source/UAFAnimNode/Public/UAF/AnimOps/UAFBlendTwoPerValueAnimOp.h
--- a/Engine/Plugins/Experimental/UAF/UAFAnimNode/Source/UAFAnimNode/Public/UAF/AnimOps/UAFBlendTwoPerValueAnimOp.h
+++ b/Engine/Plugins/Experimental/UAF/UAFAnimNode/Source/UAFAnimNode/Public/UAF/AnimOps/UAFBlendTwoPerValueAnimOp.h
@@ -12,11 +12,6 @@
 
 class UUAFBlendProfile;
 
-namespace UE::UAF
-{
-	class FValueBundle;
-}
-
 namespace UE::UAF
 {
 	USTRUCT()
diff --git a/Engine/Plugins/Experimental/UAF/UAFAnimNode/Source/UAFAnimNode/Public/UAF/AnimOps/UAFSteeringAnimOp.h b/Engine/Plugins/Experimental/UAF/UAFAnimNode/Source/UAFAnimNode/Public/UAF/AnimOps/UAFSteeringAnimOp.h
--- a/Engine/Plugins/Experimental/UAF/UAFAnimNode/Source/UAFAnimNode/Public/UAF/AnimOps/UAFSteeringAnimOp.h
+++ b/Engine/Plugins/Experimental/UAF/UAFAnimNode/Source/UAFAnimNode/Public/UAF/AnimOps/UAFSteeringAnimOp.h
@@ -13,7 +13,6 @@ namespace UE::UAF
 {
 struct IUAFAnimNodeTimeline;
 struct IUAFRootMotionProvider;
-class FValueBundle;
 }
 
 namespace UE::UAF
diff --git a/Engine/Plugins/Experimental/UAF/UAFPoseSearch/Source/UAFPoseSearch/Private/AnimNode/UAFPoseHistoryNode.cpp b/Engine/Plugins/Experimental/UAF/UAFPoseSearch/Source/UAFPoseSearch/Private/AnimNode/UAFPoseHistoryNode.cpp
--- a/Engine/Plugins/Experimental/UAF/UAFPoseSearch/Source/UAFPoseSearch/Private/AnimNode/UAFPoseHistoryNode.cpp
+++ b/Engine/Plugins/Experimental/UAF/UAFPoseSearch/Source/UAFPoseSearch/Private/AnimNode/UAFPoseHistoryNode.cpp
@@ -30,16 +30,14 @@ void FUAFHistoryCollectorPreAnimOp::EvaluateValues(FUAFAnimOpValueEvaluator& Eva
 // @todo: revisit this code once a FCSPose<FLODPoseStack> is implemented!
 struct FUAFComponentSpacePoseProvider : public UE::PoseSearch::IComponentSpacePoseProvider
 {
-	FUAFComponentSpacePoseProvider(const FValueBundle& InValueBundle, const USkeleton* InSkeleton)
+	FUAFComponentSpacePoseProvider(const FPoseValueBundle& InValueBundle, const USkeleton* InSkeleton)
 		: ValueBundle(InValueBundle)
 		, Skeleton(InSkeleton)
 	{
-		const FAttributeMappingKey MappingKey = FAttributeMappingKey::MakeFromTo<FBoneTransformAnimationAttribute>();
+		const FPoseValueBundlePtr DefaultValues = ValueBundle.GetNamedSet()->GetDefaultAttributeValues();
 
-		const FValueBundlePtr DefaultValues = ValueBundle.GetNamedSet()->GetDefaultAttributeValues();
-
-		DefaultBoneTransforms = DefaultValues->GetBoundValueMaps().Find<FBoneTransformAnimationAttribute>(MappingKey);
-		BoneTransforms = ValueBundle.GetBoundValueMaps().Find<FBoneTransformAnimationAttribute>(MappingKey);
+		DefaultBoneTransforms = DefaultValues->FindBoneTransforms();
+		BoneTransforms = ValueBundle.FindBoneTransforms();
 		DefaultBoneSet = DefaultBoneTransforms->GetTypedSet();
 		BoneSet = BoneTransforms->GetTypedSet();
 
@@ -90,7 +88,7 @@ struct FUAFComponentSpacePoseProvider : public UE::PoseSearch::IComponentSpacePo
 	}
 
 private:
-	const FValueBundle& ValueBundle;
+	const FPoseValueBundle& ValueBundle;
 	const USkeleton* Skeleton;
 
 	const TBoundValueMap<FBoneTransformAnimationAttribute>* DefaultBoneTransforms;
@@ -115,8 +113,8 @@ void FUAFHistoryCollectorPostAnimOp::EvaluateValues(FUAFAnimOpValueEvaluator& Ev
 	const FUAFAnimOpValueEvaluationContext& AnimOpEvaluationContext = Evaluator.GetActiveEvaluationContext();
 
 	const USkeleton* Skeleton = AnimOpEvaluationContext.GetSkeleton();
-	FValueBundleCoWRef* InputRef = Evaluator.GetEvaluationStack().PeekMutable();
-	const FValueBundle& Input = InputRef->Get();
+	FPoseValueBundleCoWRef* InputRef = Evaluator.GetEvaluationStack().PeekMutable();
+	const FPoseValueBundle& Input = InputRef->Get();
 
 	if (Input.IsEmpty())
 	{
diff --git a/Engine/Plugins/Experimental/UAF/UAFWarping/Source/UAFWarping/Private/AnimNode/UAFStrafeWarpingNode.cpp b/Engine/Plugins/Experimental/UAF/UAFWarping/Source/UAFWarping/Private/AnimNode/UAFStrafeWarpingNode.cpp
--- a/Engine/Plugins/Experimental/UAF/UAFWarping/Source/UAFWarping/Private/AnimNode/UAFStrafeWarpingNode.cpp
+++ b/Engine/Plugins/Experimental/UAF/UAFWarping/Source/UAFWarping/Private/AnimNode/UAFStrafeWarpingNode.cpp
@@ -79,8 +79,8 @@ void FUAFStrafeWarpingPostAnimOp::DoWarpRootMotion(FUAFAnimOpValueEvaluator& Eva
 	const FUAFAnimOpValueEvaluationContext& AnimOpEvaluationContext = Evaluator.GetActiveEvaluationContext();
 
 	const USkeleton* Skeleton = AnimOpEvaluationContext.GetSkeleton();
-	FValueBundleCoWRef* InputRef = Evaluator.GetEvaluationStack().PeekMutable();
-	FValueBundle& ValueBundle = InputRef->GetMutable();
+	FPoseValueBundleCoWRef* InputRef = Evaluator.GetEvaluationStack().PeekMutable();
+	FPoseValueBundle& ValueBundle = InputRef->GetMutable();
 
 	const FAttributeNamedSetPtr NamedSet = ValueBundle.GetNamedSet();
 	const FAttributeTypedSetPtr AttributeTypedSet = NamedSet->FindTypedSet<FTransformAnimationAttribute>();
@@ -334,24 +334,22 @@ void FUAFStrafeWarpingPostAnimOp::DoWarpPose(FUAFAnimOpValueEvaluator& Evaluator
 	const FUAFAnimOpValueEvaluationContext& AnimOpEvaluationContext = Evaluator.GetActiveEvaluationContext();
 
 	const USkeleton* Skeleton = AnimOpEvaluationContext.GetSkeleton();
-	FValueBundleCoWRef* InputRef = Evaluator.GetEvaluationStack().PeekMutable();
-	FValueBundle& ValueBundle = InputRef->GetMutable();
+	FPoseValueBundleCoWRef* InputRef = Evaluator.GetEvaluationStack().PeekMutable();
+	FPoseValueBundle& ValueBundle = InputRef->GetMutable();
 
-	const FAttributeMappingKey MappingKey = FAttributeMappingKey::MakeFromTo<FBoneTransformAnimationAttribute>();
-
-	TBoundValueMap<FBoneTransformAnimationAttribute>* BoneTransforms = ValueBundle.GetBoundValueMaps().Find<FBoneTransformAnimationAttribute>(MappingKey);
+	TBoundValueMap<FBoneTransformAnimationAttribute>* BoneTransforms = ValueBundle.FindBoneTransforms();
 	if (!BoneTransforms)
 	{
 		return;
 	}
 
-	FValueBundleStack ComponentSpaceCollection(Evaluator.GetActiveNamedSet());
+	FPoseValueBundleStack ComponentSpaceCollection(Evaluator.GetActiveNamedSet());
 	ComponentSpaceCollection.InitWithValueSpace(FValueSpace(EValueSpaceType::Local));
 
 	// @todo: add a LocalToComponent evaluating transforms in a lazy way to avoid having to compute the entire pose in component space!
 	FBoneSpace::LocalToComponent(ValueBundle, ComponentSpaceCollection);
 
-	TBoundValueMap<FBoneTransformAnimationAttribute>* ComponentSpaceBoneTransforms = ComponentSpaceCollection.GetBoundValueMaps().Find<FBoneTransformAnimationAttribute>(MappingKey);
+	TBoundValueMap<FBoneTransformAnimationAttribute>* ComponentSpaceBoneTransforms = ComponentSpaceCollection.FindBoneTransforms();
 	if (!ComponentSpaceBoneTransforms)
 	{
 		return;
diff --git a/Engine/Plugins/Experimental/UAF/UAFWarping/Source/UAFWarping/Private/AnimNode/UAFWarpToTargetNode.cpp b/Engine/Plugins/Experimental/UAF/UAFWarping/Source/UAFWarping/Private/AnimNode/UAFWarpToTargetNode.cpp
--- a/Engine/Plugins/Experimental/UAF/UAFWarping/Source/UAFWarping/Private/AnimNode/UAFWarpToTargetNode.cpp
+++ b/Engine/Plugins/Experimental/UAF/UAFWarping/Source/UAFWarping/Private/AnimNode/UAFWarpToTargetNode.cpp
@@ -23,8 +23,8 @@ void FUAFWarpToTargetPostAnimOp::EvaluateValues(FUAFAnimOpValueEvaluator& Evalua
 	const FUAFAnimOpValueEvaluationContext& AnimOpEvaluationContext = Evaluator.GetActiveEvaluationContext();
 
 	const USkeleton* Skeleton = AnimOpEvaluationContext.GetSkeleton();
-	FValueBundleCoWRef* InputRef = Evaluator.GetEvaluationStack().PeekMutable();
-	FValueBundle& ValueBundle = InputRef->GetMutable();
+	FPoseValueBundleCoWRef* InputRef = Evaluator.GetEvaluationStack().PeekMutable();
+	FPoseValueBundle& ValueBundle = InputRef->GetMutable();
 
 	const FAttributeNamedSetPtr NamedSet = ValueBundle.GetNamedSet();
 	if (!NamedSet)
